AggregateTicking and Pipe merged into Engine/Plugins/Runtime/MassEntity
Mass split up into Engine/Plugins/Runtime/MassGameplay and Engine/Plugins/AI/MassAI
Moved MassCrowd out of restricted folder

diff --git a/Engine/Plugins/AI/MassAI/Config/BaseMassAI.ini b/Engine/Plugins/AI/MassAI/Config/BaseMassAI.ini
new file mode 100644
--- /dev/null
+++ b/Engine/Plugins/AI/MassAI/Config/BaseMassAI.ini
@@ -0,0 +1,40 @@
+[CoreRedirects]
++PackageRedirects=(OldName="/Script/MassVisualization",NewName="/Script/MassRepresentation")
++PackageRedirects=(OldName="/Script/MassMovement",NewName="/Script/MassAIMovement")
++PackageRedirects=(OldName="/Script/MassBehavior",NewName="/Script/MassAIBehavior")
+
++ClassRedirects=(OldName="/Script/MassMovement.MassAvoidanceObstacleTrait",NewName="/Script/MassAIMovement.MassAvoidanceObstacleTrait")
++ClassRedirects=(OldName="/Script/MassMovement.MassDynamicObstacleTrait",NewName="/Script/MassAIMovement.MassAvoidanceObstacleTrait")
++ClassRedirects=(OldName="/Script/MassMovement.MassAvoidanceTrait",NewName="/Script/MassAIMovement.MassAvoidanceTrait")
++ClassRedirects=(OldName="/Script/MassMovement.MassZoneGraphMovementTrait",NewName="/Script/MassAIMovement.MassZoneGraphMovementTrait")
+
+; these are potentially not needed anymore
++StructRedirects=(OldName="MassLaneHandleFragment",NewName="MassLaneNavigationFragment")
++StructRedirects=(OldName="MassAgentConfig",NewName="MassEntityConfig")
++StructRedirects=(OldName="MassSpawnedAgentType",NewName="MassSpawnedEntityType")
++StructRedirects=(OldName="MassSpawnedAgentType",NewName="MassSpawnedEntityType")
++StructRedirects=(OldName="MassStateTreeClaimSmartObjectTask",NewName="MassMoveAndClaimSmartObjectTask")
++StructRedirects=(OldName="MassStateTreeUseSmartObjectTask",NewName="MassUseSmartObjectTask")
++StructRedirects=(OldName="MassZoneGraphFindFleeTarget",NewName="MassZoneGraphFindEscapeTarget")
++StructRedirects=(OldName="MassDynamicObstacleGridCellLocationFragment",NewName="MassAvoidanceObstacleGridCellLocationFragment")
+
++ClassRedirects=(OldName="MassProcessor_SmartObjectTestBehavior",NewName="MassProcessor_SmartObjectTimedBehavior")
++ClassRedirects=(OldName="MassAgentConfigAsset",NewName="MassEntityConfigAsset")
++ClassRedirects=(OldName="MassAgentFeatureBase",NewName="MassEntityTraitBase")
++ClassRedirects=(OldName="MassStateTreeFeature",NewName="MassStateTreeTrait")
++ClassRedirects=(OldName="MassAssortedFragmentsFeature",NewName="MassAssortedFragmentsTrait")
++ClassRedirects=(OldName="MassDebugVisualizationFeature",NewName="MassDebugVisualizationTrait")
++ClassRedirects=(OldName="MassReplicationFeature",NewName="MassReplicationTrait")
++ClassRedirects=(OldName="DynamicObstacleProcessor",NewName="MassAvoidanceDynamicObstacleProcessor")
++ClassRedirects=(OldName="MassAvoidanceDynamicObstacleProcessor",NewName="MassAvoidanceObstacleProcessor")
++ClassRedirects=(OldName="MassDynamicObstacleTrait",NewName="MassAvoidanceObstacleTrait")
++ClassRedirects=(OldName="MassFragmentDestructor_DynamicObstacleRemover",NewName="MassAvoidanceObstacleRemoverFragmentDestructor")
+
++PropertyRedirects=(OldName="MassAgentComponent.bUseAgentConfig",NewName="bUseEntityConfig")
++PropertyRedirects=(OldName="MassAgentComponent.AgentConfig",NewName="EntityConfig")
++PropertyRedirects=(OldName="MassSpawnedEntityType.AgentConfig",NewName="EntityConfig")
++PropertyRedirects=(OldName="MassSpawner.AgentTypes",NewName="EntityTypes")
++PropertyRedirects=(OldName="MassSpawner.bUseAgentConfig",NewName="bUseEntityConfig")
++PropertyRedirects=(OldName="MassEntityConfig.Features",NewName="Traits")
+
++PropertyRedirects=(OldName="MassZoneGraphFindEscapeTarget.FleeTag",NewName="DisturbanceAnnotationTag")
diff --git a/Engine/Plugins/AI/MassAI/Source/MassAIBehavior/Private/Conditions/ZoneGraphTagConditions.cpp b/Engine/Plugins/AI/MassAI/Source/MassAIBehavior/Private/Conditions/ZoneGraphTagConditions.cpp
new file mode 100644
--- /dev/null
+++ b/Engine/Plugins/AI/MassAI/Source/MassAIBehavior/Private/Conditions/ZoneGraphTagConditions.cpp
@@ -0,0 +1,193 @@
+﻿// Copyright Epic Games, Inc. All Rights Reserved.
+
+#include "Conditions/ZoneGraphTagConditions.h"
+#if WITH_EDITOR
+#include "ZoneGraphSettings.h"
+#include "StateTreePropertyBindings.h"
+
+#define LOCTEXT_NAMESPACE "StateTreeEditor"
+
+namespace UE::MassBehavior::ZoneGraph
+{
+	FText GetTagName(const FZoneGraphTag Tag)
+	{
+		const UZoneGraphSettings* ZoneGraphSettings = GetDefault<UZoneGraphSettings>();
+		check(ZoneGraphSettings);
+		TConstArrayView<FZoneGraphTagInfo> TagInfos = ZoneGraphSettings->GetTagInfos();
+		
+		for (const FZoneGraphTagInfo& TagInfo : TagInfos)
+		{
+			if (TagInfo.Tag == Tag)
+			{
+				return FText::FromName(TagInfo.Name);
+			}
+		}
+		return FText::GetEmpty();
+	}
+
+	FText GetTagMaskName(const FZoneGraphTagMask TagMask)
+	{
+		const UZoneGraphSettings* ZoneGraphSettings = GetDefault<UZoneGraphSettings>();
+		check(ZoneGraphSettings);
+		TConstArrayView<FZoneGraphTagInfo> TagInfos = ZoneGraphSettings->GetTagInfos();
+
+		TArray<FText> Names;
+		for (const FZoneGraphTagInfo& Info : TagInfos)
+		{
+			if (TagMask.Contains(Info.Tag))
+			{
+				if (Info.IsValid())
+				{
+					Names.Add(FText::FromName(Info.Name));
+				}
+			}
+		}
+		if (Names.Num() == 0)
+		{
+			return LOCTEXT("EmptyMask", "(Empty)");
+		}
+		
+		if (Names.Num() > 2)
+		{
+			Names.SetNum(2);
+			Names.Add(FText::FromString(TEXT("...")));
+		}
+		
+		return FText::Join(FText::FromString(TEXT(", ")), Names);
+	}
+
+	FText GetMaskOperatorText(const EZoneLaneTagMaskComparison Operator)
+	{
+		switch (Operator)
+		{
+		case EZoneLaneTagMaskComparison::Any:
+			return LOCTEXT("ContainsAny", "Any");
+		case EZoneLaneTagMaskComparison::All:
+			return LOCTEXT("ContainsAll", "All");
+		case EZoneLaneTagMaskComparison::Not:
+			return LOCTEXT("ContainsNot", "Not");
+		default:
+			return FText::FromString(TEXT("??"));
+		}
+		return FText::GetEmpty();
+	}
+
+}
+
+FText FZoneGraphTagFilterCondition::GetDescription(const IStateTreeBindingLookup& BindingLookup) const
+{
+	const FStateTreeEditorPropertyPath LeftPath(ID, GET_MEMBER_NAME_STRING_CHECKED(FZoneGraphTagFilterCondition, Tags));
+
+	FText InvertText;
+	if (bInvert)
+	{
+		InvertText = LOCTEXT("Not", "Not");
+	}
+
+	FText LeftText;
+	if (const FStateTreeEditorPropertyPath* Binding = BindingLookup.GetPropertyBindingSource(LeftPath))
+	{
+		LeftText = BindingLookup.GetPropertyPathDisplayName(*Binding);
+	}
+	else
+	{
+		LeftText = UE::MassBehavior::ZoneGraph::GetTagMaskName(Tags);
+	}
+
+	FText FilterParts[6] = { FText::GetEmpty(), FText::GetEmpty(), FText::GetEmpty(), FText::GetEmpty(), FText::GetEmpty(), FText::GetEmpty() };
+	int32 PartIndex = 0;
+	if (Filter.AnyTags != FZoneGraphTagMask::None)
+	{
+		FilterParts[PartIndex++] = LOCTEXT("ContainsAny", "any");
+		FilterParts[PartIndex++] = UE::MassBehavior::ZoneGraph::GetTagMaskName(Filter.AnyTags);
+	}
+
+	if (Filter.AllTags != FZoneGraphTagMask::None)
+	{
+		FilterParts[PartIndex++] = LOCTEXT("ContainsAll", "all");
+		FilterParts[PartIndex++] = UE::MassBehavior::ZoneGraph::GetTagMaskName(Filter.AllTags);
+	}
+
+	if (Filter.NotTags != FZoneGraphTagMask::None)
+	{
+		FilterParts[PartIndex++] = LOCTEXT("ContainsNot", "not");
+		FilterParts[PartIndex++] = UE::MassBehavior::ZoneGraph::GetTagMaskName(Filter.NotTags);
+	}
+
+	return FText::Format(LOCTEXT("CompareZoneGraphTagFilterDesc", "{0} <Details.Bold>{1}</> contains {2} <Details.Bold>{3}</> {4} <Details.Bold>{5}</> {6} <Details.Bold>{7}</>"), InvertText, LeftText, FilterParts[0], FilterParts[1], FilterParts[2], FilterParts[3], FilterParts[4], FilterParts[5]);
+}
+
+
+FText FZoneGraphTagMaskCondition::GetDescription(const IStateTreeBindingLookup& BindingLookup) const
+{
+	const FStateTreeEditorPropertyPath LeftPath(ID, GET_MEMBER_NAME_STRING_CHECKED(FZoneGraphTagMaskCondition, Left));
+	const FStateTreeEditorPropertyPath RightPath(ID, GET_MEMBER_NAME_STRING_CHECKED(FZoneGraphTagMaskCondition, Right));
+
+	FText InvertText;
+	if (bInvert)
+	{
+		InvertText = LOCTEXT("Not", "Not");
+	}
+
+	FText LeftText;
+	if (const FStateTreeEditorPropertyPath* Binding = BindingLookup.GetPropertyBindingSource(LeftPath))
+	{
+		LeftText = BindingLookup.GetPropertyPathDisplayName(*Binding);
+	}
+	else
+	{
+		LeftText = UE::MassBehavior::ZoneGraph::GetTagMaskName(Left);
+	}
+
+	FText OperatorText = UE::MassBehavior::ZoneGraph::GetMaskOperatorText(Operator);
+
+	FText RightText;
+	if (const FStateTreeEditorPropertyPath* Binding = BindingLookup.GetPropertyBindingSource(RightPath))
+	{
+		RightText = BindingLookup.GetPropertyPathDisplayName(*Binding);
+	}
+	else
+	{
+		RightText = UE::MassBehavior::ZoneGraph::GetTagMaskName(Right);
+	}
+
+	return FText::Format(LOCTEXT("CompareZoneGraphTagMaskDesc", "{0} <Details.Bold>{1}</> contains {2} <Details.Bold>{3}</>"), InvertText, LeftText, OperatorText, RightText);
+}
+
+
+FText FZoneGraphTagCondition::GetDescription(const IStateTreeBindingLookup& BindingLookup) const
+{
+	const FStateTreeEditorPropertyPath LeftPath(ID, GET_MEMBER_NAME_STRING_CHECKED(FZoneGraphTagCondition, Left));
+	const FStateTreeEditorPropertyPath RightPath(ID, GET_MEMBER_NAME_STRING_CHECKED(FZoneGraphTagCondition, Right));
+
+	FText InvertText;
+	if (bInvert)
+	{
+		InvertText = LOCTEXT("Not", "Not");
+	}
+
+	FText LeftText;
+	if (const FStateTreeEditorPropertyPath* Binding = BindingLookup.GetPropertyBindingSource(LeftPath))
+	{
+		LeftText = BindingLookup.GetPropertyPathDisplayName(*Binding);
+	}
+	else
+	{
+		LeftText = UE::MassBehavior::ZoneGraph::GetTagName(Left);
+	}
+
+	FText RightText;
+	if (const FStateTreeEditorPropertyPath* Binding = BindingLookup.GetPropertyBindingSource(RightPath))
+	{
+		RightText = BindingLookup.GetPropertyPathDisplayName(*Binding);
+	}
+	else
+	{
+		RightText = UE::MassBehavior::ZoneGraph::GetTagName(Right);
+	}
+
+	return FText::Format(LOCTEXT("CompareZoneGraphTagDesc", "{0} <Details.Bold>{1}</> is <Details.Bold>{3}</>"), InvertText, LeftText, RightText);
+}
+
+#undef LOCTEXT_NAMESPACE
+#endif // WITH_EDITOR
diff --git a/Engine/Plugins/AI/MassAI/Source/MassAIBehavior/Private/Evaluators/MassComponentHitEvaluator.cpp b/Engine/Plugins/AI/MassAI/Source/MassAIBehavior/Private/Evaluators/MassComponentHitEvaluator.cpp
new file mode 100644
--- /dev/null
+++ b/Engine/Plugins/AI/MassAI/Source/MassAIBehavior/Private/Evaluators/MassComponentHitEvaluator.cpp
@@ -0,0 +1,34 @@
+// Copyright Epic Games, Inc. All Rights Reserved.
+
+#include "Evaluators/MassComponentHitEvaluator.h"
+#include "MassAIBehaviorTypes.h"
+#include "MassComponentHitSubsystem.h"
+#include "MassStateTreeExecutionContext.h"
+#include "Engine/World.h"
+
+void FMassComponentHitEvaluator::Evaluate(FStateTreeExecutionContext &Context, const EStateTreeEvaluationType EvalType, const float DeltaTime)
+{
+	// Look for recent hits
+	UMassComponentHitSubsystem& HitSubsystem = Context.GetExternalItem(ComponentHitSubsystemHandle).GetMutable<UMassComponentHitSubsystem>();
+	const FMassHitResult* HitResult = HitSubsystem.GetLastHit(static_cast<FMassStateTreeExecutionContext&>(Context).GetEntity());
+
+	// LastHitEntity is not reset intentionally, so that it's available the duration of the behavior reacting to it.
+	bGotHit = false;
+
+	// If the hit is very recent, set the got hit, and update last hit entity.
+	if (HitResult != nullptr)
+	{
+		// @todo: This is a bit of a kludge to expose an event to StateTree.
+		const UWorld* World = Context.GetWorld();
+		check(World);
+		const float CurrentTime = World->GetTimeSeconds();
+		const float TimeSinceHit = CurrentTime - HitResult->HitTime;
+		constexpr float HitEventDuration = 0.1f;
+		if (TimeSinceHit < HitEventDuration)
+		{
+			MASSBEHAVIOR_LOG(VeryVerbose, TEXT("Got hit"));
+			bGotHit = true;
+			LastHitEntity = HitResult->OtherEntity;
+		}
+	}
+}
diff --git a/Engine/Plugins/AI/MassAI/Source/MassAIBehavior/Private/Evaluators/MassStateTreeSmartObjectEvaluator.cpp b/Engine/Plugins/AI/MassAI/Source/MassAIBehavior/Private/Evaluators/MassStateTreeSmartObjectEvaluator.cpp
new file mode 100644
--- /dev/null
+++ b/Engine/Plugins/AI/MassAI/Source/MassAIBehavior/Private/Evaluators/MassStateTreeSmartObjectEvaluator.cpp
@@ -0,0 +1,136 @@
+// Copyright Epic Games, Inc. All Rights Reserved.
+
+#include "Evaluators/MassStateTreeSmartObjectEvaluator.h"
+
+#include "MassAIBehaviorTypes.h"
+#include "MassCommonFragments.h"
+#include "MassAIMovementFragments.h"
+#include "MassSignalSubsystem.h"
+#include "MassStateTreeExecutionContext.h"
+#include "MassSmartObjectHandler.h"
+#include "MassSmartObjectProcessor.h"
+#include "MassStateTreeProcessors.h"
+#include "MassZoneGraphMovementFragments.h"
+#include "SmartObjectSubsystem.h"
+#include "Engine/World.h"
+
+//----------------------------------------------------------------------//
+// FMassStateTreeSmartObjectEvaluator
+//----------------------------------------------------------------------//
+void FMassStateTreeSmartObjectEvaluator::EnterState(FStateTreeExecutionContext& Context, const EStateTreeStateChangeType ChangeType, const FStateTreeTransitionResult& Transition)
+{
+	if (ChangeType != EStateTreeStateChangeType::Changed)
+	{
+		return;
+	}
+
+	Reset();
+}
+
+void FMassStateTreeSmartObjectEvaluator::ExitState(FStateTreeExecutionContext& Context, const EStateTreeStateChangeType ChangeType, const FStateTreeTransitionResult& Transition)
+{
+	if (ChangeType != EStateTreeStateChangeType::Changed)
+	{
+		return;
+	}
+
+	if (SearchRequestID.IsSet())
+	{
+		const FMassStateTreeExecutionContext& MassContext = static_cast<FMassStateTreeExecutionContext&>(Context);
+		USmartObjectSubsystem& SmartObjectSubsystem = Context.GetExternalItem(SmartObjectSubsystemHandle).GetMutable<USmartObjectSubsystem>();
+		const FMassSmartObjectHandler MassSmartObjectHandler(MassContext.GetEntitySubsystem(), MassContext.GetEntitySubsystemExecutionContext(), SmartObjectSubsystem);
+		MassSmartObjectHandler.RemoveRequest(SearchRequestID);
+		SearchRequestID.Reset();
+	}
+	Reset();
+}
+
+void FMassStateTreeSmartObjectEvaluator::Reset()
+{
+	bCandidatesFound = false;
+	bClaimed = false;
+	SearchRequestID.Reset();
+}
+
+void FMassStateTreeSmartObjectEvaluator::Evaluate(FStateTreeExecutionContext& Context, const EStateTreeEvaluationType EvalType,  const float DeltaTime)
+{
+	const FDataFragment_SmartObjectUser& SOUser = Context.GetExternalItem(SmartObjectUserHandle).Get<FDataFragment_SmartObjectUser>();
+
+	bCandidatesFound = false;
+	bClaimed = SOUser.GetClaimHandle().IsValid();
+
+	// Already claimed, nothing to do
+	if (bClaimed)
+	{
+		return;
+	}
+
+	const UWorld* World = Context.GetWorld();
+	if (SOUser.GetCooldown() > World->GetTimeSeconds())
+	{
+		return;
+	}
+
+	// We need to track our next update cooldown since we can get ticked from any signals waking up the StateTree
+	if (NextUpdate > World->GetTimeSeconds())
+	{
+		return;
+	}
+	NextUpdate = 0.f;
+
+	USmartObjectSubsystem& SmartObjectSubsystem = Context.GetExternalItem(SmartObjectSubsystemHandle).GetMutable<USmartObjectSubsystem>();
+	const FMassStateTreeExecutionContext& MassContext = static_cast<FMassStateTreeExecutionContext&>(Context);
+	const FMassSmartObjectHandler MassSmartObjectHandler(MassContext.GetEntitySubsystem(), MassContext.GetEntitySubsystemExecutionContext(), SmartObjectSubsystem);
+
+	// Nothing claimed -> search for candidates
+	if (!SearchRequestID.IsSet())
+	{
+		// Use lanes if possible for faster queries using zone graph annotations
+		const FLWEntity RequestingEntity = MassContext.GetEntity();
+		const FMassZoneGraphLaneLocationFragment* LaneLocation = Context.GetExternalItem(LocationHandle).GetPtr<FMassZoneGraphLaneLocationFragment>();
+		bUsingZoneGraphAnnotations = LaneLocation != nullptr;
+		if (bUsingZoneGraphAnnotations)
+		{
+			MASSBEHAVIOR_CLOG(!LaneLocation->LaneHandle.IsValid(), Error, TEXT("Always expecting a valid lane from the ZoneGraph movement"));
+			if (LaneLocation->LaneHandle.IsValid())
+			{
+				SearchRequestID = MassSmartObjectHandler.FindCandidatesAsync(RequestingEntity, { LaneLocation->LaneHandle, LaneLocation->DistanceAlongLane });
+			}
+		}
+		else
+		{
+			const FDataFragment_Transform& TransformFragment = Context.GetExternalItem(EntityTransformHandle).Get<FDataFragment_Transform>();
+			SearchRequestID = MassSmartObjectHandler.FindCandidatesAsync(RequestingEntity, TransformFragment.GetTransform().GetLocation());
+		}
+	}
+	else
+	{
+		// Fetch request results
+		SearchRequestResult = MassSmartObjectHandler.GetRequestResult(SearchRequestID);
+
+		// Check if results are ready
+		if (SearchRequestResult.bProcessed)
+		{
+			// Remove requests
+			MassSmartObjectHandler.RemoveRequest(SearchRequestID);
+			SearchRequestID.Reset();
+
+			// Update bindable flag to indicate to tasks and conditions if some candidates were found
+			bCandidatesFound = SearchRequestResult.NumCandidates > 0;
+
+			const FLWEntity RequestingEntity = MassContext.GetEntity();
+			MASSBEHAVIOR_CLOG(bCandidatesFound, Log, TEXT("Found %d smart object candidates for %s"), SearchRequestResult.NumCandidates, *RequestingEntity.DebugGetDescription());
+
+			// When using ZoneGraph annotations we don't need to schedule a new update since we only need the CurrentLaneChanged signal.
+			// Otherwise we reschedule with default interval on success or retry interval on failed attempt
+			if (!bUsingZoneGraphAnnotations)
+			{
+				const float DelayInSeconds = bCandidatesFound ? TickInterval : RetryCooldown;
+				NextUpdate = World->GetTimeSeconds() + DelayInSeconds;
+				UMassSignalSubsystem& MassSignalSubsystem = Context.GetExternalItem(MassSignalSubsystemHandle).GetMutable<UMassSignalSubsystem>();
+				MassSignalSubsystem.DelaySignalEntity(UE::Mass::Signals::SmartObjectRequestCandidates, RequestingEntity, DelayInSeconds);
+			}
+		}
+		// else wait for the Evaluation that will be triggered by the "candidates ready" signal
+	}
+}
diff --git a/Engine/Plugins/AI/MassAI/Source/MassAIBehavior/Private/Evaluators/MassStateTreeTestEvaluator.cpp b/Engine/Plugins/AI/MassAI/Source/MassAIBehavior/Private/Evaluators/MassStateTreeTestEvaluator.cpp
new file mode 100644
--- /dev/null
+++ b/Engine/Plugins/AI/MassAI/Source/MassAIBehavior/Private/Evaluators/MassStateTreeTestEvaluator.cpp
@@ -0,0 +1,33 @@
+// Copyright Epic Games, Inc. All Rights Reserved.
+
+#include "Evaluators/MassStateTreeTestEvaluator.h"
+#include "Engine/Engine.h"
+#include "StateTreeExecutionContext.h"
+#include "VisualLogger/VisualLogger.h"
+
+
+FMassStateTreeTestEvaluator::FMassStateTreeTestEvaluator()
+	: Time(0.0f)
+	, bSignal(false)
+{
+}
+
+FMassStateTreeTestEvaluator::~FMassStateTreeTestEvaluator()
+{
+}
+
+void FMassStateTreeTestEvaluator::EnterState(FStateTreeExecutionContext& Context, const EStateTreeStateChangeType ChangeType, const FStateTreeTransitionResult& Transition)
+{
+	Time = 0.0f;
+	bSignal = false;
+}
+
+void FMassStateTreeTestEvaluator::Evaluate(FStateTreeExecutionContext& Context, const EStateTreeEvaluationType EvalType, const float DeltaTime)
+{
+	Time += DeltaTime;
+	if (Period > 0.0f)
+	{
+		bSignal = (FMath::FloorToInt(Time / Period) & 1) ? true : false;
+		Type = EStateTreeEvaluationType::PreSelect;
+	}
+}
diff --git a/Engine/Plugins/AI/MassAI/Source/MassAIBehavior/Private/Evaluators/MassZoneGraphAnnotationEvaluator.cpp b/Engine/Plugins/AI/MassAI/Source/MassAIBehavior/Private/Evaluators/MassZoneGraphAnnotationEvaluator.cpp
new file mode 100644
--- /dev/null
+++ b/Engine/Plugins/AI/MassAI/Source/MassAIBehavior/Private/Evaluators/MassZoneGraphAnnotationEvaluator.cpp
@@ -0,0 +1,16 @@
+﻿// Copyright Epic Games, Inc. All Rights Reserved.
+
+#include "Evaluators/MassZoneGraphAnnotationEvaluator.h"
+#include "MassStateTreeExecutionContext.h"
+#include "MassZoneGraphAnnotationFragments.h"
+
+
+FMassZoneGraphAnnotationEvaluator::FMassZoneGraphAnnotationEvaluator()
+{
+}
+
+void FMassZoneGraphAnnotationEvaluator::Evaluate(FStateTreeExecutionContext& Context, const EStateTreeEvaluationType EvalType, const float DeltaTime)
+{
+	const FMassZoneGraphAnnotationTagsFragment& BehaviorTagsFragment = Context.GetExternalItem(BehaviorTagsHandle).Get<FMassZoneGraphAnnotationTagsFragment>();
+	BehaviorTags = BehaviorTagsFragment.Tags;
+}
diff --git a/Engine/Plugins/AI/MassAI/Source/MassAIBehavior/Private/MassLookAtProcessors.cpp b/Engine/Plugins/AI/MassAI/Source/MassAIBehavior/Private/MassLookAtProcessors.cpp
new file mode 100644
--- /dev/null
+++ b/Engine/Plugins/AI/MassAI/Source/MassAIBehavior/Private/MassLookAtProcessors.cpp
@@ -0,0 +1,523 @@
+// Copyright Epic Games, Inc. All Rights Reserved.
+
+#include "MassLookAtProcessors.h"
+#include "EntityView.h"
+#include "MassAIBehaviorTypes.h"
+#include "MassCommonFragments.h"
+#include "MassLookAtFragments.h"
+#include "MassAIMovementFragments.h"
+#include "MassMovementSubsystem.h"
+#include "MassRepresentationFragments.h"
+#include "Math/UnrealMathUtility.h"
+#include "VisualLogger/VisualLogger.h"
+#include "MassZoneGraphMovementFragments.h"
+#include "ZoneGraphSubsystem.h"
+#include "ZoneGraphQuery.h"
+#include "BezierUtilities.h"
+#include "Algo/RandomShuffle.h"
+#include "Engine/World.h"
+
+#define UNSAFE_FOR_MT 1
+
+namespace UE::MassBehavior
+{
+	namespace Tweakables
+	{
+		float TrajectoryLookAhead = 600.f;
+	}
+
+	FAutoConsoleVariableRef CVars[] =
+	{
+		FAutoConsoleVariableRef(TEXT("ai.mass.LookAt.TrajectoryLookAhead"), Tweakables::TrajectoryLookAhead,
+								TEXT("Distance (in cm) further along the look at trajectory (based on current path) to lookat while moving."), ECVF_Cheat),
+	};
+
+	
+	// Clamps direction vector to a cone specified by the cone angle along X-axis
+	FVector ClampDirectionToXAxisCone(const FVector Direction, const float ConeAngle)
+	{
+		float ConeSin = 0.0f, ConeCos = 0.0f;
+		FMath::SinCos(&ConeSin, &ConeCos, ConeAngle);
+		
+		const float AngleCos = Direction.X; // Same as FVector::DotProduct(FVector::ForwardVector, Direction);
+		if (AngleCos < ConeCos)
+		{
+			const float DistToRimSq = FMath::Square(Direction.Y) + FMath::Square(Direction.Z);
+			const float InvDistToRim = DistToRimSq > KINDA_SMALL_NUMBER ? (1.0f / FMath::Sqrt(DistToRimSq)) : 0.0f;
+			return FVector(ConeCos, Direction.Y * InvDistToRim * ConeSin, Direction.Z * InvDistToRim * ConeSin);
+		}
+		
+		return Direction;
+	}
+
+	float GazeEnvelope(const float GazeTime, const float GazeDuration, const EMassLookAtGazeMode Mode)
+	{
+		if (GazeDuration < KINDA_SMALL_NUMBER || Mode == EMassLookAtGazeMode::None)
+		{
+			return 0.0f;
+		}
+
+		if (Mode == EMassLookAtGazeMode::Constant)
+		{
+			return 1.0;
+		}
+
+		// @todo: mae configurable
+		const float SustainTime = GazeDuration * 0.25f;
+		const float DecayTime = GazeDuration * 0.45f;
+		
+		if (GazeTime < SustainTime)
+		{
+			return 1.0f;
+		}
+		if (GazeTime > DecayTime)
+		{
+			return 0.0f;
+		}
+		
+		const float Duration = FMath::Max(KINDA_SMALL_NUMBER, DecayTime - SustainTime);
+		const float NormTime = FMath::Clamp((GazeTime - SustainTime) / Duration, 0.0f, 1.0f);
+		return 1.0f - NormTime;
+	}
+
+}// namespace UE::MassBehavior
+
+
+//----------------------------------------------------------------------//
+// UMassLookAtProcessor
+//----------------------------------------------------------------------//
+UMassLookAtProcessor::UMassLookAtProcessor()
+{
+	ExecutionFlags = (int32)(EProcessorExecutionFlags::Client | EProcessorExecutionFlags::Standalone);
+	ExecutionOrder.ExecuteInGroup = UE::Mass::ProcessorGroupNames::Tasks;
+	ExecutionOrder.ExecuteAfter.Add(UE::Mass::ProcessorGroupNames::Representation);
+}
+
+void UMassLookAtProcessor::Initialize(UObject& Owner)
+{
+	Super::Initialize(Owner);
+
+	MovementSubsystem = UWorld::GetSubsystem<UMassMovementSubsystem>(Owner.GetWorld());
+	checkf(MovementSubsystem != nullptr, TEXT("UMassMovementSubsystem is mandatory when using MassLookAtProcessor processor."));
+
+	ZoneGraphSubsystem = UWorld::GetSubsystem<UZoneGraphSubsystem>(Owner.GetWorld());
+	checkf(ZoneGraphSubsystem != nullptr, TEXT("UMassZoneGraphSubsystem is mandatory when using MassLookAtProcessor processor."));
+}
+
+void UMassLookAtProcessor::ConfigureQueries()
+{
+	EntityQuery_Conditional.AddRequirement<FMassLookAtFragment>(ELWComponentAccess::ReadWrite);
+	EntityQuery_Conditional.AddRequirement<FDataFragment_Transform>(ELWComponentAccess::ReadOnly);
+	EntityQuery_Conditional.AddRequirement<FMassMoveTargetFragment>(ELWComponentAccess::ReadOnly);
+	EntityQuery_Conditional.AddRequirement<FMassZoneGraphLaneLocationFragment>(ELWComponentAccess::ReadOnly, ELWComponentPresence::Optional);
+	EntityQuery_Conditional.AddRequirement<FMassLookAtTrajectoryFragment>(ELWComponentAccess::ReadWrite, ELWComponentPresence::Optional);
+	EntityQuery_Conditional.AddRequirement<FMassZoneGraphShortPathFragment>(ELWComponentAccess::ReadOnly, ELWComponentPresence::Optional);
+	EntityQuery_Conditional.AddTagRequirement<FMassHighLODTag>(ELWComponentPresence::All);
+	EntityQuery_Conditional.AddChunkRequirement<FMassVisualizationChunkFragment>(ELWComponentAccess::ReadOnly);
+	EntityQuery_Conditional.SetChunkFilter(&FMassVisualizationChunkFragment::AreAnyEntitiesVisibleInChunk);
+}
+
+void UMassLookAtProcessor::Execute(UEntitySubsystem& EntitySubsystem, FLWComponentSystemExecutionContext& Context)
+{
+	QUICK_SCOPE_CYCLE_COUNTER(LookAtProcessor_Run);
+
+	const float CurrentTime = GetWorld()->GetTimeSeconds();
+
+	EntityQuery_Conditional.ForEachEntityChunk(EntitySubsystem, Context, [this, &EntitySubsystem, CurrentTime](FLWComponentSystemExecutionContext& Context)
+		{
+			const int32 NumEntities = Context.GetEntitiesNum();
+			const TArrayView<FMassLookAtFragment> LookAtList = Context.GetMutableComponentView<FMassLookAtFragment>();
+			const TConstArrayView<FDataFragment_Transform> TransformList = Context.GetComponentView<FDataFragment_Transform>();
+			const TConstArrayView<FMassMoveTargetFragment> MoveTargetList = Context.GetComponentView<FMassMoveTargetFragment>();
+			const TConstArrayView<FMassZoneGraphLaneLocationFragment> ZoneGraphLocationList = Context.GetComponentView<FMassZoneGraphLaneLocationFragment>();
+			const TConstArrayView<FMassZoneGraphShortPathFragment> ShortPathList = Context.GetComponentView<FMassZoneGraphShortPathFragment>();
+			const TArrayView<FMassLookAtTrajectoryFragment> LookAtTrajectoryList = Context.GetMutableComponentView<FMassLookAtTrajectoryFragment>();
+
+			for (int32 i = 0; i < NumEntities; ++i)
+			{
+				FMassLookAtFragment& LookAt = LookAtList[i];
+				const FMassMoveTargetFragment& MoveTarget = MoveTargetList[i];
+				const FDataFragment_Transform& TransformFragment = TransformList[i];
+
+				const bool bHasLookAtTrajectory = ZoneGraphLocationList.Num() > 0 && LookAtTrajectoryList.Num() > 0 && ShortPathList.Num() > 0;
+
+				bool bDisplayDebug = false;
+				const FLWEntity Entity = Context.GetEntity(i);
+	#if WITH_MASS_DEBUG
+				FColor EntityColor = FColor::White;
+				bDisplayDebug = UE::MassDebug::IsDebuggingEntity(Entity, &EntityColor);
+	#endif // WITH_MASS_DEBUG
+
+				// Update gaze target when current cycle is finished.
+				if (LookAt.RandomGazeMode != EMassLookAtGazeMode::None)
+				{
+					const float TimeSinceUpdate = CurrentTime - LookAt.GazeStartTime;
+					if (TimeSinceUpdate >= LookAt.GazeDuration)
+					{
+						FindNewGazeTarget(EntitySubsystem, CurrentTime, TransformFragment.GetTransform(), LookAt);
+					}
+				}
+
+				// Update specific look at mode.
+				LookAt.Direction = FVector::ForwardVector;
+
+				switch (LookAt.LookAtMode)
+				{
+				case EMassLookAtMode::LookForward:
+					// Empty, forward set already above.
+					break;
+					
+				case EMassLookAtMode::LookAlongPath:
+					if (bHasLookAtTrajectory)
+					{
+						const FMassZoneGraphLaneLocationFragment& ZoneGraphLocation = ZoneGraphLocationList[i];
+						FMassLookAtTrajectoryFragment& LookAtTrajectory = LookAtTrajectoryList[i];
+
+						if (MoveTarget.GetCurrentActionID() != LookAt.LastSeenActionID)
+						{
+							const FMassZoneGraphLaneLocationFragment& LaneLocation = ZoneGraphLocationList[i];
+							const FMassZoneGraphShortPathFragment& ShortPath = ShortPathList[i];
+							
+							BuildTrajectory(LaneLocation, ShortPath, Entity, bDisplayDebug, LookAtTrajectory);
+							LookAt.LastSeenActionID = MoveTarget.GetCurrentActionID();
+						}
+						
+						UpdateLookAtTrajectory(TransformFragment.GetTransform(), ZoneGraphLocation, LookAtTrajectory, bDisplayDebug, LookAt);
+					}
+					break;
+
+				case EMassLookAtMode::LookAtEntity:
+					UpdateLookAtTrackedEntity(EntitySubsystem, TransformFragment.GetTransform(), bDisplayDebug, LookAt);
+					break;
+					
+				default:
+					break;
+				}
+
+				// Apply gaze
+				if (LookAt.RandomGazeMode != EMassLookAtGazeMode::None)
+				{
+					const float TimeSinceUpdate = CurrentTime - LookAt.GazeStartTime;
+					const float GazeStrength = UE::MassBehavior::GazeEnvelope(TimeSinceUpdate, LookAt.GazeDuration, LookAt.RandomGazeMode);
+
+					if (GazeStrength > KINDA_SMALL_NUMBER)
+					{
+						const bool bHasTarget = UpdateGazeTrackedEntity(EntitySubsystem, TransformFragment.GetTransform(), bDisplayDebug, LookAt);
+
+						if (bHasTarget)
+						{
+							// Treat target gaze as absolute direction.
+							LookAt.Direction = FMath::Lerp(LookAt.Direction, LookAt.GazeDirection, GazeStrength).GetSafeNormal();
+						}
+						else
+						{
+							// Treat random offset as relative direction.
+							const FQuat GazeRotation = FQuat::FindBetweenNormals(FVector::ForwardVector, FMath::Lerp(FVector::ForwardVector, LookAt.GazeDirection, GazeStrength).GetSafeNormal());
+							LookAt.Direction = GazeRotation.RotateVector(LookAt.Direction);
+						}
+					}
+				}
+
+				// Clamp
+				LookAt.Direction = UE::MassBehavior::ClampDirectionToXAxisCone(LookAt.Direction, FMath::DegreesToRadians(AngleThresholdInDegrees));
+
+	#if WITH_MASS_DEBUG && UNSAFE_FOR_MT
+				if (bDisplayDebug)
+				{
+					const FVector Origin = TransformFragment.GetTransform().GetLocation() + FVector(0.f,0.f,DebugZOffset);
+					const FVector Dest = Origin + 100.f*TransformFragment.GetTransform().TransformVector(LookAt.Direction);
+					UE_VLOG_ARROW(this, LogMassBehavior, Display, Origin, Dest, EntityColor, TEXT(""));
+				}
+	#endif // WITH_MASS_DEBUG
+			}
+		});
+}
+
+void UMassLookAtProcessor::FindNewGazeTarget(const UEntitySubsystem& EntitySubsystem, const float CurrentTime, const FTransform& Transform, FMassLookAtFragment& LookAt) const
+{
+	const FLWEntity LastTrackedEntity = LookAt.GazeTrackedEntity;
+	
+	LookAt.GazeTrackedEntity.Reset();
+	LookAt.GazeDirection = FVector::ForwardVector;
+
+	// Search for potential targets in front
+	bool bTargetFound = false;
+	if (LookAt.bRandomGazeEntities && MovementSubsystem != nullptr)
+	{
+		const float CosAngleThreshold = FMath::Cos(FMath::DegreesToRadians(AngleThresholdInDegrees));
+		const FVector Extent(QueryExtent, QueryExtent, QueryExtent);
+		const FVector QueryOrigin = Transform.TransformPosition(FVector(0.5f*QueryExtent, 0.f, 0.f));
+		const FBox QueryBox = FBox(QueryOrigin - 0.5f*Extent, QueryOrigin + 0.5f*Extent);
+
+		TArray<FAvoidanceObstacleHashGrid2D::ItemIDType> NearbyEntities;
+		NearbyEntities.Reserve(16);
+		MovementSubsystem->GetGrid().Query(QueryBox, NearbyEntities);
+
+		// We'll pick the first entity that passes, this ensure that it's random one.
+		Algo::RandomShuffle(NearbyEntities);
+
+		const FVector Location = Transform.GetLocation();
+		
+		for (const FAvoidanceObstacleHashGrid2D::ItemIDType NearbyEntity : NearbyEntities)
+		{
+			// This can happen if we remove entities in the system.
+			if (!EntitySubsystem.IsEntityValid(NearbyEntity.Entity))
+			{
+				UE_LOG(LogMassBehavior, VeryVerbose, TEXT("Nearby entity is invalid, skipped."));
+				continue;
+			}
+
+			// Do not select same target twice in a row.
+			if (NearbyEntity.Entity == LastTrackedEntity)
+			{
+				continue;
+			}
+
+			FEntityView EntityView(EntitySubsystem, NearbyEntity.Entity);
+			if (!EntityView.HasTag<FMassLookAtTargetTag>())
+			{
+				continue;
+			}
+
+			// TargetTag is added through the LookAtTargetTrait and Transform was added with it
+			const FDataFragment_Transform& TargetTransform = EntityView.GetComponentData<FDataFragment_Transform>();
+			const FVector TargetLocation = TargetTransform.GetTransform().GetLocation();
+			FVector Direction = (TargetLocation - Location).GetSafeNormal();
+			Direction = Transform.InverseTransformVector(Direction);
+
+			const bool bIsTargetInView = FVector::DotProduct(FVector::ForwardVector, Direction) > CosAngleThreshold;
+			if (bIsTargetInView)
+			{
+				LookAt.GazeDirection = Direction;
+				LookAt.GazeTrackedEntity = NearbyEntity.Entity;
+				bTargetFound = true;
+				break;
+			}
+
+			// Allow to pick entities out of view if they are moving towards us.
+			if (const FMassVelocityFragment* Velocity = EntityView.GetComponentDataPtr<FMassVelocityFragment>())
+			{
+				const FVector MoveDirection = Transform.InverseTransformVector(Velocity->Value.GetSafeNormal());
+				
+				const bool bIsTargetMovingTowards = FVector::DotProduct(MoveDirection, -Direction) > CosAngleThreshold; // Direction negated as it is from the agent to target, and we want target to agent. 
+				if (bIsTargetMovingTowards)
+				{
+					LookAt.GazeDirection = Direction;
+					LookAt.GazeTrackedEntity = NearbyEntity.Entity;
+					bTargetFound = true;
+					break;
+				}
+			}
+			
+		}
+	}
+
+	// If no gaze target found, use random angle if specified.
+	if (!bTargetFound)
+	{
+		const FRotator Rot(FMath::FRandRange(-(float)LookAt.RandomGazePitchVariation, (float)LookAt.RandomGazePitchVariation),FMath::FRandRange(-(float)LookAt.RandomGazeYawVariation, (float)LookAt.RandomGazeYawVariation), 0.f);
+		LookAt.GazeDirection = UE::MassBehavior::ClampDirectionToXAxisCone(Rot.Vector(), FMath::DegreesToRadians(AngleThresholdInDegrees));
+	}
+
+	// @todo: This does not currently carry over time. It's intentional, since there might be big gaps between updates.
+	LookAt.GazeStartTime = CurrentTime;
+	LookAt.GazeDuration = FMath::FRandRange(FMath::Max(Duration - DurationVariation, 0.f), Duration + DurationVariation);
+}
+
+void UMassLookAtProcessor::UpdateLookAtTrajectory(const FTransform& Transform, const FMassZoneGraphLaneLocationFragment& ZoneGraphLocation,
+												  const FMassLookAtTrajectoryFragment& LookAtTrajectory, const bool bDisplayDebug, FMassLookAtFragment& LookAt) const
+{
+	if (LookAtTrajectory.NumPoints > 0 && LookAtTrajectory.LaneHandle == ZoneGraphLocation.LaneHandle)
+	{
+		// Look at anticipated position in the future when moving.
+		const float LookAheadDistanceAlongPath = ZoneGraphLocation.DistanceAlongLane + UE::MassBehavior::Tweakables::TrajectoryLookAhead * (LookAtTrajectory.bMoveReverse ? -1.0f : 1.0f);
+
+		// Calculate lookat direction to the anticipated position.
+		const FVector AnticipatedPosition = LookAtTrajectory.GetPointAtDistanceExtrapolated(LookAheadDistanceAlongPath);
+		const FVector AgentPosition = Transform.GetLocation();
+		const FVector NewGlobalDirection = (AnticipatedPosition - AgentPosition).GetSafeNormal();
+		LookAt.Direction = Transform.InverseTransformVector(NewGlobalDirection);
+		LookAt.Direction.Z = 0.0f;
+						
+#if WITH_MASS_DEBUG && UNSAFE_FOR_MT
+		if (bDisplayDebug)
+		{
+			const FVector ZOffset(0.f,0.f,DebugZOffset);
+			UE_VLOG_SEGMENT_THICK(this, LogMassBehavior, Display, AgentPosition + ZOffset, AgentPosition + ZOffset + 100.f * NewGlobalDirection, FColor::White, /*Thickness*/3.0f, TEXT("LookAt Trajectory"));
+		}
+#endif
+	}
+}
+
+void UMassLookAtProcessor::UpdateLookAtTrackedEntity(const UEntitySubsystem& EntitySubsystem, const FTransform& Transform, const bool bDisplayDebug, FMassLookAtFragment& LookAt) const
+{
+	// Update direction toward target
+	if (EntitySubsystem.IsEntityValid(LookAt.TrackedEntity))
+	{
+		if (const FDataFragment_Transform* TargetTransform = EntitySubsystem.GetComponentDataPtr<FDataFragment_Transform>(LookAt.TrackedEntity))
+		{
+			const FVector AgentPosition = Transform.GetLocation();
+			const FVector NewGlobalDirection = (TargetTransform->GetTransform().GetLocation() - AgentPosition).GetSafeNormal();
+			LookAt.Direction = Transform.InverseTransformVector(NewGlobalDirection);
+
+#if WITH_MASS_DEBUG && UNSAFE_FOR_MT
+			if (bDisplayDebug)
+			{
+				const FVector ZOffset(0.f,0.f,DebugZOffset);
+				UE_VLOG_SEGMENT_THICK(this, LogMassBehavior, Display, AgentPosition + ZOffset, AgentPosition + ZOffset + 100.f * NewGlobalDirection, FColor::White, /*Thickness*/3.0f, TEXT("LookAt Track"));
+			}
+#endif
+		}
+	}
+}
+
+bool UMassLookAtProcessor::UpdateGazeTrackedEntity(const UEntitySubsystem& EntitySubsystem, const FTransform& Transform, const bool bDisplayDebug, FMassLookAtFragment& LookAt) const
+{
+	bool bHasTarget = false;
+	
+	// Update direction toward gaze target
+	if (LookAt.GazeTrackedEntity.IsSet() && EntitySubsystem.IsEntityValid(LookAt.GazeTrackedEntity))
+	{
+		if (const FDataFragment_Transform* TargetTransform = EntitySubsystem.GetComponentDataPtr<FDataFragment_Transform>(LookAt.GazeTrackedEntity))
+		{
+			const FVector AgentPosition = Transform.GetLocation();
+			const FVector NewGlobalDirection = (TargetTransform->GetTransform().GetLocation() - AgentPosition).GetSafeNormal();
+			LookAt.GazeDirection = Transform.InverseTransformVector(NewGlobalDirection);
+
+			bHasTarget = true;
+			
+#if WITH_MASS_DEBUG && UNSAFE_FOR_MT
+			if (bDisplayDebug)
+			{
+				const FVector ZOffset(0.f,0.f,DebugZOffset);
+				UE_VLOG_SEGMENT_THICK(this, LogMassBehavior, Display, AgentPosition + ZOffset, AgentPosition + ZOffset + 100.f * NewGlobalDirection, FColor(160,160,160), /*Thickness*/3.0f, TEXT("Gaze Track"));
+			}
+#endif
+		}
+	}
+
+	return bHasTarget;
+}
+
+void UMassLookAtProcessor::BuildTrajectory(const FMassZoneGraphLaneLocationFragment& LaneLocation, const FMassZoneGraphShortPathFragment& ShortPath,
+											const FLWEntity Entity, const bool bDisplayDebug, FMassLookAtTrajectoryFragment& LookAtTrajectory)
+{
+	LookAtTrajectory.Reset();
+
+	if (ShortPath.NumPoints < 2)
+	{
+		return;
+	}
+
+	LookAtTrajectory.bMoveReverse = ShortPath.bMoveReverse;
+	LookAtTrajectory.LaneHandle = LaneLocation.LaneHandle;
+
+	const float NextLaneLookAheadDistance = UE::MassBehavior::Tweakables::TrajectoryLookAhead;
+	
+	// Initialize the look at trajectory from the current path.
+	const FMassZoneGraphPathPoint& FirstPathPoint = ShortPath.Points[0];
+	const FMassZoneGraphPathPoint& LastPathPoint = ShortPath.Points[ShortPath.NumPoints - 1];
+	ensure(LookAtTrajectory.AddPoint(FirstPathPoint.Position, FirstPathPoint.Tangent.Get(), FirstPathPoint.DistanceAlongLane.Get()));
+	ensure(LookAtTrajectory.AddPoint(LastPathPoint.Position, LastPathPoint.Tangent.Get(), LastPathPoint.DistanceAlongLane.Get()));
+
+	// If the path will lead to next lane, add a point from next lane too.
+	if (ShortPath.NextLaneHandle.IsValid())
+	{
+		const FZoneGraphStorage* ZoneGraphStorage = ZoneGraphSubsystem->GetZoneGraphStorage(LaneLocation.LaneHandle.DataHandle);
+		if (ZoneGraphStorage != nullptr)
+		{
+			if (ShortPath.NextExitLinkType == EZoneLaneLinkType::Outgoing)
+			{
+				FZoneGraphLaneLocation Location;
+				UE::ZoneGraph::Query::CalculateLocationAlongLane(*ZoneGraphStorage, ShortPath.NextLaneHandle, NextLaneLookAheadDistance, Location);
+
+				ensure(LookAtTrajectory.AddPoint(Location.Position, FVector2D(Location.Tangent), LastPathPoint.DistanceAlongLane.Get() + Location.DistanceAlongLane));
+			}
+			else if (ShortPath.NextExitLinkType == EZoneLaneLinkType::Incoming)
+			{
+				float LaneLength = 0.0f;
+				UE::ZoneGraph::Query::GetLaneLength(*ZoneGraphStorage, ShortPath.NextLaneHandle, LaneLength);
+
+				FZoneGraphLaneLocation Location;
+				UE::ZoneGraph::Query::CalculateLocationAlongLane(*ZoneGraphStorage, ShortPath.NextLaneHandle, LaneLength - NextLaneLookAheadDistance, Location);
+
+				// Moving backwards, reverse tangent and distance.
+				ensure(LookAtTrajectory.bMoveReverse);
+				ensure(LookAtTrajectory.AddPoint(Location.Position, FVector2D(-Location.Tangent), LastPathPoint.DistanceAlongLane.Get() - (LaneLength - Location.DistanceAlongLane)));
+			}
+			else if (ShortPath.NextExitLinkType == EZoneLaneLinkType::Adjacent)
+			{
+				// No extra point
+			}
+			else
+			{
+				ensureMsgf(false, TEXT("Unhandle NextExitLinkType type %s"), *UEnum::GetValueAsString(ShortPath.NextExitLinkType));
+			}
+		}
+		else
+		{
+			UE_CVLOG(bDisplayDebug, this, LogMassBehavior, Error, TEXT("%s Could not find ZoneGraph storage for lane %s."),
+				*Entity.DebugGetDescription(), *LaneLocation.LaneHandle.ToString());
+		}				
+	}
+
+	// Ensure that the points are always in ascending distance order (it is, in case of reverse path).
+	if (LookAtTrajectory.NumPoints > 1 && LookAtTrajectory.bMoveReverse)
+	{
+		ensureMsgf(LookAtTrajectory.Points[0].DistanceAlongLane.Get() >= LookAtTrajectory.Points[LookAtTrajectory.NumPoints - 1].DistanceAlongLane.Get(),
+			TEXT("Expecting trajectory 0 (%.1f) >= %d (%.1f)"), LookAtTrajectory.Points[0].DistanceAlongLane.Get(),
+			LookAtTrajectory.NumPoints - 1, LookAtTrajectory.Points[LookAtTrajectory.NumPoints - 1].DistanceAlongLane.Get());
+		
+		Algo::Reverse(LookAtTrajectory.Points.GetData(), LookAtTrajectory.NumPoints);
+		// Tangents needs to be reversed when the trajectory is reversed.
+		for (uint8 PointIndex = 0; PointIndex < LookAtTrajectory.NumPoints; PointIndex++)
+		{
+			LookAtTrajectory.Points[PointIndex].Tangent.Set(-LookAtTrajectory.Points[PointIndex].Tangent.Get());
+		}
+	}
+
+#if WITH_MASS_DEBUG && UNSAFE_FOR_MT
+	if (bDisplayDebug)
+	{
+		const FVector ZOffset(0,0,35);
+		
+		for (uint8 PointIndex = 0; PointIndex < LookAtTrajectory.NumPoints - 1; PointIndex++)
+		{
+			const FMassLookAtTrajectoryPoint& CurrPoint = LookAtTrajectory.Points[PointIndex];
+			const FMassLookAtTrajectoryPoint& NextPoint = LookAtTrajectory.Points[PointIndex + 1];
+
+			// Trajectory
+			const FVector StartPoint = CurrPoint.Position;
+			const FVector StartForward = CurrPoint.Tangent.GetVector();
+			const FVector EndPoint = NextPoint.Position;
+			const FVector EndForward = NextPoint.Tangent.GetVector();
+			const float TangentDistance = FVector::Dist(StartPoint, EndPoint) / 3.0f;
+			const FVector StartControlPoint = StartPoint + StartForward * TangentDistance;
+			const FVector EndControlPoint = EndPoint - EndForward * TangentDistance;
+
+			static constexpr int32 NumTicks = 6;
+			static constexpr float DeltaT = 1.0f / NumTicks;
+			
+			FVector PrevPoint = StartPoint;
+			for (int32 j = 0; j < NumTicks; j++)
+			{
+				const float T = (j + 1) * DeltaT;
+				const FVector Point = UE::CubicBezier::Eval(StartPoint, StartControlPoint, EndControlPoint, EndPoint, T);
+				UE_VLOG_SEGMENT_THICK(this, LogMassBehavior, Display, PrevPoint + ZOffset, Point + ZOffset, FColor::White, /*Thickness*/3.0f, TEXT(""));
+				PrevPoint = Point;
+			}
+		}
+		
+		for (uint8 PointIndex = 0; PointIndex < LookAtTrajectory.NumPoints; PointIndex++)
+		{
+			const FMassLookAtTrajectoryPoint& CurrPoint = LookAtTrajectory.Points[PointIndex];
+			const FVector CurrBase = CurrPoint.Position + ZOffset * 1.1f;
+			// Tangents
+			UE_VLOG_SEGMENT_THICK(this, LogMassBehavior, Display, CurrBase, CurrBase + CurrPoint.Tangent.GetVector() * 100.0f, FColorList::Grey, /*Thickness*/1.0f,
+				TEXT("D:%.1f"), CurrPoint.DistanceAlongLane.Get());
+		}
+	}
+#endif // WITH_MASS_DEBUG
+}
+
+#undef UNSAFE_FOR_MT
\ No newline at end of file
diff --git a/Engine/Plugins/AI/MassAI/Source/MassAIBehavior/Private/MassLookAtTargetTrait.cpp b/Engine/Plugins/AI/MassAI/Source/MassAIBehavior/Private/MassLookAtTargetTrait.cpp
new file mode 100644
--- /dev/null
+++ b/Engine/Plugins/AI/MassAI/Source/MassAIBehavior/Private/MassLookAtTargetTrait.cpp
@@ -0,0 +1,12 @@
+// Copyright Epic Games, Inc. All Rights Reserved.
+
+#include "MassLookAtTargetTrait.h"
+#include "MassLookAtFragments.h"
+#include "MassCommonFragments.h"
+#include "MassEntityTemplateRegistry.h"
+
+void UMassLookAtTargetTrait::BuildTemplate(FMassEntityTemplateBuildContext& BuildContext, UWorld& World) const
+{
+	BuildContext.AddTag<FMassLookAtTargetTag>();
+	BuildContext.AddFragment<FDataFragment_Transform>();
+}
diff --git a/Engine/Plugins/AI/MassAI/Source/MassAIBehavior/Private/MassLookAtTrait.cpp b/Engine/Plugins/AI/MassAI/Source/MassAIBehavior/Private/MassLookAtTrait.cpp
new file mode 100644
--- /dev/null
+++ b/Engine/Plugins/AI/MassAI/Source/MassAIBehavior/Private/MassLookAtTrait.cpp
@@ -0,0 +1,11 @@
+﻿// Copyright Epic Games, Inc. All Rights Reserved.
+
+#include "MassLookAtTrait.h"
+#include "MassLookAtFragments.h"
+#include "MassEntityTemplateRegistry.h"
+
+void UMassLookAtTrait::BuildTemplate(FMassEntityTemplateBuildContext& BuildContext, UWorld& World) const
+{
+	BuildContext.AddFragment<FMassLookAtFragment>();
+	BuildContext.AddFragment<FMassLookAtTrajectoryFragment>();
+}
diff --git a/Engine/Plugins/AI/MassAI/Source/MassAIBehavior/Private/MassStateTreeExecutionContext.cpp b/Engine/Plugins/AI/MassAI/Source/MassAIBehavior/Private/MassStateTreeExecutionContext.cpp
new file mode 100644
--- /dev/null
+++ b/Engine/Plugins/AI/MassAI/Source/MassAIBehavior/Private/MassStateTreeExecutionContext.cpp
@@ -0,0 +1,24 @@
+// Copyright Epic Games, Inc. All Rights Reserved.
+
+#include "MassStateTreeExecutionContext.h"
+#include "MassStateTreeTypes.h"
+#include "EntitySubsystem.h"
+#include "MassSignalSubsystem.h"
+#include "Engine/World.h"
+
+FMassStateTreeExecutionContext::FMassStateTreeExecutionContext(UEntitySubsystem& InEntitySubsystem,
+                                                               FLWComponentSystemExecutionContext& InContext):
+	EntitySubsystem(&InEntitySubsystem), EntitySubsystemExecutionContext(&InContext)
+{
+	World = InEntitySubsystem.GetWorld();
+	SignalSubsystem = UWorld::GetSubsystem<UMassSignalSubsystem>(World);
+}
+
+void FMassStateTreeExecutionContext::BeginGatedTransition(const FStateTreeExecutionState& Exec)
+{
+	if (SignalSubsystem != nullptr && Entity.IsSet())
+	{
+		// Tick again after the games time has passed to see if the condition still holds true.
+		SignalSubsystem->DelaySignalEntity(UE::Mass::Signals::DelayedTransitionWakeup, Entity, Exec.GatedTransitionTime + KINDA_SMALL_NUMBER);
+	}
+}
diff --git a/Engine/Plugins/AI/MassAI/Source/MassAIBehavior/Private/MassStateTreeProcessors.cpp b/Engine/Plugins/AI/MassAI/Source/MassAIBehavior/Private/MassStateTreeProcessors.cpp
new file mode 100644
--- /dev/null
+++ b/Engine/Plugins/AI/MassAI/Source/MassAIBehavior/Private/MassStateTreeProcessors.cpp
@@ -0,0 +1,333 @@
+// Copyright Epic Games, Inc. All Rights Reserved.
+
+#include "MassStateTreeProcessors.h"
+#include "StateTree.h"
+#include "MassStateTreeExecutionContext.h"
+#include "EntityView.h"
+#include "MassAIMovementTypes.h"
+#include "MassComponentHitTypes.h"
+#include "MassSmartObjectTypes.h"
+#include "MassZoneGraphAnnotationTypes.h"
+#include "MassStateTreeSubsystem.h"
+#include "MassSignals/Public/MassSignalSubsystem.h"
+#include "ProfilingDebugging/CsvProfiler.h"
+#include "Engine/World.h"
+
+CSV_DEFINE_CATEGORY(StateTreeProcessor, true);
+
+namespace UE::MassBehavior
+{
+
+bool SetExternalFragments(FMassStateTreeExecutionContext& Context, const UEntitySubsystem& EntitySubsystem)
+{
+	bool bFoundAllFragments = true;
+	const FEntityView EntityView(EntitySubsystem, Context.GetEntity());
+	for (const FStateTreeExternalItemDesc& ItemDesc : Context.GetExternalItems())
+	{
+		if (ItemDesc.Struct && ItemDesc.Struct->IsChildOf(FLWComponentData::StaticStruct()))
+		{
+			const UScriptStruct* ScriptStruct = Cast<const UScriptStruct>(ItemDesc.Struct);
+			FStructView Fragment = EntityView.GetComponentDataStruct(ScriptStruct);
+			if (Fragment.IsValid())
+			{
+				Context.SetExternalItem(ItemDesc.Handle, FStateTreeItemView(Fragment));
+			}
+			else
+			{
+				if (!ItemDesc.bOptional)
+				{
+					// Note: Not breaking here, so that we can validate all missing ones in one go with FMassStateTreeExecutionContext::AreExternalItemsValid().
+					bFoundAllFragments = false;
+				}
+			}
+		}
+	}
+	return bFoundAllFragments;
+}
+
+bool SetExternalSubsystems(FMassStateTreeExecutionContext& Context)
+{
+	const UWorld* World = Context.GetWorld();
+	if (World == nullptr)
+	{
+		return false;
+	}
+
+	bool bFoundAllSubsystems = true;
+	for (const FStateTreeExternalItemDesc& ItemDesc : Context.GetExternalItems())
+	{
+		if (ItemDesc.Struct && ItemDesc.Struct->IsChildOf(UWorldSubsystem::StaticClass()))
+		{
+			const TSubclassOf<UWorldSubsystem> SubClass = Cast<UClass>(const_cast<UStruct*>(ItemDesc.Struct));
+			UWorldSubsystem* Subsystem = World->GetSubsystemBase(SubClass);
+			if (Subsystem)
+			{
+				Context.SetExternalItem(ItemDesc.Handle, FStateTreeItemView(Subsystem));
+			}
+			else
+			{
+				if (!ItemDesc.bOptional)
+				{
+					// Note: Not breaking here, so that we can validate all missing ones in one go with FMassStateTreeExecutionContext::AreExternalItemsValid().
+					bFoundAllSubsystems = false;
+				}
+			}
+		}
+	}
+	return bFoundAllSubsystems;
+}
+
+void ProcessChunk(
+	FMassStateTreeExecutionContext& StateTreeContext,
+	UMassStateTreeSubsystem& MassStateTreeSubsystem,
+	const TFunctionRef<void(FMassStateTreeExecutionContext&, FStateTreeItemView)> ForEachEntityCallback)
+{
+	const FLWComponentSystemExecutionContext& Context = StateTreeContext.GetEntitySubsystemExecutionContext();
+	const TConstArrayView<FMassStateTreeFragment> StateTreeList = Context.GetComponentView<FMassStateTreeFragment>();
+
+	// Assuming that all the entities share same StateTree, because they all should have the same storage fragment.
+	const int32 NumEntities = Context.GetEntitiesNum();
+	check(NumEntities > 0);
+	UStateTree* StateTree = MassStateTreeSubsystem.GetRegisteredStateTreeAsset(StateTreeList[0].StateTreeHandle);
+
+	// Initialize the execution context if changed between chunks.
+	if (StateTreeContext.GetStateTree() != StateTree)
+	{
+		StateTreeContext.Init(MassStateTreeSubsystem, *StateTree, EStateTreeStorage::External);
+
+		// Gather subsystems.
+		{
+			CSV_SCOPED_TIMING_STAT_EXCLUSIVE(StateTreeProcessorExternalSubsystems);
+			if (!ensureMsgf(UE::MassBehavior::SetExternalSubsystems(StateTreeContext), TEXT("StateTree will not execute due to missing subsystem requirements.")))
+			{
+				return;
+			}
+		}
+	}
+
+	const UScriptStruct* StorageScriptStruct = StateTree->GetRuntimeStorageStruct();
+	for (int32 i = 0; i < NumEntities; ++i)
+	{
+		const FLWEntity Entity = Context.GetEntity(i);
+		StateTreeContext.SetEntity(Entity);
+		StateTreeContext.SetEntityIndex(i);
+
+		// Gather all required fragments.
+		{
+			CSV_SCOPED_TIMING_STAT_EXCLUSIVE(StateTreeProcessorExternalFragments);
+			if (!ensureMsgf(UE::MassBehavior::SetExternalFragments(StateTreeContext, StateTreeContext.GetEntitySubsystem()), TEXT("StateTree will not execute due to missing required fragments.")))
+			{
+				break;
+			}
+		}
+
+		// Make sure all required external items are set.
+		{
+			CSV_SCOPED_TIMING_STAT_EXCLUSIVE(StateTreeProcessorExternalItemsValidation);
+			// TODO: disable this when not in debug.
+			if (!ensureMsgf(StateTreeContext.AreExternalItemsValid(), TEXT("StateTree will not execute due to missing external items.")))
+			{
+				break;
+			}
+		}
+
+		ForEachEntityCallback(StateTreeContext, StateTreeContext.GetEntitySubsystem().GetComponentDataStruct(Entity, StorageScriptStruct));
+	}
+}
+
+} // UE::MassBehavior
+
+
+//----------------------------------------------------------------------//
+// UMassStateTreeFragmentInitializer
+//----------------------------------------------------------------------//
+UMassStateTreeFragmentInitializer::UMassStateTreeFragmentInitializer()
+{
+	ExecutionFlags = (int32)(EProcessorExecutionFlags::Standalone | EProcessorExecutionFlags::Server);
+	FragmentType = FMassStateTreeFragment::StaticStruct();
+}
+
+void UMassStateTreeFragmentInitializer::ConfigureQueries()
+{
+	EntityQuery.AddRequirement<FMassStateTreeFragment>(ELWComponentAccess::ReadOnly);
+}
+
+void UMassStateTreeFragmentInitializer::Execute(UEntitySubsystem& EntitySubsystem, FLWComponentSystemExecutionContext& Context)
+{
+	// StateTree processor relies on signals to be ticked but we need an 'initial tick' to set the tree in the proper state.
+	// The initializer provides that by sending a signal to all new entities that use StateTree.
+	TArray<FLWEntity> EntitiesToSignal;
+
+	FMassStateTreeExecutionContext StateTreeContext(EntitySubsystem, Context);
+	UMassStateTreeSubsystem* MassStateTreeSubsystem = UWorld::GetSubsystem<UMassStateTreeSubsystem>(EntitySubsystem.GetWorld());
+
+	EntityQuery.ForEachEntityChunk(
+		EntitySubsystem,
+		Context,
+		[this, &StateTreeContext, &MassStateTreeSubsystem, &EntitiesToSignal](const FLWComponentSystemExecutionContext& Context)
+		{
+			UE::MassBehavior::ProcessChunk(
+				StateTreeContext,
+				*MassStateTreeSubsystem,
+				[](FMassStateTreeExecutionContext& StateTreeExecutionContext, FStateTreeItemView Storage)
+				{
+					// Start the tree instance
+					StateTreeExecutionContext.Start(Storage);
+				});
+
+			// Append all entities of the current chunk to the consolidated list to send signal once
+			EntitiesToSignal.Append(Context.GetEntities().GetData(), Context.GetEntities().Num());
+		});
+
+	// Signal all entities inside the consolidated list
+	if (EntitiesToSignal.Num())
+	{
+		UMassSignalSubsystem* SignalSubsystem = UWorld::GetSubsystem<UMassSignalSubsystem>(EntitySubsystem.GetWorld());
+		checkf(SignalSubsystem != nullptr, TEXT("MassSignalSubsystem should exist when executing fragment initializers."));
+		SignalSubsystem->SignalEntities(UE::Mass::Signals::StateTreeInitializationRequested, EntitiesToSignal);
+	}
+}
+
+//----------------------------------------------------------------------//
+// UMassStateTreeFragmentDestructor
+//----------------------------------------------------------------------//
+UMassStateTreeFragmentDestructor::UMassStateTreeFragmentDestructor()
+{
+	ExecutionFlags = (int32)(EProcessorExecutionFlags::Standalone | EProcessorExecutionFlags::Server);
+	FragmentType = FMassStateTreeFragment::StaticStruct();
+}
+
+void UMassStateTreeFragmentDestructor::ConfigureQueries()
+{
+	EntityQuery.AddRequirement<FMassStateTreeFragment>(ELWComponentAccess::ReadOnly);
+}
+
+void UMassStateTreeFragmentDestructor::Execute(UEntitySubsystem& EntitySubsystem, FLWComponentSystemExecutionContext& Context)
+{
+	FMassStateTreeExecutionContext StateTreeContext(EntitySubsystem, Context);
+	UMassStateTreeSubsystem* MassStateTreeSubsystem = UWorld::GetSubsystem<UMassStateTreeSubsystem>(EntitySubsystem.GetWorld());
+
+	EntityQuery.ForEachEntityChunk(
+		EntitySubsystem,
+		Context,
+		[this, &StateTreeContext, &MassStateTreeSubsystem](FLWComponentSystemExecutionContext&)
+		{
+			UE::MassBehavior::ProcessChunk(
+				StateTreeContext,
+				*MassStateTreeSubsystem,
+				[](FMassStateTreeExecutionContext& StateTreeExecutionContext, FStateTreeItemView Storage)
+				{
+					// Stop the tree instance
+					StateTreeExecutionContext.Stop(Storage);
+				});
+		});
+}
+
+//----------------------------------------------------------------------//
+// UMassStateTreeProcessor
+//----------------------------------------------------------------------//
+UMassStateTreeProcessor::UMassStateTreeProcessor(const FObjectInitializer& ObjectInitializer)
+	: Super(ObjectInitializer)
+{
+	bRequiresGameThreadExecution = true;
+
+	ExecutionOrder.ExecuteInGroup = UE::Mass::ProcessorGroupNames::Behavior;
+
+	// `Behavior` doesn't run on clients but `Tasks` do.
+	// We define the dependencies here so task won't need to set their dependency on `Behavior`,
+	// but only on `SyncWorldToMass`
+	ExecutionOrder.ExecuteAfter.Add(UE::Mass::ProcessorGroupNames::SyncWorldToMass);
+	ExecutionOrder.ExecuteBefore.Add(UE::Mass::ProcessorGroupNames::Tasks);
+}
+
+void UMassStateTreeProcessor::Initialize(UObject& Owner)
+{
+	Super::Initialize(Owner);
+	MassStateTreeSubsystem = UWorld::GetSubsystem<UMassStateTreeSubsystem>(Owner.GetWorld());
+
+	SubscribeToSignal(UE::Mass::Signals::StateTreeInitializationRequested);
+	SubscribeToSignal(UE::Mass::Signals::LookAtFinished);
+	SubscribeToSignal(UE::Mass::Signals::NewStateTreeTaskRequired);
+	SubscribeToSignal(UE::Mass::Signals::StandTaskFinished);
+	SubscribeToSignal(UE::Mass::Signals::DelayedTransitionWakeup);
+
+	// @todo MassStateTree: add a way to register/unregister from enter/exit state (need reference counting)
+	SubscribeToSignal(UE::Mass::Signals::SmartObjectRequestCandidates);
+	SubscribeToSignal(UE::Mass::Signals::SmartObjectCandidatesReady);
+	SubscribeToSignal(UE::Mass::Signals::SmartObjectInteractionDone);
+
+	SubscribeToSignal(UE::Mass::Signals::FollowPointPathStart);
+	SubscribeToSignal(UE::Mass::Signals::FollowPointPathDone);
+	SubscribeToSignal(UE::Mass::Signals::CurrentLaneChanged);
+
+	SubscribeToSignal(UE::Mass::Signals::AnnotationTagsChanged);
+
+	SubscribeToSignal(UE::Mass::Signals::HitReceived);
+
+	// @todo MassStateTree: move this to its game plugin when possible
+	SubscribeToSignal(UE::Mass::Signals::ContextualAnimTaskFinished);
+}
+
+void UMassStateTreeProcessor::ConfigureQueries()
+{
+	EntityQuery.AddRequirement<FMassStateTreeFragment>(ELWComponentAccess::ReadWrite);
+}
+
+void UMassStateTreeProcessor::SignalEntities(UEntitySubsystem& EntitySubsystem, FLWComponentSystemExecutionContext& Context, FMassSignalNameLookup& EntitySignals)
+{
+	if (!MassStateTreeSubsystem)
+	{
+		return;
+	}
+	QUICK_SCOPE_CYCLE_COUNTER(StateTreeProcessor_Run);
+	CSV_SCOPED_TIMING_STAT_EXCLUSIVE(StateTreeProcessorExecute);
+
+	const float TimeDelta = Context.GetDeltaTimeSeconds();
+	FMassStateTreeExecutionContext StateTreeContext(EntitySubsystem, Context);
+	const float TimeInSeconds = EntitySubsystem.GetWorld()->GetTimeSeconds();
+
+	TArray<FLWEntity> EntitiesToSignal;
+
+	EntityQuery.ForEachEntityChunk(
+		EntitySubsystem,
+		Context,
+		[this, &StateTreeContext, TimeDelta, TimeInSeconds, &EntitiesToSignal](FLWComponentSystemExecutionContext& Context)
+		{
+			TArrayView<FMassStateTreeFragment> StateTreeList = Context.GetMutableComponentView<FMassStateTreeFragment>();
+
+			UE::MassBehavior::ProcessChunk(
+				StateTreeContext,
+				*MassStateTreeSubsystem,
+				[&StateTreeList, TimeDelta, TimeInSeconds, &EntitiesToSignal](FMassStateTreeExecutionContext& StateTreeExecutionContext, const FStateTreeItemView Storage)
+				{
+					// Keep stats regarding the amount of tree instances ticked per frame
+					CSV_CUSTOM_STAT(StateTreeProcessor, NumTickedStateTree, StateTreeExecutionContext.GetEntitySubsystemExecutionContext().GetEntitiesNum(), ECsvCustomStatOp::Accumulate);
+
+					// Compute adjusted delta time
+					TOptional<float>& LastUpdate = StateTreeList[StateTreeExecutionContext.GetEntityIndex()].LastUpdateTimeInSeconds;
+					const float AdjustedTimeDelta = LastUpdate.IsSet() ? TimeDelta + (TimeInSeconds - LastUpdate.GetValue()) : TimeDelta;
+					LastUpdate = TimeInSeconds;
+
+					// Tick the tree instance
+					StateTreeExecutionContext.Tick(AdjustedTimeDelta, Storage);
+
+					// When last tick status is different than "Running", the state tree need to be tick again
+					// For performance reason, tick again to see if we could find a new state right away instead of waiting the next frame.
+					if (StateTreeExecutionContext.GetLastTickStatus(Storage) != EStateTreeRunStatus::Running)
+					{
+						StateTreeExecutionContext.Tick(0.0f, Storage);
+
+						// Could not find new state yet, try again next frame
+						if (StateTreeExecutionContext.GetLastTickStatus(Storage) != EStateTreeRunStatus::Running)
+						{
+							EntitiesToSignal.Add(StateTreeExecutionContext.GetEntity());
+						}
+					}
+				});
+		});
+
+	if (EntitiesToSignal.Num())
+	{
+		SignalSubsystem->SignalEntities(UE::Mass::Signals::NewStateTreeTaskRequired, EntitiesToSignal);
+	}
+}
diff --git a/Engine/Plugins/AI/MassAI/Source/MassAIBehavior/Private/MassStateTreeSchema.cpp b/Engine/Plugins/AI/MassAI/Source/MassAIBehavior/Private/MassStateTreeSchema.cpp
new file mode 100644
--- /dev/null
+++ b/Engine/Plugins/AI/MassAI/Source/MassAIBehavior/Private/MassStateTreeSchema.cpp
@@ -0,0 +1,24 @@
+﻿// Copyright Epic Games, Inc. All Rights Reserved.
+
+#include "MassStateTreeSchema.h"
+#include "LWComponentTypes.h"
+#include "MassStateTreeTypes.h"
+#include "Subsystems/WorldSubsystem.h"
+
+UScriptStruct* UMassStateTreeSchema::GetStorageSuperStruct() const
+{
+	// The runtime state must be derived from FLWComponentData since it's stored in the entity subsystem.
+	return FLWComponentData::StaticStruct();
+}
+
+bool UMassStateTreeSchema::IsStructAllowed(const UScriptStruct* InScriptStruct) const
+{
+	// Only allow Mass evals and tasks.
+	return InScriptStruct->IsChildOf(FMassStateTreeEvaluatorBase::StaticStruct()) || InScriptStruct->IsChildOf(FMassStateTreeTaskBase::StaticStruct());
+}
+
+bool UMassStateTreeSchema::IsExternalItemAllowed(const UStruct& InStruct) const
+{
+	// Allow only WorldSubsystems and fragments as external data.
+	return InStruct.IsChildOf(UWorldSubsystem::StaticClass()) || InStruct.IsChildOf(FLWComponentData::StaticStruct());
+}
diff --git a/Engine/Plugins/AI/MassAI/Source/MassAIBehavior/Private/MassStateTreeSubsystem.cpp b/Engine/Plugins/AI/MassAI/Source/MassAIBehavior/Private/MassStateTreeSubsystem.cpp
new file mode 100644
--- /dev/null
+++ b/Engine/Plugins/AI/MassAI/Source/MassAIBehavior/Private/MassStateTreeSubsystem.cpp
@@ -0,0 +1,27 @@
+// Copyright Epic Games, Inc. All Rights Reserved.
+
+#include "MassStateTreeSubsystem.h"
+#include "StateTree.h"
+#include "Engine/Engine.h"
+
+
+void UMassStateTreeSubsystem::Initialize(FSubsystemCollectionBase& Collection)
+{
+	Super::Initialize(Collection);
+	
+}
+
+FMassStateTreeHandle UMassStateTreeSubsystem::RegisterStateTreeAsset(UStateTree* StateTree)
+{
+	// Return already registered asset if found.
+	int32 Index = RegisteredStateTrees.IndexOfByPredicate([StateTree](UStateTree* ExistingStateTree) -> bool { return ExistingStateTree == StateTree; });
+	if (Index != INDEX_NONE)
+	{
+		return FMassStateTreeHandle::Make(uint16(Index));
+	}
+	// Add new, check that it fits the StateTree handle.
+	Index = RegisteredStateTrees.Add(StateTree);
+	check(Index < int32(MAX_uint16));
+	return FMassStateTreeHandle::Make(uint16(Index));
+}
+
diff --git a/Engine/Plugins/AI/MassAI/Source/MassAIBehavior/Private/MassStateTreeTrait.cpp b/Engine/Plugins/AI/MassAI/Source/MassAIBehavior/Private/MassStateTreeTrait.cpp
new file mode 100644
--- /dev/null
+++ b/Engine/Plugins/AI/MassAI/Source/MassAIBehavior/Private/MassStateTreeTrait.cpp
@@ -0,0 +1,79 @@
+// Copyright Epic Games, Inc. All Rights Reserved.
+
+#include "MassStateTreeTrait.h"
+#include "MassStateTreeFragments.h"
+#include "MassStateTreeSubsystem.h"
+#include "MassAIBehaviorTypes.h"
+#include "StateTree.h"
+#include "VisualLogger/VisualLogger.h"
+#include "MassEntityTemplateRegistry.h"
+#include "Engine/World.h"
+
+void UMassStateTreeTrait::BuildTemplate(FMassEntityTemplateBuildContext& BuildContext, UWorld& World) const
+{
+	UMassStateTreeSubsystem* MassStateTreeSubsystem = World.GetSubsystem<UMassStateTreeSubsystem>();
+	if (!MassStateTreeSubsystem)
+	{
+		UE_VLOG(&World, LogMassBehavior, Error, TEXT("Failed to get Mass StateTree Subsystem."));
+		return;
+	}
+
+	if (!StateTree)
+	{
+		UE_VLOG(MassStateTreeSubsystem, LogMassBehavior, Error, TEXT("StateTree asset is not set or unavailable."));
+		return;
+	}
+	if (!StateTree->GetRuntimeStorageDefaultValue().IsValid())
+	{
+		UE_VLOG(MassStateTreeSubsystem, LogMassBehavior, Error, TEXT("StateTree asset is valid but missing runtime storage type."));
+		return;
+	}
+
+	const FMassStateTreeHandle Handle = MassStateTreeSubsystem->RegisterStateTreeAsset(StateTree);
+
+	// Add fragment describing which StateTree to run
+	FMassStateTreeFragment& StateTreeFragment = BuildContext.AddFragmentWithDefaultInitializer_GetRef<FMassStateTreeFragment>();
+	StateTreeFragment.StateTreeHandle = Handle;
+
+	// Add runtime storage as a fragment
+	BuildContext.AddFragment(StateTree->GetRuntimeStorageDefaultValue());
+}
+
+void UMassStateTreeTrait::ValidateTemplate(FMassEntityTemplateBuildContext& BuildContext, UWorld& World) const
+{
+	UMassStateTreeSubsystem* MassStateTreeSubsystem = World.GetSubsystem<UMassStateTreeSubsystem>();
+	if (!MassStateTreeSubsystem)
+	{
+		UE_VLOG(&World, LogMassBehavior, Error, TEXT("Failed to get Mass StateTree Subsystem."));
+		return;
+	}
+
+	if (!StateTree)
+	{
+		UE_VLOG(MassStateTreeSubsystem, LogMassBehavior, Error, TEXT("StateTree asset is not set or unavailable."));
+		return;
+	}
+
+	// Make sure all the required subsystems can be found.
+	for (const FStateTreeExternalItemDesc& ItemDesc : StateTree->GetExternalItems())
+	{
+		if (ensure(ItemDesc.Struct) && !ItemDesc.bOptional)
+		{
+			if (ItemDesc.Struct->IsChildOf(UWorldSubsystem::StaticClass()))
+			{
+				const TSubclassOf<UWorldSubsystem> SubClass = Cast<UClass>(const_cast<UStruct*>(ItemDesc.Struct));
+				USubsystem* Subsystem = World.GetSubsystemBase(SubClass);
+				UE_CVLOG(!Subsystem, MassStateTreeSubsystem, LogMassBehavior, Error, TEXT("StateTree %s: Could not find required subsystem %s"), *GetNameSafe(StateTree), *GetNameSafe(ItemDesc.Struct));
+			}
+			else if (ItemDesc.Struct->IsChildOf(FLWComponentData::StaticStruct()))
+			{
+				const bool bContainsFragment = BuildContext.GetMutableFragments().GetFragments().ContainsByPredicate([&ItemDesc](const FInstancedStruct& Item) { return Item.GetScriptStruct() == ItemDesc.Struct; });
+				UE_CVLOG(!bContainsFragment, MassStateTreeSubsystem, LogMassBehavior, Error, TEXT("StateTree %s: Could not find required fragment %s"), *GetNameSafe(StateTree), *GetNameSafe(ItemDesc.Struct));
+			}
+			else
+			{
+				UE_VLOG(MassStateTreeSubsystem, LogMassBehavior, Error, TEXT("StateTree %s: Unsupported requirement %s"), *GetNameSafe(StateTree), *GetNameSafe(ItemDesc.Struct));
+			}
+		}
+	}
+}
diff --git a/Engine/Plugins/AI/MassAI/Source/MassAIBehavior/Private/MassZoneGraphAnnotationFragments.cpp b/Engine/Plugins/AI/MassAI/Source/MassAIBehavior/Private/MassZoneGraphAnnotationFragments.cpp
new file mode 100644
--- /dev/null
+++ b/Engine/Plugins/AI/MassAI/Source/MassAIBehavior/Private/MassZoneGraphAnnotationFragments.cpp
@@ -0,0 +1,42 @@
+// Copyright Epic Games, Inc. All Rights Reserved.
+
+#include "MassZoneGraphAnnotationFragments.h"
+#include "EntitySubsystem.h"
+#include "MassSimulationLOD.h"
+
+namespace UE::Mass::ZoneGraphAnnotations
+{
+// Update interval range for periodic annotation tag update.
+static const float MinUpdateInterval = 0.25f;
+static const float MaxUpdateInterval = 0.5f;
+
+// Update interval range for periodic annotation tag update for Off LOD.
+static const float OffLODMinUpdateInterval = 1.905f;
+static const float OffLODMaxUpdateInterval = 2.10f;
+
+} // UE::Mass::ZoneGraphAnnotations
+
+//----------------------------------------------------------------------//
+//  FMassZoneGraphAnnotationVariableTickChunkFragment
+//----------------------------------------------------------------------//
+
+bool FMassZoneGraphAnnotationVariableTickChunkFragment::UpdateChunk(FLWComponentSystemExecutionContext& Context)
+{
+	FMassZoneGraphAnnotationVariableTickChunkFragment& ChunkFrag = Context.GetMutableChunkComponent<FMassZoneGraphAnnotationVariableTickChunkFragment>();
+	ChunkFrag.TimeUntilNextTick -= Context.GetDeltaTimeSeconds();
+	if (ChunkFrag.TimeUntilNextTick <= 0.0f)
+	{
+		// @todo Possible future optimization, right now it is faster but not by a whole lot.
+		if (FMassSimulationVariableTickChunkFragment::GetChunkLOD(Context) == EMassLOD::Off)
+		{
+			ChunkFrag.TimeUntilNextTick = FMath::RandRange(UE::Mass::ZoneGraphAnnotations::OffLODMinUpdateInterval, UE::Mass::ZoneGraphAnnotations::OffLODMaxUpdateInterval);
+		}
+		else
+		{
+			ChunkFrag.TimeUntilNextTick = FMath::RandRange(UE::Mass::ZoneGraphAnnotations::MinUpdateInterval, UE::Mass::ZoneGraphAnnotations::MaxUpdateInterval);
+		}
+		return true;
+	}
+
+	return false;
+}
diff --git a/Engine/Plugins/AI/MassAI/Source/MassAIBehavior/Private/MassZoneGraphAnnotationProcessors.cpp b/Engine/Plugins/AI/MassAI/Source/MassAIBehavior/Private/MassZoneGraphAnnotationProcessors.cpp
new file mode 100644
--- /dev/null
+++ b/Engine/Plugins/AI/MassAI/Source/MassAIBehavior/Private/MassZoneGraphAnnotationProcessors.cpp
@@ -0,0 +1,160 @@
+// Copyright Epic Games, Inc. All Rights Reserved.
+
+#include "MassZoneGraphAnnotationProcessors.h"
+#include "MassAIBehaviorTypes.h"
+#include "MassZoneGraphMovementFragments.h"
+#include "MassZoneGraphAnnotationFragments.h"
+#include "MassZoneGraphAnnotationTypes.h"
+#include "ZoneGraphTypes.h"
+#include "ZoneGraphAnnotationSubsystem.h"
+#include "MassSignalSubsystem.h"
+#include "MassSimulationLOD.h"
+#include "Engine/World.h"
+
+//----------------------------------------------------------------------//
+//  UMassZoneGraphAnnotationTagsInitializer
+//----------------------------------------------------------------------//
+UMassZoneGraphAnnotationTagsInitializer::UMassZoneGraphAnnotationTagsInitializer()
+{
+	FragmentType = FMassZoneGraphAnnotationTagsFragment::StaticStruct();
+}
+
+void UMassZoneGraphAnnotationTagsInitializer::Initialize(UObject& Owner)
+{
+	Super::Initialize(Owner);
+	ZoneGraphAnnotationSubsystem = UWorld::GetSubsystem<UZoneGraphAnnotationSubsystem>(Owner.GetWorld());
+}
+
+void UMassZoneGraphAnnotationTagsInitializer::ConfigureQueries()
+{
+	EntityQuery.AddRequirement<FMassZoneGraphAnnotationTagsFragment>(ELWComponentAccess::ReadWrite);
+	EntityQuery.AddRequirement<FMassZoneGraphLaneLocationFragment>(ELWComponentAccess::ReadOnly);
+}
+
+void UMassZoneGraphAnnotationTagsInitializer::Execute(UEntitySubsystem& EntitySubsystem, FLWComponentSystemExecutionContext& Context)
+{
+	EntityQuery.ForEachEntityChunk(EntitySubsystem, Context, [this](FLWComponentSystemExecutionContext& Context)
+	{
+		const int32 NumEntities = Context.GetEntitiesNum();
+		const TArrayView<FMassZoneGraphAnnotationTagsFragment> AnnotationTagsList = Context.GetMutableComponentView<FMassZoneGraphAnnotationTagsFragment>();
+		const TConstArrayView<FMassZoneGraphLaneLocationFragment> LaneLocationList = Context.GetComponentView<FMassZoneGraphLaneLocationFragment>();
+
+		for (int32 EntityIndex = 0; EntityIndex < NumEntities; ++EntityIndex)
+		{
+			FMassZoneGraphAnnotationTagsFragment& AnnotationTags = AnnotationTagsList[EntityIndex];
+			const FMassZoneGraphLaneLocationFragment& LaneLocation = LaneLocationList[EntityIndex];
+
+			if (!LaneLocation.LaneHandle.IsValid())
+			{
+				AnnotationTags.Tags = FZoneGraphTagMask::None;
+			}
+			else
+			{
+				AnnotationTags.Tags = ZoneGraphAnnotationSubsystem->GetAnnotationTags(LaneLocation.LaneHandle);
+			}
+		}
+	});
+}
+
+//----------------------------------------------------------------------//
+//  UMassZoneGraphAnnotationTagUpdateProcessor
+//----------------------------------------------------------------------//
+UMassZoneGraphAnnotationTagUpdateProcessor::UMassZoneGraphAnnotationTagUpdateProcessor()
+{
+	ExecutionOrder.ExecuteInGroup = UE::Mass::ProcessorGroupNames::UpdateAnnotationTags;
+	ExecutionOrder.ExecuteAfter.Add(UE::Mass::ProcessorGroupNames::SyncWorldToMass);
+	ExecutionOrder.ExecuteBefore.Add(UE::Mass::ProcessorGroupNames::Behavior);
+}
+
+void UMassZoneGraphAnnotationTagUpdateProcessor::Initialize(UObject& Owner)
+{
+	Super::Initialize(Owner);
+	ZoneGraphAnnotationSubsystem = UWorld::GetSubsystem<UZoneGraphAnnotationSubsystem>(Owner.GetWorld());
+
+	SubscribeToSignal(UE::Mass::Signals::CurrentLaneChanged);
+}
+
+void UMassZoneGraphAnnotationTagUpdateProcessor::ConfigureQueries()
+{
+	Super::ConfigureQueries();
+	EntityQuery.AddRequirement<FMassZoneGraphAnnotationTagsFragment>(ELWComponentAccess::ReadWrite);
+	EntityQuery.AddRequirement<FMassZoneGraphLaneLocationFragment>(ELWComponentAccess::ReadOnly);
+	EntityQuery.AddChunkRequirement<FMassZoneGraphAnnotationVariableTickChunkFragment>(ELWComponentAccess::ReadOnly);
+	EntityQuery.AddChunkRequirement<FMassSimulationVariableTickChunkFragment>(ELWComponentAccess::ReadOnly);
+}
+
+void UMassZoneGraphAnnotationTagUpdateProcessor::Execute(UEntitySubsystem& EntitySubsystem, FLWComponentSystemExecutionContext& Context)
+{
+	TransientEntitiesToSignal.Reset();
+
+	// Calling super will update the signals, and call SignalEntities() below.
+	Super::Execute(EntitySubsystem, Context);
+
+	EntityQuery.ForEachEntityChunk(EntitySubsystem, Context, [this](FLWComponentSystemExecutionContext& Context)
+	{
+		// Periodically update tags.
+		if (!FMassZoneGraphAnnotationVariableTickChunkFragment::UpdateChunk(Context))
+		{
+			return;
+		}
+
+		const int32 NumEntities = Context.GetEntitiesNum();
+		const TArrayView<FMassZoneGraphAnnotationTagsFragment> AnnotationTagsList = Context.GetMutableComponentView<FMassZoneGraphAnnotationTagsFragment>();
+		const TConstArrayView<FMassZoneGraphLaneLocationFragment> LaneLocationList = Context.GetComponentView<FMassZoneGraphLaneLocationFragment>();
+
+		for (int32 EntityIndex = 0; EntityIndex < NumEntities; ++EntityIndex)
+		{
+			FMassZoneGraphAnnotationTagsFragment& AnnotationTags = AnnotationTagsList[EntityIndex];
+			const FMassZoneGraphLaneLocationFragment& LaneLocation = LaneLocationList[EntityIndex];
+
+			UpdateAnnotationTags(AnnotationTags, LaneLocation, Context.GetEntity(EntityIndex));
+		}
+	});
+
+	if (TransientEntitiesToSignal.Num())
+	{
+		SignalSubsystem->SignalEntities(UE::Mass::Signals::AnnotationTagsChanged, TransientEntitiesToSignal);
+	}
+}
+
+void UMassZoneGraphAnnotationTagUpdateProcessor::UpdateAnnotationTags(FMassZoneGraphAnnotationTagsFragment& AnnotationTags, const FMassZoneGraphLaneLocationFragment& LaneLocation, const FLWEntity Entity)
+{
+	const FZoneGraphTagMask OldTags = AnnotationTags.Tags;
+
+	if (LaneLocation.LaneHandle.IsValid())
+	{
+		AnnotationTags.Tags = ZoneGraphAnnotationSubsystem->GetAnnotationTags(LaneLocation.LaneHandle);
+	}
+	else
+	{
+		AnnotationTags.Tags = FZoneGraphTagMask::None;
+	}
+
+	if (OldTags != AnnotationTags.Tags)
+	{
+		TransientEntitiesToSignal.Add(Entity);
+	}
+}
+
+void UMassZoneGraphAnnotationTagUpdateProcessor::SignalEntities(UEntitySubsystem& EntitySubsystem, FLWComponentSystemExecutionContext& Context, FMassSignalNameLookup& EntitySignals)
+{
+	if (!ZoneGraphAnnotationSubsystem)
+	{
+		return;
+	}
+
+	EntityQuery.ForEachEntityChunk(EntitySubsystem, Context, [this](FLWComponentSystemExecutionContext& Context)
+	{
+		const int32 NumEntities = Context.GetEntitiesNum();
+		const TConstArrayView<FMassZoneGraphLaneLocationFragment> LaneLocationList = Context.GetComponentView<FMassZoneGraphLaneLocationFragment>();
+		const TArrayView<FMassZoneGraphAnnotationTagsFragment> AnnotationTagsList = Context.GetMutableComponentView<FMassZoneGraphAnnotationTagsFragment>();
+
+		for (int32 EntityIndex = 0; EntityIndex < NumEntities; ++EntityIndex)
+		{
+			FMassZoneGraphAnnotationTagsFragment& AnnotationTags = AnnotationTagsList[EntityIndex];
+			const FMassZoneGraphLaneLocationFragment& LaneLocation = LaneLocationList[EntityIndex];
+
+			UpdateAnnotationTags(AnnotationTags, LaneLocation, Context.GetEntity(EntityIndex));
+		}
+	});
+}
diff --git a/Engine/Plugins/AI/MassAI/Source/MassAIBehavior/Private/MassZoneGraphAnnotationTrait.cpp b/Engine/Plugins/AI/MassAI/Source/MassAIBehavior/Private/MassZoneGraphAnnotationTrait.cpp
new file mode 100644
--- /dev/null
+++ b/Engine/Plugins/AI/MassAI/Source/MassAIBehavior/Private/MassZoneGraphAnnotationTrait.cpp
@@ -0,0 +1,12 @@
+// Copyright Epic Games, Inc. All Rights Reserved.
+
+#include "MassZoneGraphAnnotationTrait.h"
+#include "MassEntityTemplateRegistry.h"
+#include "Engine/World.h"
+#include "MassZoneGraphAnnotationFragments.h"
+
+void UMassZoneGraphAnnotationTrait::BuildTemplate(FMassEntityTemplateBuildContext& BuildContext, UWorld& World) const
+{
+	BuildContext.AddFragmentWithDefaultInitializer<FMassZoneGraphAnnotationTagsFragment>();
+	BuildContext.AddChunkFragment<FMassZoneGraphAnnotationVariableTickChunkFragment>();
+}
diff --git a/Engine/Plugins/AI/MassAI/Source/MassAIBehavior/Private/Tasks/MassLookAtTask.cpp b/Engine/Plugins/AI/MassAI/Source/MassAIBehavior/Private/Tasks/MassLookAtTask.cpp
new file mode 100644
--- /dev/null
+++ b/Engine/Plugins/AI/MassAI/Source/MassAIBehavior/Private/Tasks/MassLookAtTask.cpp
@@ -0,0 +1,61 @@
+// Copyright Epic Games, Inc. All Rights Reserved.
+
+#include "Tasks/MassLookAtTask.h"
+#include "MassAIBehaviorTypes.h"
+#include "MassLookAtFragments.h"
+#include "MassSignalSubsystem.h"
+#include "MassStateTreeExecutionContext.h"
+
+EStateTreeRunStatus FMassLookAtTask::EnterState(FStateTreeExecutionContext& Context, const EStateTreeStateChangeType ChangeType, const FStateTreeTransitionResult& Transition)
+{
+	Time = 0.f;
+	
+	const FMassStateTreeExecutionContext& MassContext = static_cast<FMassStateTreeExecutionContext&>(Context);
+	FMassLookAtFragment& LookAtFragment = MassContext.GetExternalItem(LookAtHandle).GetMutable<FMassLookAtFragment>();
+
+	LookAtFragment.Reset();
+	LookAtFragment.LookAtMode = LookAtMode;
+	
+	if (LookAtMode == EMassLookAtMode::LookAtEntity)
+	{
+		if (!TargetEntity.IsSet())
+		{
+			LookAtFragment.LookAtMode = EMassLookAtMode::LookForward;
+			MASSBEHAVIOR_LOG(Error, TEXT("Failed LookAt: invalid target entity"));
+		}
+		else
+		{
+			LookAtFragment.LookAtMode = EMassLookAtMode::LookAtEntity;
+			LookAtFragment.TrackedEntity = TargetEntity;
+		}
+	}
+
+	LookAtFragment.RandomGazeMode = RandomGazeMode;
+	LookAtFragment.RandomGazeYawVariation = RandomGazeYawVariation;
+	LookAtFragment.RandomGazePitchVariation = RandomGazePitchVariation;
+	LookAtFragment.bRandomGazeEntities = bRandomGazeEntities;
+
+	// A Duration <= 0 indicates that the task runs until a transition in the state tree stops it.
+	// Otherwise we schedule a signal to end the task.
+	if (Duration > 0.0f)
+	{
+		UMassSignalSubsystem& MassSignalSubsystem = MassContext.GetExternalItem(MassSignalSubsystemHandle).GetMutable<UMassSignalSubsystem>();
+		MassSignalSubsystem.DelaySignalEntity(UE::Mass::Signals::LookAtFinished, MassContext.GetEntity(), Duration);
+	}
+
+	return EStateTreeRunStatus::Running;
+}
+
+void FMassLookAtTask::ExitState(FStateTreeExecutionContext& Context, const EStateTreeStateChangeType ChangeType, const FStateTreeTransitionResult& Transition)
+{
+	const FMassStateTreeExecutionContext& MassContext = static_cast<FMassStateTreeExecutionContext&>(Context);
+	FMassLookAtFragment& LookAtFragment = MassContext.GetExternalItem(LookAtHandle).GetMutable<FMassLookAtFragment>();
+	
+	LookAtFragment.Reset();
+}
+
+EStateTreeRunStatus FMassLookAtTask::Tick(FStateTreeExecutionContext& Context, const float DeltaTime)
+{
+	Time += DeltaTime;
+	return Duration <= 0.0f ? EStateTreeRunStatus::Running : (Time < Duration ? EStateTreeRunStatus::Running : EStateTreeRunStatus::Succeeded);
+}
diff --git a/Engine/Plugins/AI/MassAI/Source/MassAIBehavior/Private/Tasks/MassStateTreeSmartObjectTask.cpp b/Engine/Plugins/AI/MassAI/Source/MassAIBehavior/Private/Tasks/MassStateTreeSmartObjectTask.cpp
new file mode 100644
--- /dev/null
+++ b/Engine/Plugins/AI/MassAI/Source/MassAIBehavior/Private/Tasks/MassStateTreeSmartObjectTask.cpp
@@ -0,0 +1,213 @@
+// Copyright Epic Games, Inc. All Rights Reserved.
+
+#include "Tasks/MassStateTreeSmartObjectTask.h"
+#include "MassCommonFragments.h"
+#include "MassAIBehaviorTypes.h"
+#include "MassSmartObjectHandler.h"
+#include "MassSmartObjectProcessor.h"
+#include "MassStateTreeExecutionContext.h"
+#include "MassAIMovementFragments.h"
+#include "SmartObjectSubsystem.h"
+#include "MassZoneGraphMovementUtils.h"
+#include "Engine/World.h"
+
+//----------------------------------------------------------------------//
+// FMassClaimSmartObjectTask
+//----------------------------------------------------------------------//
+EStateTreeRunStatus FMassClaimSmartObjectTask::EnterState(FStateTreeExecutionContext& Context, const EStateTreeStateChangeType ChangeType, const FStateTreeTransitionResult& Transition)
+{
+	if (ChangeType != EStateTreeStateChangeType::Changed)
+	{
+		return EStateTreeRunStatus::Running;
+	}
+
+	if (!SearchRequestResult.bProcessed)
+	{
+		MASSBEHAVIOR_LOG(Error, TEXT("Failed to claim smart object: unprocessed search request results."));
+		return EStateTreeRunStatus::Failed;
+	}
+
+	if (ClaimResult != EMassSmartObjectClaimResult::Unset)
+	{
+		MASSBEHAVIOR_LOG(Error, TEXT("Claim result should be 'Unset' but is '%s'"), *UEnum::GetValueAsString(ClaimResult));
+		return EStateTreeRunStatus::Failed;
+	}
+
+	EStateTreeRunStatus Status = EStateTreeRunStatus::Failed;
+
+	// Retrieve fragments and subsystems
+	USmartObjectSubsystem& SmartObjectSubsystem = Context.GetExternalItem(SmartObjectSubsystemHandle).GetMutable<USmartObjectSubsystem>();
+	FDataFragment_SmartObjectUser& SOUser = Context.GetExternalItem(SmartObjectUserHandle).GetMutable<FDataFragment_SmartObjectUser>();
+
+	// Setup MassSmartObject handler and claim
+	const FMassStateTreeExecutionContext& MassContext = static_cast<FMassStateTreeExecutionContext&>(Context);
+	const FMassSmartObjectHandler MassSmartObjectHandler(MassContext.GetEntitySubsystem(), MassContext.GetEntitySubsystemExecutionContext(), SmartObjectSubsystem);
+	ClaimResult = MassSmartObjectHandler.ClaimCandidate(MassContext.GetEntity(), SOUser, SearchRequestResult);
+
+	switch (ClaimResult)
+	{
+	case EMassSmartObjectClaimResult::Succeeded:
+		Status = EStateTreeRunStatus::Running;
+		break;
+	case EMassSmartObjectClaimResult::Failed_InvalidRequest:
+		MASSBEHAVIOR_LOG(Error, TEXT("Failed to claim smart object: %s"), *UEnum::GetValueAsString(ClaimResult));
+		break;
+	case EMassSmartObjectClaimResult::Failed_NoAvailableCandidate:
+		MASSBEHAVIOR_LOG(Log, TEXT("Failed to claim smart object: %s"), *UEnum::GetValueAsString(ClaimResult));
+		break;
+	case EMassSmartObjectClaimResult::Failed_UnprocessedRequest:
+		MASSBEHAVIOR_LOG(VeryVerbose, TEXT("Failed to claim smart object: %s"), *UEnum::GetValueAsString(ClaimResult));
+		break;
+	default:
+		ensureMsgf(false, TEXT("Unsupported EMassSmartObjectClaimResult value: %d"), ClaimResult);
+	}
+
+	return Status;
+}
+
+void FMassClaimSmartObjectTask::ExitState(FStateTreeExecutionContext& Context, const EStateTreeStateChangeType ChangeType, const FStateTreeTransitionResult& Transition)
+{
+	if (ChangeType != EStateTreeStateChangeType::Changed)
+	{
+		return;
+	}
+
+	ClaimResult = EMassSmartObjectClaimResult::Unset;
+
+	FDataFragment_SmartObjectUser& SOUser = Context.GetExternalItem(SmartObjectUserHandle).GetMutable<FDataFragment_SmartObjectUser>();
+
+	if (SOUser.ClaimHandle.IsValid())
+	{
+		if (SOUser.InteractionStatus == EMassSmartObjectInteractionStatus::Unset)
+		{
+			MASSBEHAVIOR_LOG(VeryVerbose, TEXT("Exiting state with a valid claim handle but interaction was not started: release claim on the smart object."));
+			const FMassStateTreeExecutionContext& MassContext = static_cast<FMassStateTreeExecutionContext&>(Context);
+			USmartObjectSubsystem& SmartObjectSubsystem = Context.GetExternalItem(SmartObjectSubsystemHandle).GetMutable<USmartObjectSubsystem>();
+			const FMassSmartObjectHandler MassSmartObjectHandler(MassContext.GetEntitySubsystem(), MassContext.GetEntitySubsystemExecutionContext(), SmartObjectSubsystem);
+			MassSmartObjectHandler.ReleaseSmartObject(MassContext.GetEntity(), SOUser);
+		}
+		else
+		{
+			MASSBEHAVIOR_LOG(VeryVerbose, TEXT("Exiting state with a valid claim handle and interaction was started: nothing to do since use task will stop the interaction."));
+		}
+	}
+	else
+	{
+		MASSBEHAVIOR_LOG(VeryVerbose, TEXT("Exiting state with an invalid ClaimHandle: nothing to do."));
+	}
+}
+
+EStateTreeRunStatus FMassClaimSmartObjectTask::Tick(FStateTreeExecutionContext& Context, const float DeltaTime)
+{
+	return EStateTreeRunStatus::Running;
+}
+
+//----------------------------------------------------------------------//
+// FMassUseSmartObjectTask
+//----------------------------------------------------------------------//
+EStateTreeRunStatus FMassUseSmartObjectTask::EnterState(FStateTreeExecutionContext& Context, const EStateTreeStateChangeType ChangeType, const FStateTreeTransitionResult& Transition)
+{
+	if (ChangeType != EStateTreeStateChangeType::Changed)
+	{
+		return EStateTreeRunStatus::Running;
+	}
+
+	// Retrieve fragments and subsystems
+	USmartObjectSubsystem& SmartObjectSubsystem = Context.GetExternalItem(SmartObjectSubsystemHandle).GetMutable<USmartObjectSubsystem>();
+	FDataFragment_SmartObjectUser& SOUser = Context.GetExternalItem(SmartObjectUserHandle).GetMutable<FDataFragment_SmartObjectUser>();
+	const FDataFragment_Transform& TransformFragment = Context.GetExternalItem(EntityTransformHandle).Get<FDataFragment_Transform>();
+	FMassMoveTargetFragment& MoveTarget = Context.GetExternalItem(MoveTargetHandle).GetMutable<FMassMoveTargetFragment>();
+
+	// Setup MassSmartObject handler and start interaction
+	const FMassStateTreeExecutionContext& MassContext = static_cast<FMassStateTreeExecutionContext&>(Context);
+	const FMassSmartObjectHandler MassSmartObjectHandler(MassContext.GetEntitySubsystem(), MassContext.GetEntitySubsystemExecutionContext(), SmartObjectSubsystem);
+
+	if (!MassSmartObjectHandler.UseSmartObject(MassContext.GetEntity(), SOUser, TransformFragment))
+	{
+		return EStateTreeRunStatus::Failed;
+	}
+
+	// @todo: we should have common API to control this, currently handle via tasks.
+	const UWorld* World = Context.GetWorld();
+	checkf(World != nullptr, TEXT("A valid world is expected from the execution context"));
+
+	MoveTarget.CreateNewAction(EMassMovementAction::Animate, *World);
+	const bool bSuccess = UE::MassMovement::ActivateActionAnimate(*World, Context.GetOwner(), MassContext.GetEntity(), MoveTarget);
+
+	return bSuccess ? EStateTreeRunStatus::Running : EStateTreeRunStatus::Failed;
+}
+
+void FMassUseSmartObjectTask::ExitState(FStateTreeExecutionContext& Context, const EStateTreeStateChangeType ChangeType, const FStateTreeTransitionResult& Transition)
+{
+	if (ChangeType == EStateTreeStateChangeType::Changed)
+	{
+		FDataFragment_SmartObjectUser& SOUser = Context.GetExternalItem(SmartObjectUserHandle).GetMutable<FDataFragment_SmartObjectUser>();
+
+		if (SOUser.ClaimHandle.IsValid())
+		{
+			MASSBEHAVIOR_LOG(VeryVerbose, TEXT("Exiting state with a valid ClaimHandle: stop using the smart object."));
+
+			const FMassStateTreeExecutionContext& MassContext = static_cast<FMassStateTreeExecutionContext&>(Context);
+			USmartObjectSubsystem& SmartObjectSubsystem = Context.GetExternalItem(SmartObjectSubsystemHandle).GetMutable<USmartObjectSubsystem>();
+			const FMassSmartObjectHandler MassSmartObjectHandler(MassContext.GetEntitySubsystem(), MassContext.GetEntitySubsystemExecutionContext(), SmartObjectSubsystem);
+			MassSmartObjectHandler.ReleaseSmartObject(MassContext.GetEntity(), SOUser, EMassSmartObjectInteractionStatus::Aborted);
+		}
+		else
+		{
+			MASSBEHAVIOR_LOG(VeryVerbose, TEXT("Exiting state with an invalid ClaimHandle: nothing to do."));
+		}
+	}
+}
+
+void FMassUseSmartObjectTask::StateCompleted(FStateTreeExecutionContext& Context, const EStateTreeRunStatus CompletionStatus, const FStateTreeHandle CompletedState)
+{
+	FDataFragment_SmartObjectUser& SOUser = Context.GetExternalItem(SmartObjectUserHandle).GetMutable<FDataFragment_SmartObjectUser>();
+
+	if (SOUser.ClaimHandle.IsValid())
+	{
+		MASSBEHAVIOR_LOG(VeryVerbose, TEXT("Completing state with a valid ClaimHandle: stop using the smart object."));
+
+		const FMassStateTreeExecutionContext& MassContext = static_cast<FMassStateTreeExecutionContext&>(Context);
+		USmartObjectSubsystem& SmartObjectSubsystem = Context.GetExternalItem(SmartObjectSubsystemHandle).GetMutable<USmartObjectSubsystem>();
+		const FMassSmartObjectHandler MassSmartObjectHandler(MassContext.GetEntitySubsystem(), MassContext.GetEntitySubsystemExecutionContext(), SmartObjectSubsystem);
+		const EMassSmartObjectInteractionStatus NewStatus = CompletionStatus == EStateTreeRunStatus::Succeeded ? EMassSmartObjectInteractionStatus::Completed : EMassSmartObjectInteractionStatus::Aborted;
+		MassSmartObjectHandler.ReleaseSmartObject(MassContext.GetEntity(), SOUser, NewStatus);
+	}
+}
+
+EStateTreeRunStatus FMassUseSmartObjectTask::Tick(FStateTreeExecutionContext& Context, const float DeltaTime)
+{
+	EStateTreeRunStatus Status = EStateTreeRunStatus::Failed;
+
+	FDataFragment_SmartObjectUser& SOUser = Context.GetExternalItem(SmartObjectUserHandle).GetMutable<FDataFragment_SmartObjectUser>();
+	switch (SOUser.InteractionStatus)
+	{
+	case EMassSmartObjectInteractionStatus::InProgress:
+		MASSBEHAVIOR_LOG(VeryVerbose, TEXT("Interacting ..."));
+		Status = EStateTreeRunStatus::Running;
+		break;
+
+	case EMassSmartObjectInteractionStatus::Completed:
+		MASSBEHAVIOR_LOG(Log, TEXT("Interaction completed"));
+		SOUser.SetCooldown(Context.GetWorld()->GetTimeSeconds() + Cooldown);
+		Status = EStateTreeRunStatus::Succeeded;
+		break;
+
+	case EMassSmartObjectInteractionStatus::Aborted:
+		MASSBEHAVIOR_LOG(Log, TEXT("Interaction aborted"));
+		Status = EStateTreeRunStatus::Failed;
+		break;
+
+	case EMassSmartObjectInteractionStatus::Unset:
+		MASSBEHAVIOR_LOG(Error, TEXT("Error while using  smart object: interaction state is not valid"));
+		Status = EStateTreeRunStatus::Failed;
+		break;
+
+	default:
+		ensureMsgf(false, TEXT("Unhandled interaction status % => Returning EStateTreeRunStatus::Failed"),
+			*StaticEnum<EMassSmartObjectInteractionStatus>()->GetValueAsString(SOUser.InteractionStatus));
+		Status = EStateTreeRunStatus::Failed;
+	}
+
+	return Status;
+}
diff --git a/Engine/Plugins/AI/MassAI/Source/MassAIBehavior/Private/Tasks/MassStateTreeTestTask.cpp b/Engine/Plugins/AI/MassAI/Source/MassAIBehavior/Private/Tasks/MassStateTreeTestTask.cpp
new file mode 100644
--- /dev/null
+++ b/Engine/Plugins/AI/MassAI/Source/MassAIBehavior/Private/Tasks/MassStateTreeTestTask.cpp
@@ -0,0 +1,36 @@
+// Copyright Epic Games, Inc. All Rights Reserved.
+
+#include "Tasks/MassStateTreeTestTask.h"
+#include "Engine/Engine.h"
+#include "StateTreeExecutionContext.h"
+#include "MassStateTreeSubsystem.h"
+#include "MassCommonFragments.h"
+#include "MassSmartObjectProcessor.h"
+
+
+FMassStateTreeTestTask::FMassStateTreeTestTask()
+	: Time(0.0f)
+{
+}
+
+FMassStateTreeTestTask::~FMassStateTreeTestTask()
+{
+}
+
+EStateTreeRunStatus FMassStateTreeTestTask::EnterState(FStateTreeExecutionContext& Context, const EStateTreeStateChangeType ChangeType, const FStateTreeTransitionResult& Transition)
+{
+	Time = 0.0f;
+	return EStateTreeRunStatus::Running;
+}
+
+EStateTreeRunStatus FMassStateTreeTestTask::Tick(FStateTreeExecutionContext& Context, const float DeltaTime)
+{
+	UMassStateTreeSubsystem& MassStateTreeSubSystem = Context.GetExternalItem(MassStateTreeSubSystemHandle).GetMutable<UMassStateTreeSubsystem>();
+	const FDataFragment_SmartObjectUser* SmartObjectUser = Context.GetExternalItem(SmartObjectUserHandle).GetPtr<FDataFragment_SmartObjectUser>();
+	FDataFragment_Transform& Transform = Context.GetExternalItem(TransformHandle).GetMutable<FDataFragment_Transform>();
+
+	GEngine->AddOnScreenDebugMessage(INDEX_NONE, 1.0f, FColor::Orange, *FString::Printf(TEXT("[%s] Time=%f X=%f\n"), *Name.ToString(), Time, Transform.GetTransform().GetLocation().X));
+
+	Time += DeltaTime;
+	return Time < Duration ? EStateTreeRunStatus::Running : EStateTreeRunStatus::Succeeded;
+}
diff --git a/Engine/Plugins/AI/MassAI/Source/MassAIBehavior/Private/Tasks/MassZoneGraphFindEscapeTarget.cpp b/Engine/Plugins/AI/MassAI/Source/MassAIBehavior/Private/Tasks/MassZoneGraphFindEscapeTarget.cpp
new file mode 100644
--- /dev/null
+++ b/Engine/Plugins/AI/MassAI/Source/MassAIBehavior/Private/Tasks/MassZoneGraphFindEscapeTarget.cpp
@@ -0,0 +1,142 @@
+﻿// Copyright Epic Games, Inc. All Rights Reserved.
+
+#include "Tasks/MassZoneGraphFindEscapeTarget.h"
+#include "StateTreeExecutionContext.h"
+#include "MassStateTreeSubsystem.h"
+#include "ZoneGraphSubsystem.h"
+#include "ZoneGraphAnnotationSubsystem.h"
+#include "MassZoneGraphMovementFragments.h"
+#include "Annotations/ZoneGraphDisturbanceAnnotation.h"
+#include "MassAIBehaviorTypes.h"
+#include "MassStateTreeExecutionContext.h"
+#include "ZoneGraphSettings.h"
+
+FMassZoneGraphFindEscapeTarget::FMassZoneGraphFindEscapeTarget()
+{
+}
+
+EStateTreeRunStatus FMassZoneGraphFindEscapeTarget::EnterState(FStateTreeExecutionContext& Context, const EStateTreeStateChangeType ChangeType, const FStateTreeTransitionResult& Transition)
+{
+	FMassStateTreeExecutionContext& MassContext = static_cast<FMassStateTreeExecutionContext&>(Context);
+	bool bDisplayDebug = false;
+
+	EscapeTargetLocation = nullptr;
+
+#if WITH_MASS_DEBUG
+	bDisplayDebug = UE::MassDebug::IsDebuggingEntity(MassContext.GetEntity());
+#endif // WITH_MASS_DEBUG
+
+	const FMassZoneGraphLaneLocationFragment& LaneLocation = Context.GetExternalItem(LocationHandle).Get<FMassZoneGraphLaneLocationFragment>();
+	UZoneGraphSubsystem& ZoneGraphSubsystem = Context.GetExternalItem(ZoneGraphSubsystemHandle).GetMutable<UZoneGraphSubsystem>();
+	UZoneGraphAnnotationSubsystem& ZoneGraphAnnotationSubsystem = Context.GetExternalItem(ZoneGraphAnnotationSubsystemHandle).GetMutable<UZoneGraphAnnotationSubsystem>();
+
+	if (!LaneLocation.LaneHandle.IsValid())
+	{
+		MASSBEHAVIOR_LOG(Error, TEXT("Invalid lane handle."));
+		return EStateTreeRunStatus::Failed;
+	}
+			
+	const FZoneGraphStorage* ZoneGraphStorage = ZoneGraphSubsystem.GetZoneGraphStorage(LaneLocation.LaneHandle.DataHandle);
+	if (!ZoneGraphStorage)
+	{
+		MASSBEHAVIOR_LOG(Error, TEXT("Missing ZoneGraph Storage for current lane %s."), *LaneLocation.LaneHandle.ToString());
+		return EStateTreeRunStatus::Failed;
+	}
+
+	const UZoneGraphDisturbanceAnnotation* DisturbanceAnnotation = Cast<const UZoneGraphDisturbanceAnnotation>(ZoneGraphAnnotationSubsystem.GetFirstAnnotationForTag(DisturbanceAnnotationTag));
+	if (!DisturbanceAnnotation)
+	{
+		MASSBEHAVIOR_LOG(Error, TEXT("Failed to find Flee Behavior for tag %s on lane %s."), *UE::ZoneGraph::Helpers::GetTagName(DisturbanceAnnotationTag).ToString(), *LaneLocation.LaneHandle.ToString());
+		return EStateTreeRunStatus::Failed;
+	}
+
+	const FZoneGraphEscapeLaneAction* EscapeAction = DisturbanceAnnotation->GetEscapeAction(LaneLocation.LaneHandle);
+	if (!EscapeAction)
+	{
+		MASSBEHAVIOR_LOG(Warning, TEXT("Failed to find escape action for current lane %s."), *LaneLocation.LaneHandle.ToString());
+		return EStateTreeRunStatus::Failed;
+	}
+
+	const uint8 SpanIndex = EscapeAction->FindSpanIndex(LaneLocation.DistanceAlongLane);
+	const FZoneGraphEscapeLaneSpan& EscapeSpan = EscapeAction->Spans[SpanIndex];
+	const float MoveDir = EscapeSpan.bReverseLaneDirection ? -1 : 1;	
+
+	constexpr float AdjacentMoveDistance = 50.0f;
+	constexpr float MoveDistanceRandomDeviation = 250.0f;
+	constexpr float BaseMoveDistance = 800.0f;
+	const float MoveDistance = FMath::Max(0.0f, BaseMoveDistance + FMath::RandRange(-MoveDistanceRandomDeviation, MoveDistanceRandomDeviation));
+
+	if (EscapeSpan.ExitLaneIndex == INDEX_NONE)
+	{
+		MASSBEHAVIOR_LOG(Warning, TEXT("Invalid flee exit lane."));
+		return EStateTreeRunStatus::Failed;
+	}
+	
+	if (EscapeSpan.ExitLinkType == EZoneLaneLinkType::Adjacent)
+	{
+		// TODO: could improve this by checking the adjacent lane type and move based on that:
+		// - split: move closer to beginning of the lane
+		// - merge: move closer to end of the lane
+		// - adjacent: ?
+		// Could also try to sample few locations along the lane to see which is closest.
+		
+		// Small move, and goto adjacent lane
+		TargetLocation.LaneHandle = LaneLocation.LaneHandle;
+		TargetLocation.TargetDistance = FMath::Clamp(LaneLocation.DistanceAlongLane + AdjacentMoveDistance * MoveDir, 0, LaneLocation.LaneLength);
+		TargetLocation.NextExitLinkType = EZoneLaneLinkType::Adjacent;
+		TargetLocation.NextLaneHandle = FZoneGraphLaneHandle(EscapeSpan.ExitLaneIndex, ZoneGraphStorage->DataHandle);
+		TargetLocation.bMoveReverse = EscapeSpan.bReverseLaneDirection;
+		TargetLocation.EndOfPathIntent = EMassMovementAction::Move;
+
+		MASSBEHAVIOR_CLOG(bDisplayDebug, Log, TEXT("Switching from lane %s to adjacent lane %s."),
+			*LaneLocation.LaneHandle.ToString(), *TargetLocation.NextLaneHandle.ToString());
+	}
+	else
+	{
+		// Forward or backwards on current lane.
+		TargetLocation.LaneHandle = LaneLocation.LaneHandle;
+		TargetLocation.TargetDistance = LaneLocation.DistanceAlongLane + MoveDistance * MoveDir;
+		TargetLocation.NextExitLinkType = EZoneLaneLinkType::None;
+		TargetLocation.NextLaneHandle.Reset();
+		TargetLocation.bMoveReverse = EscapeSpan.bReverseLaneDirection;
+		TargetLocation.EndOfPathIntent = EMassMovementAction::Move;
+
+		// When close to end of a lane, choose next lane.
+		const bool bPastStart = TargetLocation.TargetDistance < 0.0f;
+		const bool bPastEnd = TargetLocation.TargetDistance > LaneLocation.LaneLength;
+		if (bPastStart || bPastEnd)
+		{
+			TargetLocation.TargetDistance = FMath::Clamp(TargetLocation.TargetDistance, 0.0f, LaneLocation.LaneLength);
+
+			TargetLocation.NextExitLinkType = EscapeSpan.ExitLinkType;
+			TargetLocation.NextLaneHandle = FZoneGraphLaneHandle(EscapeSpan.ExitLaneIndex, ZoneGraphStorage->DataHandle);
+
+			MASSBEHAVIOR_CLOG(bDisplayDebug, Log, TEXT("Advancing %s along flee lane %s to next lane %s at distance %.1f."),
+				TargetLocation.bMoveReverse ? TEXT("forward") : TEXT("reverse"),
+				*TargetLocation.LaneHandle.ToString(), *TargetLocation.NextLaneHandle.ToString(),
+				TargetLocation.TargetDistance);
+		}
+		else
+		{
+			MASSBEHAVIOR_CLOG(bDisplayDebug, Log, TEXT("Advancing %s along flee lane %s to distance %.1f."),
+				TargetLocation.bMoveReverse ? TEXT("forward") : TEXT("reverse"),
+				*TargetLocation.LaneHandle.ToString(), TargetLocation.TargetDistance);
+		}
+	}
+
+	EscapeTargetLocation = &TargetLocation;
+
+	return EStateTreeRunStatus::Running;
+}
+
+void FMassZoneGraphFindEscapeTarget::ExitState(FStateTreeExecutionContext& Context, const EStateTreeStateChangeType ChangeType, const FStateTreeTransitionResult& Transition)
+{
+	EscapeTargetLocation = &TargetLocation;
+}
+
+EStateTreeRunStatus FMassZoneGraphFindEscapeTarget::Tick(FStateTreeExecutionContext& Context, const float DeltaTime)
+{
+	EscapeTargetLocation = &TargetLocation;
+	
+	return EStateTreeRunStatus::Running;
+}
