AggregateTicking and Pipe merged into Engine/Plugins/Runtime/MassEntity

diff --git a/Engine/Plugins/AI/MassAI/Source/MassAIBehavior/Private/Tasks/MassZoneGraphFindSmartObjectTarget.cpp b/Engine/Plugins/AI/MassAI/Source/MassAIBehavior/Private/Tasks/MassZoneGraphFindSmartObjectTarget.cpp
new file mode 100644
--- /dev/null
+++ b/Engine/Plugins/AI/MassAI/Source/MassAIBehavior/Private/Tasks/MassZoneGraphFindSmartObjectTarget.cpp
@@ -0,0 +1,87 @@
+// Copyright Epic Games, Inc. All Rights Reserved.
+
+#include "Tasks/MassZoneGraphFindSmartObjectTarget.h"
+#include "ZoneGraphAnnotationSubsystem.h"
+#include "MassAIBehaviorTypes.h"
+#include "MassSmartObjectProcessor.h"
+#include "MassSmartObjectSettings.h"
+#include "MassStateTreeExecutionContext.h"
+#include "MassZoneGraphMovementFragments.h"
+#include "SmartObjectZoneAnnotations.h"
+
+EStateTreeRunStatus FMassZoneGraphFindSmartObjectTarget::EnterState(FStateTreeExecutionContext& Context, const EStateTreeStateChangeType ChangeType, const FStateTreeTransitionResult& Transition)
+{
+	const UZoneGraphAnnotationSubsystem& AnnotationSubsystem = Context.GetExternalItem(AnnotationSubsystemHandle).Get<UZoneGraphAnnotationSubsystem>();
+	const FDataFragment_SmartObjectUser& SOUser = Context.GetExternalItem(SmartObjectUserHandle).Get<FDataFragment_SmartObjectUser>();
+	const FMassZoneGraphLaneLocationFragment& LaneLocation = Context.GetExternalItem(LocationHandle).Get<FMassZoneGraphLaneLocationFragment>();
+	const FZoneGraphLaneHandle LaneHandle(LaneLocation.LaneHandle);
+
+	TargetLocationRef = nullptr;
+
+	if (Context.GetEnterStateStatus() == EStateTreeRunStatus::Failed)
+	{
+		MASSBEHAVIOR_LOG(Log, TEXT("Transition to the state has been denied by other tasks. Nothing to do."));
+		return EStateTreeRunStatus::Failed;
+	}
+
+	if (!SOUser.ClaimHandle.SmartObjectID.IsValid())
+	{
+		MASSBEHAVIOR_LOG(Error, TEXT("Invalid claimed smart object ID."));
+		return EStateTreeRunStatus::Failed;
+	}
+
+	if (!LaneHandle.IsValid())
+	{
+		MASSBEHAVIOR_LOG(Error, TEXT("Invalid lane location."));
+		return EStateTreeRunStatus::Failed;
+	}
+
+	const FZoneGraphTag SmartObjectTag = GetDefault<UMassSmartObjectSettings>()->SmartObjectTag;
+	const USmartObjectZoneAnnotations* SOAnnotations = Cast<USmartObjectZoneAnnotations>(AnnotationSubsystem.GetFirstAnnotationForTag(SmartObjectTag));
+
+	TargetLocation.LaneHandle = LaneHandle;
+	TargetLocation.NextExitLinkType = EZoneLaneLinkType::None;
+	TargetLocation.NextLaneHandle.Reset();
+	TargetLocation.bMoveReverse = false;
+	TargetLocation.EndOfPathIntent = EMassMovementAction::Stand;
+	TargetLocation.EndOfPathPosition = SOUser.GetTargetLocation();
+	// Let's start moving toward the interaction a bit before the entry point.
+	constexpr float AnticipationDistance = 100.f;
+	TargetLocation.AnticipationDistance = AnticipationDistance;
+
+	if (SOAnnotations != nullptr)
+	{
+		// Find entry point on lanes for the claimed object
+		const FSmartObjectAnnotationData* AnnotationData = SOAnnotations->GetAnnotationData(LaneHandle.DataHandle);
+		checkf(AnnotationData, TEXT("FSmartObjectAnnotationData should have been created for each registered valid ZoneGraphData"));
+		const FSmartObjectLaneLocation EntryPoint = AnnotationData->ObjectToEntryPointLookup.FindChecked(SOUser.ClaimHandle.SmartObjectID);
+
+		// Request path along current lane to reach entry point on lane
+		MASSBEHAVIOR_LOG(Log, TEXT("Claim successful: create path along lane to reach interaction location."));
+		TargetLocation.TargetDistance = EntryPoint.DistanceAlongLane;
+	}
+	else
+	{
+		// Request path from current lane location directly to interaction location
+		MASSBEHAVIOR_LOG(Warning, TEXT("Claim successful: create path from current lane location directly to interaction location since SmartObject zone annotations weren't found."));
+		TargetLocation.TargetDistance = LaneLocation.DistanceAlongLane;
+	}
+
+	TargetLocationRef = &TargetLocation;
+	
+	return EStateTreeRunStatus::Running;
+}
+
+void FMassZoneGraphFindSmartObjectTarget::ExitState(FStateTreeExecutionContext& Context, const EStateTreeStateChangeType ChangeType, const FStateTreeTransitionResult& Transition)
+{
+	// Temp solution to make sure the target location is always up to date (will be replaced with automatic update).
+	TargetLocationRef = &TargetLocation;
+}
+
+EStateTreeRunStatus FMassZoneGraphFindSmartObjectTarget::Tick(FStateTreeExecutionContext& Context, const float DeltaTime)
+{
+	// Temp solution to make sure the target location is always up to date (will be replaced with automatic update).
+	TargetLocationRef = &TargetLocation;
+	
+	return EStateTreeRunStatus::Running;
+}
diff --git a/Engine/Plugins/AI/MassAI/Source/MassAIBehavior/Private/Tasks/MassZoneGraphPathFollowTask.cpp b/Engine/Plugins/AI/MassAI/Source/MassAIBehavior/Private/Tasks/MassZoneGraphPathFollowTask.cpp
new file mode 100644
--- /dev/null
+++ b/Engine/Plugins/AI/MassAI/Source/MassAIBehavior/Private/Tasks/MassZoneGraphPathFollowTask.cpp
@@ -0,0 +1,167 @@
+// Copyright Epic Games, Inc. All Rights Reserved.
+
+#include "Tasks/MassZoneGraphPathFollowTask.h"
+
+#include "MassAIBehaviorTypes.h"
+#include "MassCommonFragments.h"
+#include "MassAIMovementFragments.h"
+#include "MassMovementSettings.h"
+#include "MassStateTreeExecutionContext.h"
+#include "MassZoneGraphMovementFragments.h"
+#include "MassZoneGraphMovementUtils.h"
+#include "ZoneGraphSubsystem.h"
+
+bool FMassZoneGraphPathFollowTask::RequestPath(FMassStateTreeExecutionContext& Context, const FMassZoneGraphTargetLocation& RequestedTargetLocation) const
+{
+	const UZoneGraphSubsystem& ZoneGraphSubsystem = Context.GetExternalItem(ZoneGraphSubsystemHandle).Get<UZoneGraphSubsystem>();
+	const FMassZoneGraphLaneLocationFragment& LaneLocation = Context.GetExternalItem(LocationHandle).Get<FMassZoneGraphLaneLocationFragment>();
+	const FMassMovementConfigFragment& MovementConfig = Context.GetExternalItem(MovementConfigHandle).Get<FMassMovementConfigFragment>();
+	const FDataFragment_AgentRadius& AgentRadius = Context.GetExternalItem(AgentRadiusHandle).Get<FDataFragment_AgentRadius>();
+	FMassZoneGraphShortPathFragment& ShortPath = Context.GetExternalItem(ShortPathHandle).GetMutable<FMassZoneGraphShortPathFragment>();
+	FMassZoneGraphCachedLaneFragment& CachedLane = Context.GetExternalItem(CachedLaneHandle).GetMutable<FMassZoneGraphCachedLaneFragment>();
+	FMassMoveTargetFragment& MoveTarget = Context.GetExternalItem(MoveTargetHandle).GetMutable<FMassMoveTargetFragment>();
+	FMassZoneGraphPathRequestFragment& RequestFragment = Context.GetExternalItem(PathRequestHandle).GetMutable<FMassZoneGraphPathRequestFragment>();
+
+	const UMassMovementSettings* Settings = GetDefault<UMassMovementSettings>();
+	check(Settings);
+
+	const FMassMovementConfig* Config = Settings->GetMovementConfigByHandle(MovementConfig.ConfigHandle);
+	if (!Config)
+	{
+		MASSBEHAVIOR_LOG(Error, TEXT("Failed to get move config %s."), *LaneLocation.LaneHandle.ToString());
+		return false;
+	}
+
+	bool bDisplayDebug = false;
+#if WITH_MASS_DEBUG
+	bDisplayDebug = UE::MassDebug::IsDebuggingEntity(Context.GetEntity());
+#endif // WITH_MASS_DEBUG
+
+	if (RequestedTargetLocation.LaneHandle != LaneLocation.LaneHandle)
+	{
+		MASSBEHAVIOR_LOG(Error, TEXT("Target location lane %s does not match current lane location %s."),
+			*RequestedTargetLocation.LaneHandle.ToString(), *LaneLocation.LaneHandle.ToString());
+		return false;
+	}
+
+	if (bDisplayDebug)
+	{
+		MASSBEHAVIOR_LOG(Log, TEXT("PathFollow Request: %s, lane %s, Start: %f End:%f, next lane %s."),
+			RequestedTargetLocation.bMoveReverse ? TEXT("reverse") : TEXT("forward"),
+			*LaneLocation.LaneHandle.ToString(),
+			LaneLocation.DistanceAlongLane, RequestedTargetLocation.TargetDistance,
+			*RequestedTargetLocation.NextLaneHandle.ToString());
+	}
+
+	// @todo: Combine FMassZoneGraphTargetLocation and FZoneGraphShortPathRequest.
+	FZoneGraphShortPathRequest& PathRequest = RequestFragment.PathRequest;
+	PathRequest.StartPosition = MoveTarget.Center;
+	PathRequest.bMoveReverse = RequestedTargetLocation.bMoveReverse;
+	PathRequest.TargetDistance = RequestedTargetLocation.TargetDistance;
+	PathRequest.NextLaneHandle = RequestedTargetLocation.NextLaneHandle;
+	PathRequest.NextExitLinkType = RequestedTargetLocation.NextExitLinkType;
+	PathRequest.EndOfPathIntent = RequestedTargetLocation.EndOfPathIntent;
+	PathRequest.bIsEndOfPathPositionSet = RequestedTargetLocation.EndOfPathPosition.IsSet();
+	PathRequest.EndOfPathPosition = RequestedTargetLocation.EndOfPathPosition.Get(FVector::ZeroVector);
+	PathRequest.bIsEndOfPathDirectionSet = RequestedTargetLocation.EndOfPathDirection.IsSet();
+	PathRequest.EndOfPathDirection.Set(RequestedTargetLocation.EndOfPathDirection.Get(FVector::ForwardVector));
+	PathRequest.AnticipationDistance = RequestedTargetLocation.AnticipationDistance;
+	PathRequest.EndOfPathOffset.Set(FMath::RandRange(-AgentRadius.Radius, AgentRadius.Radius));
+
+	const float DesiredSpeed = FMath::Min(Config->GenerateDesiredSpeed(MovementStyle, Context.GetEntity().Index) * SpeedScale, Config->MaximumSpeed);
+	const UWorld* World = Context.GetWorld();
+	checkf(World != nullptr, TEXT("A valid world is expected from the execution context"));
+
+	MoveTarget.CreateNewAction(EMassMovementAction::Move, *World);
+	return UE::MassMovement::ActivateActionMove(*World, Context.GetOwner(), Context.GetEntity(), ZoneGraphSubsystem, LaneLocation, PathRequest, AgentRadius.Radius, DesiredSpeed, MoveTarget, ShortPath, CachedLane);
+}
+
+EStateTreeRunStatus FMassZoneGraphPathFollowTask::EnterState(FStateTreeExecutionContext& Context, const EStateTreeStateChangeType ChangeType, const FStateTreeTransitionResult& Transition)
+{
+	// Do not reset of the state if current state is still active after transition, unless transitioned specifically to this state.
+	if (ChangeType == EStateTreeStateChangeType::Sustained && Transition.Current != Transition.Next)
+	{
+		return EStateTreeRunStatus::Running;
+	}
+	
+	FMassStateTreeExecutionContext& MassContext = static_cast<FMassStateTreeExecutionContext&>(Context);
+
+	bool bDisplayDebug = false;
+#if WITH_MASS_DEBUG
+	bDisplayDebug = UE::MassDebug::IsDebuggingEntity(MassContext.GetEntity());
+#endif // WITH_MASS_DEBUG
+	if (bDisplayDebug)
+	{
+		MASSBEHAVIOR_LOG(Verbose, TEXT("enterstate."));
+	}
+
+	if (Context.GetEnterStateStatus() == EStateTreeRunStatus::Failed)
+	{
+		MASSBEHAVIOR_LOG(Log, TEXT("Transition to the state has been denied by other tasks. Nothing to do."));
+		return EStateTreeRunStatus::Failed;
+	}
+
+	const FMassZoneGraphLaneLocationFragment& LaneLocation = Context.GetExternalItem(LocationHandle).Get<FMassZoneGraphLaneLocationFragment>();
+
+	const FMassZoneGraphTargetLocation* TargetLocationPtr = TargetLocation.GetPtr<FMassZoneGraphTargetLocation>();
+	if (!TargetLocationPtr)
+	{
+		MASSBEHAVIOR_LOG(Error, TEXT("Target location not set."));
+		return EStateTreeRunStatus::Failed;
+	}
+	
+	if (TargetLocationPtr->LaneHandle != LaneLocation.LaneHandle)
+	{
+		MASSBEHAVIOR_LOG(Error, TEXT("Target is not on current lane, target lane is %s expected %s."), *TargetLocationPtr->LaneHandle.ToString(), *LaneLocation.LaneHandle.ToString());
+		return EStateTreeRunStatus::Failed;
+	}
+
+	if (!RequestPath(MassContext, *TargetLocationPtr))
+	{
+		return EStateTreeRunStatus::Failed;
+	}
+
+	return EStateTreeRunStatus::Running;
+}
+
+EStateTreeRunStatus FMassZoneGraphPathFollowTask::Tick(FStateTreeExecutionContext& Context, const float DeltaTime)
+{
+	FMassStateTreeExecutionContext& MassContext = static_cast<FMassStateTreeExecutionContext&>(Context);
+
+	bool bDisplayDebug = false;
+#if WITH_MASS_DEBUG
+	bDisplayDebug = UE::MassDebug::IsDebuggingEntity(MassContext.GetEntity());
+#endif // WITH_MASS_DEBUG
+	if (bDisplayDebug)
+	{
+		MASSBEHAVIOR_LOG(Verbose, TEXT("tick"));
+	}
+
+	const FMassZoneGraphShortPathFragment& ShortPath = Context.GetExternalItem(ShortPathHandle).Get<FMassZoneGraphShortPathFragment>();
+	const FMassZoneGraphLaneLocationFragment& LaneLocation = Context.GetExternalItem(LocationHandle).Get<FMassZoneGraphLaneLocationFragment>();
+
+	const FMassZoneGraphTargetLocation* TargetLocationPtr = TargetLocation.GetPtr<FMassZoneGraphTargetLocation>();
+	if (!TargetLocationPtr)
+	{
+		MASSBEHAVIOR_LOG(Error, TEXT("Target location not set."));
+		return EStateTreeRunStatus::Failed;
+	}
+
+	// Current path follow is done, but it was partial (i.e. many points on a curve), try again until we get there.
+	if (ShortPath.IsDone() && ShortPath.bPartialResult)
+	{
+		if (TargetLocationPtr->LaneHandle != LaneLocation.LaneHandle)
+		{
+			MASSBEHAVIOR_LOG(Error, TEXT("Target is not on current lane, target lane is %s expected %s."), *TargetLocationPtr->LaneHandle.ToString(), *LaneLocation.LaneHandle.ToString());
+			return EStateTreeRunStatus::Failed;
+		}
+		
+		if (!RequestPath(MassContext, *TargetLocationPtr))
+		{
+			MASSBEHAVIOR_LOG(Error, TEXT("Failed to request path."));
+			return EStateTreeRunStatus::Failed;
+		}
+	}
+	
+	return ShortPath.IsDone() ? EStateTreeRunStatus::Succeeded : EStateTreeRunStatus::Running;
+}
diff --git a/Engine/Plugins/AI/MassAI/Source/MassAIBehavior/Private/Tasks/MassZoneGraphStandTask.cpp b/Engine/Plugins/AI/MassAI/Source/MassAIBehavior/Private/Tasks/MassZoneGraphStandTask.cpp
new file mode 100644
--- /dev/null
+++ b/Engine/Plugins/AI/MassAI/Source/MassAIBehavior/Private/Tasks/MassZoneGraphStandTask.cpp
@@ -0,0 +1,78 @@
+// Copyright Epic Games, Inc. All Rights Reserved.
+
+#include "Tasks/MassZoneGraphStandTask.h"
+#include "StateTreeExecutionContext.h"
+#include "ZoneGraphSubsystem.h"
+#include "MassZoneGraphMovementFragments.h"
+#include "MassAIBehaviorTypes.h"
+#include "MassAIMovementFragments.h"
+#include "MassMovementSettings.h"
+#include "MassStateTreeExecutionContext.h"
+#include "MassSignalSubsystem.h"
+#include "MassSimulationLOD.h"
+#include "MassZoneGraphMovementUtils.h"
+
+EStateTreeRunStatus FMassZoneGraphStandTask::EnterState(FStateTreeExecutionContext& Context, const EStateTreeStateChangeType ChangeType, const FStateTreeTransitionResult& Transition)
+{
+	// Do not reset of the state if current state is still active after transition, unless transitioned specifically to this state.
+	if (ChangeType == EStateTreeStateChangeType::Sustained && Transition.Current != Transition.Next)
+	{
+		return EStateTreeRunStatus::Running;
+	}
+
+	const FMassStateTreeExecutionContext& MassContext = static_cast<FMassStateTreeExecutionContext&>(Context);
+
+	const FMassZoneGraphLaneLocationFragment& LaneLocation = Context.GetExternalItem(LocationHandle).Get<FMassZoneGraphLaneLocationFragment>();
+	const UZoneGraphSubsystem& ZoneGraphSubsystem = Context.GetExternalItem(ZoneGraphSubsystemHandle).Get<UZoneGraphSubsystem>();
+	const FMassMovementConfigFragment& MovementConfig = Context.GetExternalItem(MovementConfigHandle).Get<FMassMovementConfigFragment>();
+
+	const UMassMovementSettings* Settings = GetDefault<UMassMovementSettings>();
+	check(Settings);
+
+	const FMassMovementConfig* Config = Settings->GetMovementConfigByHandle(MovementConfig.ConfigHandle);
+	if (!Config)
+	{
+		MASSBEHAVIOR_LOG(Error, TEXT("Failed to get move config."));
+		return EStateTreeRunStatus::Failed;
+	}
+	
+	if (!LaneLocation.LaneHandle.IsValid())
+	{
+		MASSBEHAVIOR_LOG(Error, TEXT("Invalid lande handle"));
+		return EStateTreeRunStatus::Failed;
+	}
+
+	FMassZoneGraphShortPathFragment& ShortPath = Context.GetExternalItem(ShortPathHandle).GetMutable<FMassZoneGraphShortPathFragment>();
+	FMassZoneGraphCachedLaneFragment& CachedLane = Context.GetExternalItem(CachedLaneHandle).GetMutable<FMassZoneGraphCachedLaneFragment>();
+	FMassMoveTargetFragment& MoveTarget = Context.GetExternalItem(MoveTargetHandle).GetMutable<FMassMoveTargetFragment>();
+
+	// TODO: This could be smarter too, like having a stand location/direction, or even make a small path to stop, if we're currently running.
+
+	const UWorld* World = Context.GetWorld();
+	checkf(World != nullptr, TEXT("A valid world is expected from the execution context"));
+
+	MoveTarget.CreateNewAction(EMassMovementAction::Stand, *World);
+	const bool bSuccess = UE::MassMovement::ActivateActionStand(*World, Context.GetOwner(), MassContext.GetEntity(), ZoneGraphSubsystem, LaneLocation, Config->DefaultDesiredSpeed, MoveTarget, ShortPath, CachedLane);
+	if (!bSuccess)
+	{
+		return EStateTreeRunStatus::Failed;
+	}
+
+	Time = 0.0f;
+
+	// A Duration <= 0 indicates that the task runs until a transition in the state tree stops it.
+	// Otherwise we schedule a signal to end the task.
+	if (Duration > 0.0f)
+	{
+		UMassSignalSubsystem& MassSignalSubsystem = Context.GetExternalItem(MassSignalSubsystemHandle).GetMutable<UMassSignalSubsystem>();
+		MassSignalSubsystem.DelaySignalEntity(UE::Mass::Signals::StandTaskFinished, MassContext.GetEntity(), Duration);
+	}
+
+	return EStateTreeRunStatus::Running;
+}
+
+EStateTreeRunStatus FMassZoneGraphStandTask::Tick(FStateTreeExecutionContext& Context, const float DeltaTime)
+{
+	Time += DeltaTime;
+	return Duration <= 0.0f ? EStateTreeRunStatus::Running : (Time < Duration ? EStateTreeRunStatus::Running : EStateTreeRunStatus::Succeeded);
+}
diff --git a/Engine/Plugins/AI/MassAI/Source/MassAIBehavior/Public/Conditions/ZoneGraphTagConditions.h b/Engine/Plugins/AI/MassAI/Source/MassAIBehavior/Public/Conditions/ZoneGraphTagConditions.h
new file mode 100644
--- /dev/null
+++ b/Engine/Plugins/AI/MassAI/Source/MassAIBehavior/Public/Conditions/ZoneGraphTagConditions.h
@@ -0,0 +1,98 @@
+﻿// Copyright Epic Games, Inc. All Rights Reserved.
+
+#pragma once
+
+#include "CoreMinimal.h"
+#include "StateTreeConditionBase.h"
+#include "ZoneGraphTypes.h"
+#include "ZoneGraphTagConditions.generated.h"
+
+/**
+* ZoneGraph Tag condition.
+*/
+USTRUCT(DisplayName="ZoneGraphTagFilter Compare")
+struct MASSAIBEHAVIOR_API FZoneGraphTagFilterCondition : public FStateTreeConditionBase
+{
+	GENERATED_BODY()
+
+	FZoneGraphTagFilterCondition() : FStateTreeConditionBase() { }
+
+#if WITH_EDITOR
+	virtual FText GetDescription(const IStateTreeBindingLookup& BindingLookup) const override;
+#endif
+
+	virtual bool TestCondition() const override
+	{
+		return Filter.Pass(Tags) ^ bInvert;
+	}
+
+	UPROPERTY(EditAnywhere, Category = Condition, meta = (Bindable))
+	FZoneGraphTagMask Tags = FZoneGraphTagMask::None;
+
+	UPROPERTY(EditAnywhere, Category = Condition)
+	FZoneGraphTagFilter Filter;
+
+	UPROPERTY(EditAnywhere, Category = Condition)
+	bool bInvert = false;
+};
+
+/**
+* ZoneGraph Tag condition.
+*/
+USTRUCT(DisplayName="ZoneGraphTagMask Compare")
+struct MASSAIBEHAVIOR_API FZoneGraphTagMaskCondition : public FStateTreeConditionBase
+{
+	GENERATED_BODY()
+
+	FZoneGraphTagMaskCondition() : FStateTreeConditionBase() { }
+
+#if WITH_EDITOR
+	virtual FText GetDescription(const IStateTreeBindingLookup& BindingLookup) const override;
+#endif
+
+	virtual bool TestCondition() const override
+	{
+		return Left.CompareMasks(Right, Operator) ^ bInvert;
+	}
+
+	UPROPERTY(EditAnywhere, Category = Condition, meta = (Bindable))
+	FZoneGraphTagMask Left = FZoneGraphTagMask::None;
+
+	UPROPERTY(EditAnywhere, Category = Condition)
+	EZoneLaneTagMaskComparison Operator = EZoneLaneTagMaskComparison::Any;
+	
+	UPROPERTY(EditAnywhere, Category = Condition, meta = (Bindable))
+	FZoneGraphTagMask Right = FZoneGraphTagMask::None;
+
+	UPROPERTY(EditAnywhere, Category = Condition)
+	bool bInvert = false;
+};
+
+/**
+* ZoneGraph Tag condition.
+*/
+USTRUCT(DisplayName="ZoneGraphTag Compare")
+struct MASSAIBEHAVIOR_API FZoneGraphTagCondition : public FStateTreeConditionBase
+{
+	GENERATED_BODY()
+
+	FZoneGraphTagCondition() : FStateTreeConditionBase() { }
+
+#if WITH_EDITOR
+	virtual FText GetDescription(const IStateTreeBindingLookup& BindingLookup) const override;
+#endif
+
+	virtual bool TestCondition() const override
+	{
+		return (Left == Right) ^ bInvert;
+	}
+
+	UPROPERTY(EditAnywhere, Category = Condition, meta = (Bindable))
+	FZoneGraphTag Left = FZoneGraphTag::None;
+
+	UPROPERTY(EditAnywhere, Category = Condition, meta = (Bindable))
+	FZoneGraphTag Right = FZoneGraphTag::None;
+
+	UPROPERTY(EditAnywhere, Category = Condition)
+	bool bInvert = false;
+};
diff --git a/Engine/Plugins/AI/MassAI/Source/MassAIBehavior/Public/Evaluators/MassComponentHitEvaluator.h b/Engine/Plugins/AI/MassAI/Source/MassAIBehavior/Public/Evaluators/MassComponentHitEvaluator.h
new file mode 100644
--- /dev/null
+++ b/Engine/Plugins/AI/MassAI/Source/MassAIBehavior/Public/Evaluators/MassComponentHitEvaluator.h
@@ -0,0 +1,28 @@
+// Copyright Epic Games, Inc. All Rights Reserved.
+
+#pragma once
+
+#include "MassComponentHitSubsystem.h"
+#include "MassStateTreeTypes.h"
+#include "MassComponentHitEvaluator.generated.h"
+
+/**
+ * Evaluator to extract last hit from the MassComponentHitSubsystem and expose it for tasks and transitions
+ */
+USTRUCT(meta = (DisplayName = "Mass ComponentHit Eval"))
+struct MASSAIBEHAVIOR_API FMassComponentHitEvaluator : public FMassStateTreeEvaluatorBase
+{
+	GENERATED_BODY()
+
+protected:
+	virtual void Evaluate(FStateTreeExecutionContext& Context, const EStateTreeEvaluationType EvalType, const float DeltaTime) override;
+
+	UPROPERTY(meta=(BaseClass="MassComponentHitSubsystem"))
+	FStateTreeExternalItemHandle ComponentHitSubsystemHandle;
+
+	UPROPERTY(VisibleAnywhere, Category = Output)
+	bool bGotHit = false;
+
+	UPROPERTY(VisibleAnywhere, Category = Output)
+	FLWEntity LastHitEntity;
+};
diff --git a/Engine/Plugins/AI/MassAI/Source/MassAIBehavior/Public/Evaluators/MassStateTreeSmartObjectEvaluator.h b/Engine/Plugins/AI/MassAI/Source/MassAIBehavior/Public/Evaluators/MassStateTreeSmartObjectEvaluator.h
new file mode 100644
--- /dev/null
+++ b/Engine/Plugins/AI/MassAI/Source/MassAIBehavior/Public/Evaluators/MassStateTreeSmartObjectEvaluator.h
@@ -0,0 +1,70 @@
+// Copyright Epic Games, Inc. All Rights Reserved.
+
+#pragma once
+
+#include "MassSmartObjectProcessor.h"
+#include "MassStateTreeTypes.h"
+#include "MassSmartObjectTypes.h"
+#include "MassStateTreeSmartObjectEvaluator.generated.h"
+
+struct FStateTreeExecutionContext;
+
+/**
+ * Evaluator that will keep track if there is one or more smart object(s) that can be used.
+ */
+USTRUCT(meta = (DisplayName = "Mass SmartObject Eval"))
+struct MASSAIBEHAVIOR_API FMassStateTreeSmartObjectEvaluator : public FMassStateTreeEvaluatorBase
+{
+	GENERATED_BODY()
+
+protected:
+	virtual void EnterState(FStateTreeExecutionContext& Context, const EStateTreeStateChangeType ChangeType, const FStateTreeTransitionResult& Transition) override;
+	virtual void ExitState(FStateTreeExecutionContext& Context, const EStateTreeStateChangeType ChangeType, const FStateTreeTransitionResult& Transition) override;
+	virtual void Evaluate(FStateTreeExecutionContext& Context, const EStateTreeEvaluationType EvalType, const float DeltaTime) override;
+	void Reset();
+
+	UPROPERTY(meta=(BaseClass="SmartObjectSubsystem"))
+	FStateTreeExternalItemHandle SmartObjectSubsystemHandle;
+
+	UPROPERTY(meta=(BaseClass="MassSignalSubsystem"))
+	FStateTreeExternalItemHandle MassSignalSubsystemHandle;
+
+	UPROPERTY(meta=(BaseStruct="DataFragment_Transform"))
+	FStateTreeExternalItemHandle EntityTransformHandle;
+
+	UPROPERTY(meta=(BaseStruct="DataFragment_SmartObjectUser"))
+	FStateTreeExternalItemHandle SmartObjectUserHandle;
+
+	UPROPERTY(meta=(BaseStruct="MassZoneGraphLaneLocationFragment", Optional))
+	FStateTreeExternalItemHandle LocationHandle;
+
+	/** The identifier of the search request send by the evaluator to find candidates */
+	UPROPERTY(VisibleAnywhere, Category = SmartObject)
+	FMassSmartObjectRequestResult SearchRequestResult;
+
+	/** The identifier of the search request send by the evaluator to find candidates */
+	UPROPERTY(VisibleAnywhere, Category = SmartObject)
+	FMassSmartObjectRequestID SearchRequestID;
+
+	/** The delay that the evaluator should wait before trying to find a smart object after a failed attempt */
+	UPROPERTY(EditAnywhere, Category = SmartObject)
+	float RetryCooldown = 0.f;
+
+	/** The default delay that the evaluator should wait before evaluating again */
+	UPROPERTY(EditAnywhere, Category = SmartObject)
+	float TickInterval = 0.f;
+
+	/** Next update time; evaluator will not do anything when Evaluate gets called before that time */
+	float NextUpdate = 0.f;
+
+	/** Indicates that the result of the candidates search is ready and contains some candidates */
+	UPROPERTY(VisibleAnywhere, Category = SmartObject)
+	bool bCandidatesFound = false;
+	
+	/** Indicates that an object has been claimed */
+	UPROPERTY(VisibleAnywhere, Category = SmartObject)
+	bool bClaimed = false;
+
+	/** Indicates that the query was able to use annotations on zone graph lanes instead of a spatial query. */
+	bool bUsingZoneGraphAnnotations = false;
+};
diff --git a/Engine/Plugins/AI/MassAI/Source/MassAIBehavior/Public/Evaluators/MassStateTreeTestEvaluator.h b/Engine/Plugins/AI/MassAI/Source/MassAIBehavior/Public/Evaluators/MassStateTreeTestEvaluator.h
new file mode 100644
--- /dev/null
+++ b/Engine/Plugins/AI/MassAI/Source/MassAIBehavior/Public/Evaluators/MassStateTreeTestEvaluator.h
@@ -0,0 +1,34 @@
+// Copyright Epic Games, Inc. All Rights Reserved.
+
+#pragma once
+
+#include "StateTreeExecutionContext.h"
+#include "MassStateTreeTypes.h"
+#include "MassStateTreeTestEvaluator.generated.h"
+
+/**
+ * Test Evaluator, will be removed later.
+ */
+USTRUCT(meta = (DisplayName = "Mass Test Eval"))
+struct MASSAIBEHAVIOR_API FMassStateTreeTestEvaluator : public FMassStateTreeEvaluatorBase
+{
+	GENERATED_BODY()
+
+	FMassStateTreeTestEvaluator();
+	virtual ~FMassStateTreeTestEvaluator();
+	virtual void EnterState(FStateTreeExecutionContext& Context, const EStateTreeStateChangeType ChangeType, const FStateTreeTransitionResult& Transition) override;
+	virtual void Evaluate(FStateTreeExecutionContext& Context, const EStateTreeEvaluationType EvalType, const float DeltaTime) override;
+
+protected:
+	UPROPERTY()
+	float Time = 0.0f;
+
+	UPROPERTY()
+	bool bSignal = false;
+
+	UPROPERTY(EditAnywhere, Category = Test, meta = (Bindable))
+	EStateTreeEvaluationType Type = EStateTreeEvaluationType::Tick;
+	
+	UPROPERTY(EditAnywhere, Category = Test, meta = (Bindable))
+	float Period = 5.0f;
+};
diff --git a/Engine/Plugins/AI/MassAI/Source/MassAIBehavior/Public/Evaluators/MassZoneGraphAnnotationEvaluator.h b/Engine/Plugins/AI/MassAI/Source/MassAIBehavior/Public/Evaluators/MassZoneGraphAnnotationEvaluator.h
new file mode 100644
--- /dev/null
+++ b/Engine/Plugins/AI/MassAI/Source/MassAIBehavior/Public/Evaluators/MassZoneGraphAnnotationEvaluator.h
@@ -0,0 +1,29 @@
+﻿// Copyright Epic Games, Inc. All Rights Reserved.
+
+#pragma once
+
+#include "MassStateTreeTypes.h"
+#include "ZoneGraphTypes.h"
+#include "MassZoneGraphAnnotationEvaluator.generated.h"
+
+struct FStateTreeExecutionContext;
+
+/**
+ * Evaluator to expose ZoneGraph Annotation Tags for decision making.
+ */
+USTRUCT(meta = (DisplayName = "ZG Annotation Tags"))
+struct MASSAIBEHAVIOR_API FMassZoneGraphAnnotationEvaluator : public FMassStateTreeEvaluatorBase
+{
+	GENERATED_BODY()
+
+	FMassZoneGraphAnnotationEvaluator();
+
+protected:
+	virtual void Evaluate(FStateTreeExecutionContext& Context, const EStateTreeEvaluationType EvalType, const float DeltaTime) override;
+
+	UPROPERTY(meta=(BaseStruct="MassZoneGraphAnnotationTagsFragment"))
+	FStateTreeExternalItemHandle BehaviorTagsHandle;
+
+	UPROPERTY(EditAnywhere, Category = Parameters)
+	FZoneGraphTagMask BehaviorTags = FZoneGraphTagMask::None;
+};
diff --git a/Engine/Plugins/AI/MassAI/Source/MassAIBehavior/Public/MassAIBehaviorTypes.h b/Engine/Plugins/AI/MassAI/Source/MassAIBehavior/Public/MassAIBehaviorTypes.h
new file mode 100644
--- /dev/null
+++ b/Engine/Plugins/AI/MassAI/Source/MassAIBehavior/Public/MassAIBehaviorTypes.h
@@ -0,0 +1,38 @@
+// Copyright Epic Games, Inc. All Rights Reserved.
+
+#pragma once
+
+#include "Logging/LogMacros.h"
+#include "MassCommonTypes.h"
+#include "VisualLogger/VisualLogger.h"
+
+MASSAIBEHAVIOR_API DECLARE_LOG_CATEGORY_EXTERN(LogMassBehavior, Log, All);
+
+namespace UE::Mass::ProcessorGroupNames
+{
+	const FName UpdateAnnotationTags = FName(TEXT("UpdateAnnotationTags"));
+}
+
+/**
+ * Helper macros that could be used inside FMassStateTreeEvaluators and FMassStateTreeTasks.
+ * Requirements is a property or parameters with the following declaration: FStateTreeExecutionContext& Context
+ * These macros should be used to standardize the output format and simplify code at call site.
+ * They could also easily be changed from UE_(C)VLOG to UE_(C)VLOG_UELOG (or any other implementation) in one go.
+ *  e.g. the following:
+ *		#if WITH_MASS_DEBUG
+ *			const FMassStateTreeExecutionContext& MassContext = static_cast<FMassStateTreeExecutionContext&>(Context);
+ *			UE_VLOG(MassContext.GetOwner(), LogMassBehavior, Log, TEXT("Entity [%s]: Starting action: %s"), *MassContext.GetEntity().DebugGetDescription(), *StaticEnum<ESomeActionEnum>()->GetValueAsString(SomeActionEnumValue));
+ *		#endif // WITH_MASS_DEBUG
+ *
+ *	could be replaced by:
+ *		MASSBEHAVIOR_CLOG(bDisplayDebug, Log, TEXT("Starting action: %s"), *StaticEnum<ESomeActionEnum>()->GetValueAsString(SomeActionEnumValue));
+ */
+#if WITH_MASS_DEBUG
+#define MASSBEHAVIOR_LOG(Verbosity, Format, ...) UE_VLOG(static_cast<FMassStateTreeExecutionContext&>(Context).GetOwner(), LogMassBehavior, Verbosity, \
+	TEXT("Entity [%s][%s] ") Format, *static_cast<FMassStateTreeExecutionContext&>(Context).GetEntity().DebugGetDescription(), *StaticStruct()->GetName(), ##__VA_ARGS__)
+#define MASSBEHAVIOR_CLOG(Condition, Verbosity, Format, ...) UE_CVLOG((Condition), static_cast<FMassStateTreeExecutionContext&>(Context).GetOwner(), LogMassBehavior, Verbosity, \
+	TEXT("Entity [%s][%s] ") Format, *static_cast<FMassStateTreeExecutionContext&>(Context).GetEntity().DebugGetDescription(), *StaticStruct()->GetName(), ##__VA_ARGS__)
+#else
+#define MASSBEHAVIOR_LOG(Verbosity, Format, ...)
+#define MASSBEHAVIOR_CLOG(Condition, Verbosity, Format, ...)
+#endif // WITH_MASS_DEBUG
diff --git a/Engine/Plugins/AI/MassAI/Source/MassAIBehavior/Public/MassComponentHitSubsystem.h b/Engine/Plugins/AI/MassAI/Source/MassAIBehavior/Public/MassComponentHitSubsystem.h
new file mode 100644
--- /dev/null
+++ b/Engine/Plugins/AI/MassAI/Source/MassAIBehavior/Public/MassComponentHitSubsystem.h
@@ -0,0 +1,53 @@
+// Copyright Epic Games, Inc. All Rights Reserved.
+
+#pragma once
+
+#include "LWComponentTypes.h"
+#include "MassComponentHitTypes.h"
+#include "Subsystems/WorldSubsystem.h"
+
+#include "MassComponentHitSubsystem.generated.h"
+
+class UMassAgentSubsystem;
+class UMassSignalSubsystem;
+class UCapsuleComponent;
+
+
+/**
+ * Subsystem that keeps track of the latest component hits and allow mass entities to retrieve and handle them
+ */
+UCLASS()
+class MASSAIBEHAVIOR_API UMassComponentHitSubsystem : public UTickableWorldSubsystem
+{
+	GENERATED_BODY()
+
+public:
+	const FMassHitResult* GetLastHit(const FLWEntity Entity) const;
+
+protected:
+	virtual void Initialize(FSubsystemCollectionBase& Collection) override;
+	virtual void Deinitialize() override;
+	virtual void Tick(float DeltaTime) override;
+	virtual TStatId GetStatId() const override;
+	
+	void RegisterForComponentHit(const FLWEntity Entity, UCapsuleComponent& CapsuleComponent);
+	void UnregisterForComponentHit(FLWEntity Entity, UCapsuleComponent& CapsuleComponent);
+
+	UFUNCTION()
+	void OnHitCallback(UPrimitiveComponent* HitComp, AActor* OtherActor, UPrimitiveComponent* OtherComp, FVector NormalImpulse, const FHitResult& Hit);
+
+	UPROPERTY()
+	UMassSignalSubsystem* SignalSubsystem;
+
+	UPROPERTY()
+    UMassAgentSubsystem* AgentSubsystem;
+
+	UPROPERTY()
+	TMap<FLWEntity, FMassHitResult> HitResults;
+
+	UPROPERTY()
+	TMap<UActorComponent*, FLWEntity> ComponentToEntityMap;
+
+	UPROPERTY()
+	TMap<FLWEntity, UActorComponent*> EntityToComponentMap;
+};
diff --git a/Engine/Plugins/AI/MassAI/Source/MassAIBehavior/Public/MassComponentHitTypes.h b/Engine/Plugins/AI/MassAI/Source/MassAIBehavior/Public/MassComponentHitTypes.h
new file mode 100644
--- /dev/null
+++ b/Engine/Plugins/AI/MassAI/Source/MassAIBehavior/Public/MassComponentHitTypes.h
@@ -0,0 +1,53 @@
+// Copyright Epic Games, Inc. All Rights Reserved.
+
+#pragma once
+
+#include "Engine/EngineTypes.h"
+#include "LWComponentTypes.h"
+#include "StateTreeTypes.h"
+
+#include "MassComponentHitTypes.generated.h"
+
+namespace UE::Mass::Signals
+{
+const FName HitReceived = FName(TEXT("HitReceived"));
+}
+
+USTRUCT()
+struct FMassHitResult
+{
+	GENERATED_BODY()
+
+	FMassHitResult() = default;
+
+	FMassHitResult(const FLWEntity& OtherEntity, const float Time)
+		: OtherEntity(OtherEntity)
+		, HitTime(Time)
+		, LastFilteredHitTime(Time)
+	{
+	}
+
+	FLWEntity OtherEntity;
+
+	/** Time when first hit was received. */
+	float HitTime = 0.f;
+
+	/** Time used for filtering frequent hits. */
+	float LastFilteredHitTime = 0.f;
+};
+
+USTRUCT()
+struct FComponentHitStateTreeResult : public FStateTreeResult
+{
+	GENERATED_BODY()
+
+protected:
+	virtual const UScriptStruct& GetStruct() const override { return *StaticStruct(); }
+
+public:
+	UPROPERTY()
+	FMassHitResult MassHitResult;
+
+	UPROPERTY()
+	EStateTreeResultStatus Status = EStateTreeResultStatus::Unset;
+};
diff --git a/Engine/Plugins/AI/MassAI/Source/MassAIBehavior/Public/MassLookAtProcessors.h b/Engine/Plugins/AI/MassAI/Source/MassAIBehavior/Public/MassLookAtProcessors.h
new file mode 100644
--- /dev/null
+++ b/Engine/Plugins/AI/MassAI/Source/MassAIBehavior/Public/MassLookAtProcessors.h
@@ -0,0 +1,78 @@
+// Copyright Epic Games, Inc. All Rights Reserved.
+
+#pragma once
+
+#include "MassProcessor.h"
+#include "MassCommonTypes.h"
+#include "MassSignalProcessorBase.h"
+#include "MassLookAtProcessors.generated.h"
+
+class UMassMovementSubsystem;
+class UZoneGraphSubsystem;
+struct FMassLookAtFragment;
+struct FMassZoneGraphLaneLocationFragment;
+struct FMassLookAtTrajectoryFragment;
+struct FMassZoneGraphShortPathFragment;
+struct FMassMoveTargetFragment;
+
+/**
+ * Processor to choose and assign LookAt configurations  
+ */
+UCLASS()
+class MASSAIBEHAVIOR_API UMassLookAtProcessor : public UPipeProcessor
+{
+	GENERATED_BODY()
+
+	UMassLookAtProcessor();
+
+protected:
+
+	virtual void ConfigureQueries() override;
+	virtual void Initialize(UObject& Owner) override;
+	virtual void Execute(UEntitySubsystem& EntitySubsystem, FLWComponentSystemExecutionContext& Context) override;
+
+	/** Selects a nearby target if possible or use a random fixed direction */
+	void FindNewGazeTarget(const UEntitySubsystem& EntitySubsystem, const float CurrentTime, const FTransform& Transform, FMassLookAtFragment& LookAt) const;
+
+	/** Updates look direction based on look at trajectory. */
+	void UpdateLookAtTrajectory(const FTransform& Transform, const FMassZoneGraphLaneLocationFragment& ZoneGraphLocation,
+								const FMassLookAtTrajectoryFragment& LookAtTrajectory, const bool bDisplayDebug, FMassLookAtFragment& LookAt) const;
+
+	/** Updates look at based on tracked entity. */
+	void UpdateLookAtTrackedEntity(const UEntitySubsystem& EntitySubsystem, const FTransform& Transform, const bool bDisplayDebug, FMassLookAtFragment& LookAt) const;
+
+	/** Updates gaze based on tracked entity. */
+	bool UpdateGazeTrackedEntity(const UEntitySubsystem& EntitySubsystem, const FTransform& Transform, const bool bDisplayDebug, FMassLookAtFragment& LookAt) const;
+
+	/** Builds look at trajectory along the current path. */
+	void BuildTrajectory(const FMassZoneGraphLaneLocationFragment& LaneLocation, const FMassZoneGraphShortPathFragment& ShortPath,
+							const FLWEntity Entity, const bool bDisplayDebug, FMassLookAtTrajectoryFragment& LookAtTrajectory);
+
+	/** Size of the query to find potential targets */
+	UPROPERTY(EditDefaultsOnly, Category = LookAt, config, meta = (UIMin = 0.0, ClampMin = 0.0))
+	float QueryExtent = 0.f;
+
+	/** Time an entity must use a random look at. */
+	UPROPERTY(EditDefaultsOnly, Category = LookAt, config, meta = (UIMin = 0.0, ClampMin = 0.0))
+	float Duration = 0.f;
+	
+	/** Variation applied to a random look at duration [Duration-Variation : Duration+Variation] */
+	UPROPERTY(EditDefaultsOnly, Category = LookAt, config, meta = (UIMin = 0.0, ClampMin = 0.0))
+	float DurationVariation = 0.f;
+
+	/** Height offset that will be added for debug draw of the look at vector. Useful to bring arrow near character's eyes */
+	UPROPERTY(EditDefaultsOnly, Category = LookAt, config, meta = (UIMin = 0.0, ClampMin = 0.0, DisplayName="Debug draw Z offset (cm)"))
+	float DebugZOffset = 0.f;
+
+	/** Tolerance in degrees between the forward direction and the look at duration to track an entity */
+	UPROPERTY(EditDefaultsOnly, Category = LookAt, config, meta = (UIMin = 0.0, ClampMin = 0.0, UIMax = 180.0, ClampMax = 180.0, DisplayName="Angle Threshold (degrees)"))
+	float AngleThresholdInDegrees = 0.f;
+
+	FLWComponentQuery EntityQuery_Conditional;
+
+	UPROPERTY(Transient)
+	UMassMovementSubsystem* MovementSubsystem;
+
+	UPROPERTY(Transient)
+	UZoneGraphSubsystem* ZoneGraphSubsystem;
+};
diff --git a/Engine/Plugins/AI/MassAI/Source/MassAIBehavior/Public/MassLookAtTargetTrait.h b/Engine/Plugins/AI/MassAI/Source/MassAIBehavior/Public/MassLookAtTargetTrait.h
new file mode 100644
--- /dev/null
+++ b/Engine/Plugins/AI/MassAI/Source/MassAIBehavior/Public/MassLookAtTargetTrait.h
@@ -0,0 +1,17 @@
+// Copyright Epic Games, Inc. All Rights Reserved.
+
+#pragma once
+
+#include "CoreMinimal.h"
+#include "MassEntityTraitBase.h"
+#include "MassLookAtTargetTrait.generated.h"
+
+UCLASS(BlueprintType, EditInlineNew, CollapseCategories, meta=(DisplayName="Look At Target"))
+class UMassLookAtTargetTrait : public UMassEntityTraitBase
+{
+	GENERATED_BODY()
+
+protected:
+
+	virtual void BuildTemplate(FMassEntityTemplateBuildContext& BuildContext, UWorld& World) const override;
+};
diff --git a/Engine/Plugins/AI/MassAI/Source/MassAIBehavior/Public/MassLookAtTrait.h b/Engine/Plugins/AI/MassAI/Source/MassAIBehavior/Public/MassLookAtTrait.h
new file mode 100644
--- /dev/null
+++ b/Engine/Plugins/AI/MassAI/Source/MassAIBehavior/Public/MassLookAtTrait.h
@@ -0,0 +1,17 @@
+﻿// Copyright Epic Games, Inc. All Rights Reserved.
+
+#pragma once
+
+#include "CoreMinimal.h"
+#include "MassEntityTraitBase.h"
+#include "MassLookAtTrait.generated.h"
+
+UCLASS(BlueprintType, EditInlineNew, CollapseCategories, meta=(DisplayName="Look At"))
+class UMassLookAtTrait : public UMassEntityTraitBase
+{
+	GENERATED_BODY()
+
+protected:
+
+	virtual void BuildTemplate(FMassEntityTemplateBuildContext& BuildContext, UWorld& World) const override;
+};
diff --git a/Engine/Plugins/AI/MassAI/Source/MassAIBehavior/Public/MassStateTreeExecutionContext.h b/Engine/Plugins/AI/MassAI/Source/MassAIBehavior/Public/MassStateTreeExecutionContext.h
new file mode 100644
--- /dev/null
+++ b/Engine/Plugins/AI/MassAI/Source/MassAIBehavior/Public/MassStateTreeExecutionContext.h
@@ -0,0 +1,46 @@
+// Copyright Epic Games, Inc. All Rights Reserved.
+
+#pragma once
+
+#include "LWComponentTypes.h"
+#include "StateTreeExecutionContext.h"
+#include "MassStateTreeExecutionContext.generated.h"
+
+struct FLWComponentSystemExecutionContext;
+class UEntitySubsystem;
+class UMassSignalSubsystem;
+
+/**
+ * Extends FStateTreeExecutionContext to provide additional data to Evaluators and Tasks related to MassSimulation
+ */
+USTRUCT()
+struct MASSAIBEHAVIOR_API FMassStateTreeExecutionContext : public FStateTreeExecutionContext 
+{
+	GENERATED_BODY()
+public:
+	/** Should never be used but has to be public for 'void UScriptStruct::TCppStructOps<FMassStateTreeExecutionContext>::ConstructForTests(void *)' */
+	FMassStateTreeExecutionContext() = default;
+	FMassStateTreeExecutionContext(UEntitySubsystem& InEntitySubsystem, FLWComponentSystemExecutionContext& InContext);
+
+	UEntitySubsystem& GetEntitySubsystem() const { return *EntitySubsystem; }
+	FLWComponentSystemExecutionContext& GetEntitySubsystemExecutionContext() const { return *EntitySubsystemExecutionContext; }
+
+	FLWEntity GetEntity() const { return Entity; }
+	void SetEntity(const FLWEntity& InEntity) { Entity = InEntity; }
+
+	int32 GetEntityIndex() const { return EntityIndex; }
+	void SetEntityIndex(const int32 InIndex) { EntityIndex = InIndex; }
+
+protected:
+
+	/** Prefix that will be used by STATETREE_LOG and STATETREE_CLOG, using Entity description. */
+	virtual FString GetInstanceDescription() const override { return FString::Printf(TEXT("Entity [%s]"), *Entity.DebugGetDescription()); }
+
+	virtual void BeginGatedTransition(const FStateTreeExecutionState& Exec) override;
+
+	UEntitySubsystem* EntitySubsystem = nullptr;
+	UMassSignalSubsystem* SignalSubsystem = nullptr;
+	FLWComponentSystemExecutionContext* EntitySubsystemExecutionContext = nullptr;
+	FLWEntity Entity;
+	int32 EntityIndex = INDEX_NONE;
+};
diff --git a/Engine/Plugins/AI/MassAI/Source/MassAIBehavior/Public/MassStateTreeFragments.h b/Engine/Plugins/AI/MassAI/Source/MassAIBehavior/Public/MassStateTreeFragments.h
new file mode 100644
--- /dev/null
+++ b/Engine/Plugins/AI/MassAI/Source/MassAIBehavior/Public/MassStateTreeFragments.h
@@ -0,0 +1,22 @@
+// Copyright Epic Games, Inc. All Rights Reserved.
+
+#pragma once
+
+#include "LWComponentTypes.h"
+#include "MassStateTreeSubsystem.h"
+#include "MassStateTreeFragments.generated.h"
+
+
+USTRUCT()
+struct MASSAIBEHAVIOR_API FMassStateTreeFragment : public FLWComponentData
+{
+	GENERATED_BODY()
+
+	FMassStateTreeFragment() = default;
+
+	/** Handle to a StateTree asset in MassStateTreeSubsystem */
+	FMassStateTreeHandle StateTreeHandle;
+
+	/** Keep track of the last update time to adjust time delta */
+	TOptional<float> LastUpdateTimeInSeconds;
+};
diff --git a/Engine/Plugins/AI/MassAI/Source/MassAIBehavior/Public/MassStateTreeProcessors.h b/Engine/Plugins/AI/MassAI/Source/MassAIBehavior/Public/MassStateTreeProcessors.h
new file mode 100644
--- /dev/null
+++ b/Engine/Plugins/AI/MassAI/Source/MassAIBehavior/Public/MassStateTreeProcessors.h
@@ -0,0 +1,63 @@
+// Copyright Epic Games, Inc. All Rights Reserved.
+
+#pragma once
+
+#include "MassSignals/Public/MassSignalProcessorBase.h"
+#include "MassStateTreeFragments.h"
+#include "MassTranslator.h"
+#include "MassStateTreeProcessors.generated.h"
+
+struct FMassStateTreeExecutionContext;
+
+/** 
+ * Processor to prepare StateTree fragments for their execution
+ */
+UCLASS()
+class MASSAIBEHAVIOR_API UMassStateTreeFragmentInitializer : public UMassFragmentInitializer
+{
+	GENERATED_BODY()
+
+public:
+	UMassStateTreeFragmentInitializer();
+
+protected:
+	virtual void ConfigureQueries() override;
+	virtual void Execute(UEntitySubsystem& EntitySubsystem, FLWComponentSystemExecutionContext& Context) override;
+
+	FLWComponentQuery EntityQuery;
+};
+
+UCLASS()
+class MASSAIBEHAVIOR_API UMassStateTreeFragmentDestructor : public UMassFragmentDestructor
+{
+	GENERATED_BODY()
+
+public:
+	UMassStateTreeFragmentDestructor();
+
+protected:
+	virtual void ConfigureQueries() override;
+	virtual void Execute(UEntitySubsystem& EntitySubsystem, FLWComponentSystemExecutionContext& Context) override;
+
+	FLWComponentQuery EntityQuery;
+};
+
+/** 
+ * Processor for executing a StateTree
+ */
+UCLASS()
+class MASSAIBEHAVIOR_API UMassStateTreeProcessor : public UMassSignalProcessorBase
+{
+	GENERATED_BODY()
+
+public:
+	UMassStateTreeProcessor(const FObjectInitializer& ObjectInitializer = FObjectInitializer::Get());
+
+protected:
+	virtual void Initialize(UObject& Owner) override;
+	virtual void ConfigureQueries() override;
+	virtual void SignalEntities(UEntitySubsystem& EntitySubsystem, FLWComponentSystemExecutionContext& Context, FMassSignalNameLookup& EntitySignals) override;
+
+	UPROPERTY(Transient)
+	UMassStateTreeSubsystem* MassStateTreeSubsystem;
+};
diff --git a/Engine/Plugins/AI/MassAI/Source/MassAIBehavior/Public/MassStateTreeSchema.h b/Engine/Plugins/AI/MassAI/Source/MassAIBehavior/Public/MassStateTreeSchema.h
new file mode 100644
--- /dev/null
+++ b/Engine/Plugins/AI/MassAI/Source/MassAIBehavior/Public/MassStateTreeSchema.h
@@ -0,0 +1,23 @@
+// Copyright Epic Games, Inc. All Rights Reserved.
+
+#pragma once
+
+#include "StateTreeSchema.h"
+#include "MassStateTreeSchema.generated.h"
+
+/**
+ * StateTree Schema for Mass behaviors.
+ */
+UCLASS(BlueprintType, EditInlineNew, CollapseCategories, meta = (DisplayName = "Mass Behavior"))
+class MASSAIBEHAVIOR_API UMassStateTreeSchema : public UStateTreeSchema
+{
+	GENERATED_BODY()
+
+protected:
+
+	virtual UScriptStruct* GetStorageSuperStruct() const override;
+	virtual bool IsStructAllowed(const UScriptStruct* InScriptStruct) const override;
+	virtual bool IsExternalItemAllowed(const UStruct& InStruct) const override;
+	virtual bool IsV2() const override { return true; }
+};
+
diff --git a/Engine/Plugins/AI/MassAI/Source/MassAIBehavior/Public/MassStateTreeSubsystem.h b/Engine/Plugins/AI/MassAI/Source/MassAIBehavior/Public/MassStateTreeSubsystem.h
new file mode 100644
--- /dev/null
+++ b/Engine/Plugins/AI/MassAI/Source/MassAIBehavior/Public/MassStateTreeSubsystem.h
@@ -0,0 +1,39 @@
+// Copyright Epic Games, Inc. All Rights Reserved.
+
+#pragma once
+
+#include "MassStateTreeTypes.h"
+#include "Subsystems/WorldSubsystem.h"
+#include "MassStateTreeSubsystem.generated.h"
+
+class UStateTree;
+
+/**
+* A subsystem managing StateTree assets in Mass
+*/
+UCLASS()
+class MASSAIBEHAVIOR_API UMassStateTreeSubsystem : public UWorldSubsystem
+{
+	GENERATED_BODY()
+	
+public:
+
+	// USubsystem BEGIN
+	virtual void Initialize(FSubsystemCollectionBase& Collection) override;
+	// USubsystem END
+
+	/** Registers a StateTree asset to be used */
+	FMassStateTreeHandle RegisterStateTreeAsset(UStateTree* StateTree);
+
+	/** @return StateTree asset based on a handle */
+	UStateTree* GetRegisteredStateTreeAsset(FMassStateTreeHandle Handle) { return RegisteredStateTrees[Handle.GetIndex()]; }
+
+	/** @return Array of registered StateTree assets */
+	TConstArrayView<UStateTree*> GetRegisteredStateTreeAssets() const { return RegisteredStateTrees; }
+	
+protected:
+
+	/** Array of registered (in use) StateTrees */
+	UPROPERTY(Transient)
+	TArray<UStateTree*> RegisteredStateTrees;
+};
diff --git a/Engine/Plugins/AI/MassAI/Source/MassAIBehavior/Public/MassStateTreeTrait.h b/Engine/Plugins/AI/MassAI/Source/MassAIBehavior/Public/MassStateTreeTrait.h
new file mode 100644
--- /dev/null
+++ b/Engine/Plugins/AI/MassAI/Source/MassAIBehavior/Public/MassStateTreeTrait.h
@@ -0,0 +1,26 @@
+// Copyright Epic Games, Inc. All Rights Reserved.
+
+#pragma once
+
+#include "MassEntityTraitBase.h"
+#include "MassStateTreeTrait.generated.h"
+
+class UStateTree;
+
+/**
+ * Feature that adds StateTree execution functionality to a mass agent.
+ */
+UCLASS(BlueprintType, EditInlineNew, CollapseCategories, meta=(DisplayName="StateTree"))
+class MASSAIBEHAVIOR_API UMassStateTreeTrait : public UMassEntityTraitBase
+{
+	GENERATED_BODY()
+
+protected:
+
+	virtual void BuildTemplate(FMassEntityTemplateBuildContext& BuildContext, UWorld& World) const override;
+
+	virtual void ValidateTemplate(FMassEntityTemplateBuildContext& BuildContext, UWorld& World) const override;
+
+	UPROPERTY(Category="StateTree", EditAnywhere, meta=(RequiredAssetDataTags="Schema=MassStateTreeSchema"))
+	UStateTree* StateTree;
+};
diff --git a/Engine/Plugins/AI/MassAI/Source/MassAIBehavior/Public/MassStateTreeTypes.h b/Engine/Plugins/AI/MassAI/Source/MassAIBehavior/Public/MassStateTreeTypes.h
new file mode 100644
--- /dev/null
+++ b/Engine/Plugins/AI/MassAI/Source/MassAIBehavior/Public/MassStateTreeTypes.h
@@ -0,0 +1,70 @@
+// Copyright Epic Games, Inc. All Rights Reserved.
+
+#pragma once
+
+#include "CoreMinimal.h"
+#include "StateTreeEvaluatorBase.h" 
+#include "StateTreeTaskBase.h"
+#include "MassStateTreeTypes.generated.h"
+
+/**
+ * Signals used by the StateTree framework in Mass
+ */
+namespace UE::Mass::Signals
+{
+	const FName StateTreeInitializationRequested = FName(TEXT("StateTreeInitializationRequested"));
+	const FName LookAtFinished = FName(TEXT("LookAtFinished"));
+	const FName NewStateTreeTaskRequired = FName(TEXT("NewStateTreeTaskRequired"));
+	const FName StandTaskFinished = FName(TEXT("StandTaskFinished"));
+	const FName DelayedTransitionWakeup = FName(TEXT("DelayedTransitionWakeup"));
+	// @todo MassStateTree: move this to its game plugin when possible
+	const FName ContextualAnimTaskFinished = FName(TEXT("ContextualAnimTaskFinished"));
+}
+
+/**
+ * Base struct for all Mass StateTree Evaluators.
+ */
+USTRUCT(meta = (DisplayName = "Mass Evaluator Base"))
+struct MASSAIBEHAVIOR_API FMassStateTreeEvaluatorBase : public FStateTreeEvaluator2Base
+{
+	GENERATED_BODY()
+
+	FMassStateTreeEvaluatorBase() {}
+	virtual ~FMassStateTreeEvaluatorBase() {}
+};
+template<> struct TStructOpsTypeTraits<FMassStateTreeEvaluatorBase> : public TStructOpsTypeTraitsBase2<FMassStateTreeEvaluatorBase> { enum { WithPureVirtual = true, }; };
+
+/**
+ * Base struct for all Mass StateTree Tasks.
+ */
+USTRUCT(meta = (DisplayName = "Mass Task Base"))
+struct MASSAIBEHAVIOR_API FMassStateTreeTaskBase : public FStateTreeTask2Base
+{
+	GENERATED_BODY()
+
+	FMassStateTreeTaskBase() {}
+	virtual ~FMassStateTreeTaskBase() {}
+};
+template<> struct TStructOpsTypeTraits<FMassStateTreeTaskBase> : public TStructOpsTypeTraitsBase2<FMassStateTreeTaskBase> { enum { WithPureVirtual = true, }; };
+
+/**
+* A handle pointing to a registered StateTree asset in UMassStateTreeSubsystem.
+*/
+struct FMassStateTreeHandle
+{
+	static constexpr uint16 Invalid = MAX_uint16;
+
+	FMassStateTreeHandle() = default;
+
+	/** Initializes new handle based on an index */
+	static FMassStateTreeHandle Make(const uint16 InIndex) { return FMassStateTreeHandle(InIndex); }
+	
+	/** Returns index the handle points to */
+	uint16 GetIndex() const { return Index; }
+
+protected:
+	FMassStateTreeHandle(const uint16 InIndex) : Index(InIndex) {}
+
+	uint16 Index = Invalid;
+};
+
diff --git a/Engine/Plugins/AI/MassAI/Source/MassAIBehavior/Public/MassZoneGraphAnnotationFragments.h b/Engine/Plugins/AI/MassAI/Source/MassAIBehavior/Public/MassZoneGraphAnnotationFragments.h
new file mode 100644
--- /dev/null
+++ b/Engine/Plugins/AI/MassAI/Source/MassAIBehavior/Public/MassZoneGraphAnnotationFragments.h
@@ -0,0 +1,30 @@
+// Copyright Epic Games, Inc. All Rights Reserved.
+
+#pragma once
+
+#include "LWComponentTypes.h"
+#include "ZoneGraphTypes.h"
+#include "MassZoneGraphAnnotationFragments.generated.h"
+
+struct FLWComponentSystemExecutionContext;
+
+USTRUCT()
+struct MASSAIBEHAVIOR_API FMassZoneGraphAnnotationTagsFragment : public FLWComponentData
+{
+	GENERATED_BODY()
+
+	/** Behavior tags for current lane */
+	UPROPERTY()
+	FZoneGraphTagMask Tags;
+};
+
+USTRUCT()
+struct FMassZoneGraphAnnotationVariableTickChunkFragment : public FLWChunkComponent
+{
+	GENERATED_BODY();
+
+	/** Update the ticking frequency of the chunk and return if this chunk should be process this frame */
+	static bool UpdateChunk(FLWComponentSystemExecutionContext& Context);
+
+	float TimeUntilNextTick = 0.0f;
+};
\ No newline at end of file
diff --git a/Engine/Plugins/AI/MassAI/Source/MassAIBehavior/Public/MassZoneGraphAnnotationProcessors.h b/Engine/Plugins/AI/MassAI/Source/MassAIBehavior/Public/MassZoneGraphAnnotationProcessors.h
new file mode 100644
--- /dev/null
+++ b/Engine/Plugins/AI/MassAI/Source/MassAIBehavior/Public/MassZoneGraphAnnotationProcessors.h
@@ -0,0 +1,62 @@
+// Copyright Epic Games, Inc. All Rights Reserved.
+
+#pragma once
+
+#include "MassSignalProcessorBase.h"
+#include "MassTranslator.h"
+#include "MassZoneGraphAnnotationProcessors.generated.h"
+
+class UMassSignalSubsystem;
+class UZoneGraphAnnotationSubsystem;
+struct FMassZoneGraphAnnotationTagsFragment;
+struct FMassZoneGraphLaneLocationFragment;
+struct FLWEntity;
+
+/** 
+ * Processor for initializing ZoneGraph annotation tags.
+ */
+UCLASS()
+class MASSAIBEHAVIOR_API UMassZoneGraphAnnotationTagsInitializer : public UMassFragmentInitializer
+{
+	GENERATED_BODY()
+
+public:
+	UMassZoneGraphAnnotationTagsInitializer();
+
+protected:
+	virtual void Initialize(UObject& Owner) override;
+	virtual void ConfigureQueries() override;
+	virtual void Execute(UEntitySubsystem& EntitySubsystem, FLWComponentSystemExecutionContext& Context) override;
+
+	FLWComponentQuery EntityQuery;
+
+	UPROPERTY(Transient)
+	UZoneGraphAnnotationSubsystem* ZoneGraphAnnotationSubsystem = nullptr;
+};
+
+/** 
+ * Processor for update ZoneGraph annotation tags periodically and on lane changed signal.
+ */
+UCLASS()
+class MASSAIBEHAVIOR_API UMassZoneGraphAnnotationTagUpdateProcessor : public UMassSignalProcessorBase
+{
+	GENERATED_BODY()
+
+public:
+	UMassZoneGraphAnnotationTagUpdateProcessor();
+	
+protected:
+	virtual void Initialize(UObject& Owner) override;
+	virtual void ConfigureQueries() override;
+	virtual void Execute(UEntitySubsystem& EntitySubsystem, FLWComponentSystemExecutionContext& Context) override;
+
+	void UpdateAnnotationTags(FMassZoneGraphAnnotationTagsFragment& AnnotationTags, const FMassZoneGraphLaneLocationFragment& LaneLocation, FLWEntity Entity);
+
+	virtual void SignalEntities(UEntitySubsystem& EntitySubsystem, FLWComponentSystemExecutionContext& Context, FMassSignalNameLookup& EntitySignals) override;
+
+	UPROPERTY(Transient)
+	UZoneGraphAnnotationSubsystem* ZoneGraphAnnotationSubsystem = nullptr;
+
+	// Frame buffer, it gets reset every frame.
+	TArray<FLWEntity> TransientEntitiesToSignal;
+};
diff --git a/Engine/Plugins/AI/MassAI/Source/MassAIBehavior/Public/MassZoneGraphAnnotationTrait.h b/Engine/Plugins/AI/MassAI/Source/MassAIBehavior/Public/MassZoneGraphAnnotationTrait.h
new file mode 100644
--- /dev/null
+++ b/Engine/Plugins/AI/MassAI/Source/MassAIBehavior/Public/MassZoneGraphAnnotationTrait.h
@@ -0,0 +1,15 @@
+﻿// Copyright Epic Games, Inc. All Rights Reserved.
+
+#pragma once 
+
+#include "MassEntityTraitBase.h"
+#include "MassZoneGraphAnnotationTrait.generated.h"
+
+UCLASS(BlueprintType, EditInlineNew, CollapseCategories, meta = (DisplayName = "ZoneGraph Annotation"))
+class MASSAIBEHAVIOR_API UMassZoneGraphAnnotationTrait : public UMassEntityTraitBase
+{
+	GENERATED_BODY()
+
+protected:
+	virtual void BuildTemplate(FMassEntityTemplateBuildContext& BuildContext, UWorld& World) const override;
+};
diff --git a/Engine/Plugins/AI/MassAI/Source/MassAIBehavior/Public/Tasks/MassLookAtTask.h b/Engine/Plugins/AI/MassAI/Source/MassAIBehavior/Public/Tasks/MassLookAtTask.h
new file mode 100644
--- /dev/null
+++ b/Engine/Plugins/AI/MassAI/Source/MassAIBehavior/Public/Tasks/MassLookAtTask.h
@@ -0,0 +1,59 @@
+// Copyright Epic Games, Inc. All Rights Reserved.
+
+#pragma once
+
+#include "LWComponentTypes.h"
+#include "MassStateTreeTypes.h"
+#include "MassLookAtFragments.h"
+#include "MassLookAtTask.generated.h"
+
+/**
+ * Task to assign a LookAt target for mass processing
+ */
+USTRUCT(meta = (DisplayName = "Mass LookAt Task"))
+struct MASSAIBEHAVIOR_API FMassLookAtTask : public FMassStateTreeTaskBase
+{
+	GENERATED_BODY()
+
+protected:
+	virtual EStateTreeRunStatus EnterState(FStateTreeExecutionContext& Context, const EStateTreeStateChangeType ChangeType, const FStateTreeTransitionResult& Transition) override;
+	virtual void ExitState(FStateTreeExecutionContext& Context, const EStateTreeStateChangeType ChangeType, const FStateTreeTransitionResult& Transition) override;
+	virtual EStateTreeRunStatus Tick(FStateTreeExecutionContext& Context, const float DeltaTime) override;
+
+	UPROPERTY(meta = (BaseClass="MassSignalSubsystem"))
+	FStateTreeExternalItemHandle MassSignalSubsystemHandle;
+
+	UPROPERTY(meta = (BaseStruct="MassLookAtFragment"))
+	FStateTreeExternalItemHandle LookAtHandle;
+
+	/** Delay before the task ends. Default (0 or any negative) will run indefinitely so it requires a transition in the state tree to stop it. */
+	UPROPERTY(EditAnywhere, Category = Parameters, meta = (Bindable))
+	float Duration = 0.f;
+
+	/** Look At Mode */
+	UPROPERTY(EditAnywhere, Category = Params, meta = (Bindable))
+	EMassLookAtMode LookAtMode = EMassLookAtMode::LookForward; 
+
+	/** Entity to set as the target for the LookAt behavior. */
+	UPROPERTY(VisibleAnywhere, Category = Params, meta = (Bindable))
+	FLWEntity TargetEntity;
+
+	/** Random gaze Mode */
+	UPROPERTY(EditAnywhere, Category = Params, meta = (Bindable))
+	EMassLookAtGazeMode RandomGazeMode = EMassLookAtGazeMode::None;
+	
+	/** Random gaze yaw angle added to the look direction determined by the look at mode. */
+	UPROPERTY(EditAnywhere, Category = Params, meta = (Bindable, UIMin = 0.0, ClampMin = 0.0, UIMax = 180.0, ClampMax = 180.0))
+	uint8 RandomGazeYawVariation = 0;
+
+	/** Random gaze pitch angle added to the look direction determined by the look at mode. */
+	UPROPERTY(EditAnywhere, Category = Params, meta = (Bindable, UIMin = 0.0, ClampMin = 0.0, UIMax = 180.0, ClampMax = 180.0))
+	uint8 RandomGazePitchVariation = 0;
+
+	/** If true, allow random gaze to look at other entities too. */
+	UPROPERTY(EditAnywhere, Category = Params, meta = (Bindable))
+	bool bRandomGazeEntities = false;
+
+	/** Accumulated time used to stop task if duration is set */
+	float Time = 0.f;
+};
diff --git a/Engine/Plugins/AI/MassAI/Source/MassAIBehavior/Public/Tasks/MassStateTreeSmartObjectTask.h b/Engine/Plugins/AI/MassAI/Source/MassAIBehavior/Public/Tasks/MassStateTreeSmartObjectTask.h
new file mode 100644
--- /dev/null
+++ b/Engine/Plugins/AI/MassAI/Source/MassAIBehavior/Public/Tasks/MassStateTreeSmartObjectTask.h
@@ -0,0 +1,71 @@
+// Copyright Epic Games, Inc. All Rights Reserved.
+
+#pragma once
+
+#include "MassSmartObjectRequest.h"
+#include "MassStateTreeTypes.h"
+#include "MassSmartObjectTypes.h"
+
+#include "MassStateTreeSmartObjectTask.generated.h"
+
+struct FStateTreeExecutionContext;
+
+/**
+ * Tasks to claim a smart object from search results and release it when done.
+ */
+USTRUCT(meta = (DisplayName = "Claim SmartObject"))
+struct MASSAIBEHAVIOR_API FMassClaimSmartObjectTask : public FMassStateTreeTaskBase
+{
+	GENERATED_BODY()
+
+protected:
+	virtual EStateTreeRunStatus EnterState(FStateTreeExecutionContext& Context, const EStateTreeStateChangeType ChangeType, const FStateTreeTransitionResult& Transition) override;
+	virtual void ExitState(FStateTreeExecutionContext& Context, const EStateTreeStateChangeType ChangeType, const FStateTreeTransitionResult& Transition) override;
+	virtual EStateTreeRunStatus Tick(FStateTreeExecutionContext& Context, const float DeltaTime) override;
+
+	UPROPERTY(meta=(BaseStruct="DataFragment_SmartObjectUser"))
+	FStateTreeExternalItemHandle SmartObjectUserHandle;
+
+	UPROPERTY(meta=(BaseClass="SmartObjectSubsystem"))
+	FStateTreeExternalItemHandle SmartObjectSubsystemHandle;
+
+	/** Result of the candidates search request (Input) */
+	UPROPERTY(VisibleAnywhere, Category = SmartObject, meta = (Bindable), Transient)
+	FMassSmartObjectRequestResult SearchRequestResult;
+
+	/** Result of the claim on potential candidates from the search results (Output) */
+	UPROPERTY(VisibleAnywhere, Category = SmartObject, Transient)
+	EMassSmartObjectClaimResult ClaimResult = EMassSmartObjectClaimResult::Unset;
+};
+
+/**
+ * Task to tell an entity to start using a claimed smart object.
+ */
+USTRUCT(meta = (DisplayName = "Mass Use SmartObject Task"))
+struct MASSAIBEHAVIOR_API FMassUseSmartObjectTask : public FMassStateTreeTaskBase
+{
+	GENERATED_BODY()
+
+protected:
+	virtual EStateTreeRunStatus EnterState(FStateTreeExecutionContext& Context, const EStateTreeStateChangeType ChangeType, const FStateTreeTransitionResult& Transition) override;
+	virtual void ExitState(FStateTreeExecutionContext& Context, const EStateTreeStateChangeType ChangeType, const FStateTreeTransitionResult& Transition) override;
+	virtual void StateCompleted(FStateTreeExecutionContext& Context, const EStateTreeRunStatus CompletionStatus, const FStateTreeHandle CompletedState) override;
+	virtual EStateTreeRunStatus Tick(FStateTreeExecutionContext& Context, const float DeltaTime) override;
+
+	UPROPERTY(meta=(BaseClass="SmartObjectSubsystem"))
+	FStateTreeExternalItemHandle SmartObjectSubsystemHandle;
+
+	UPROPERTY(meta=(BaseStruct="DataFragment_Transform"))
+	FStateTreeExternalItemHandle EntityTransformHandle;
+
+	UPROPERTY(meta=(BaseStruct="DataFragment_SmartObjectUser"))
+	FStateTreeExternalItemHandle SmartObjectUserHandle;
+
+	UPROPERTY(meta=(BaseStruct="MassMoveTargetFragment"))
+	FStateTreeExternalItemHandle MoveTargetHandle;
+
+	/** Delay in seconds before trying to find & use another smart object */
+	UPROPERTY(EditAnywhere, Category = SmartObject)
+	float Cooldown = 0.f;
+};
+
diff --git a/Engine/Plugins/AI/MassAI/Source/MassAIBehavior/Public/Tasks/MassStateTreeTestTask.h b/Engine/Plugins/AI/MassAI/Source/MassAIBehavior/Public/Tasks/MassStateTreeTestTask.h
new file mode 100644
--- /dev/null
+++ b/Engine/Plugins/AI/MassAI/Source/MassAIBehavior/Public/Tasks/MassStateTreeTestTask.h
@@ -0,0 +1,37 @@
+// Copyright Epic Games, Inc. All Rights Reserved.
+
+#pragma once
+
+#include "StateTreeExecutionContext.h"
+#include "MassStateTreeTypes.h"
+#include "MassStateTreeTestTask.generated.h"
+
+/**
+ * Test Task, will be removed later.
+ */
+USTRUCT(meta = (DisplayName = "Mass Test Task"))
+struct MASSAIBEHAVIOR_API FMassStateTreeTestTask : public FMassStateTreeTaskBase
+{
+	GENERATED_BODY()
+
+	FMassStateTreeTestTask();
+	virtual ~FMassStateTreeTestTask();
+	virtual EStateTreeRunStatus EnterState(FStateTreeExecutionContext& Context, const EStateTreeStateChangeType ChangeType, const FStateTreeTransitionResult& Transition) override;
+	virtual EStateTreeRunStatus Tick(FStateTreeExecutionContext& Context, const float DeltaTime) override;
+
+protected:
+	UPROPERTY()
+	float Time = 0.0f;
+
+	UPROPERTY(EditAnywhere, Category = Test, meta = (Bindable))
+	float Duration = 5.0f;
+
+	UPROPERTY(meta=(BaseClass="MassStateTreeSubsystem"))
+	FStateTreeExternalItemHandle MassStateTreeSubSystemHandle;
+
+	UPROPERTY(meta=(BaseStruct="DataFragment_SmartObjectUser", Optional))
+	FStateTreeExternalItemHandle SmartObjectUserHandle;
+
+	UPROPERTY(meta=(BaseStruct="DataFragment_Transform"))
+	FStateTreeExternalItemHandle TransformHandle;
+};
diff --git a/Engine/Plugins/AI/MassAI/Source/MassAIBehavior/Public/Tasks/MassZoneGraphFindEscapeTarget.h b/Engine/Plugins/AI/MassAI/Source/MassAIBehavior/Public/Tasks/MassZoneGraphFindEscapeTarget.h
new file mode 100644
--- /dev/null
+++ b/Engine/Plugins/AI/MassAI/Source/MassAIBehavior/Public/Tasks/MassZoneGraphFindEscapeTarget.h
@@ -0,0 +1,45 @@
+﻿// Copyright Epic Games, Inc. All Rights Reserved.
+
+#pragma once
+
+#include "MassStateTreeTypes.h"
+#include "Tasks/MassZoneGraphPathFollowTask.h"
+#include "MassZoneGraphFindEscapeTarget.generated.h"
+
+struct FStateTreeExecutionContext;
+
+/**
+ * Updates TargetLocation to a escape target based on the agents current location on ZoneGraph, and disturbance annotation.
+ */
+USTRUCT(meta = (DisplayName = "ZG Find Escape Target"))
+struct MASSAIBEHAVIOR_API FMassZoneGraphFindEscapeTarget : public FMassStateTreeTaskBase
+{
+	GENERATED_BODY()
+
+	FMassZoneGraphFindEscapeTarget();
+
+protected:
+	virtual EStateTreeRunStatus EnterState(FStateTreeExecutionContext& Context, const EStateTreeStateChangeType ChangeType, const FStateTreeTransitionResult& Transition) override;
+	virtual void ExitState(FStateTreeExecutionContext& Context, const EStateTreeStateChangeType ChangeType, const FStateTreeTransitionResult& Transition) override;
+	virtual EStateTreeRunStatus Tick(FStateTreeExecutionContext& Context, const float DeltaTime) override;
+
+	UPROPERTY(meta=(BaseStruct="MassZoneGraphLaneLocationFragment"))
+	FStateTreeExternalItemHandle LocationHandle;
+
+	UPROPERTY(meta=(BaseStruct="MassMoveTargetFragment"))
+	FStateTreeExternalItemHandle MoveTargetHandle;
+
+	UPROPERTY(meta=(BaseClass="ZoneGraphSubsystem"))
+	FStateTreeExternalItemHandle ZoneGraphSubsystemHandle;
+
+	UPROPERTY(meta=(BaseClass="ZoneGraphAnnotationSubsystem"))
+	FStateTreeExternalItemHandle ZoneGraphAnnotationSubsystemHandle;
+
+	UPROPERTY(EditAnywhere, Category = Parameters)
+	FZoneGraphTag DisturbanceAnnotationTag = FZoneGraphTag::None;
+
+	FMassZoneGraphTargetLocation TargetLocation;
+
+	UPROPERTY(EditAnywhere, Category = Parameters, meta=(Struct="MassZoneGraphTargetLocation"))
+	FStateTreeResultRef EscapeTargetLocation;
+};
diff --git a/Engine/Plugins/AI/MassAI/Source/MassAIBehavior/Public/Tasks/MassZoneGraphFindSmartObjectTarget.h b/Engine/Plugins/AI/MassAI/Source/MassAIBehavior/Public/Tasks/MassZoneGraphFindSmartObjectTarget.h
new file mode 100644
--- /dev/null
+++ b/Engine/Plugins/AI/MassAI/Source/MassAIBehavior/Public/Tasks/MassZoneGraphFindSmartObjectTarget.h
@@ -0,0 +1,35 @@
+// Copyright Epic Games, Inc. All Rights Reserved.
+
+#pragma once
+#include "MassStateTreeTypes.h"
+#include "MassZoneGraphPathFollowTask.h"
+
+#include "MassZoneGraphFindSmartObjectTarget.generated.h"
+
+/**
+* Computes move target to a smart object based on current location on ZoneGraph.
+*/
+USTRUCT(meta = (DisplayName = "ZG Find Smart Object Target"))
+struct MASSAIBEHAVIOR_API FMassZoneGraphFindSmartObjectTarget : public FMassStateTreeTaskBase
+{
+	GENERATED_BODY()
+
+protected:
+	virtual EStateTreeRunStatus EnterState(FStateTreeExecutionContext& Context, const EStateTreeStateChangeType ChangeType, const FStateTreeTransitionResult& Transition) override;
+	virtual void ExitState(FStateTreeExecutionContext& Context, const EStateTreeStateChangeType ChangeType, const FStateTreeTransitionResult& Transition) override;
+	virtual EStateTreeRunStatus Tick(FStateTreeExecutionContext& Context, const float DeltaTime) override;
+
+	UPROPERTY(meta=(BaseStruct="DataFragment_SmartObjectUser"))
+	FStateTreeExternalItemHandle SmartObjectUserHandle;
+
+	UPROPERTY(meta=(BaseStruct="MassZoneGraphLaneLocationFragment"))
+	FStateTreeExternalItemHandle LocationHandle;
+
+	UPROPERTY(meta=(BaseClass="ZoneGraphAnnotationSubsystem"))
+	FStateTreeExternalItemHandle AnnotationSubsystemHandle;
+
+	FMassZoneGraphTargetLocation TargetLocation;
+
+	UPROPERTY(EditAnywhere, Category = Parameters, meta=(Struct="MassZoneGraphTargetLocation"))
+	FStateTreeResultRef TargetLocationRef;
+};
\ No newline at end of file
diff --git a/Engine/Plugins/AI/MassAI/Source/MassAIBehavior/Public/Tasks/MassZoneGraphPathFollowTask.h b/Engine/Plugins/AI/MassAI/Source/MassAIBehavior/Public/Tasks/MassZoneGraphPathFollowTask.h
new file mode 100644
--- /dev/null
+++ b/Engine/Plugins/AI/MassAI/Source/MassAIBehavior/Public/Tasks/MassZoneGraphPathFollowTask.h
@@ -0,0 +1,108 @@
+// Copyright Epic Games, Inc. All Rights Reserved.
+
+#pragma once
+
+#include "MassAIMovementTypes.h"
+#include "MassStateTreeTypes.h"
+#include "ZoneGraphTypes.h"
+#include "MassZoneGraphPathFollowTask.generated.h"
+
+struct FMassStateTreeExecutionContext;
+
+USTRUCT()
+struct MASSAIBEHAVIOR_API FMassZoneGraphTargetLocation : public FStateTreeResult
+{
+	GENERATED_BODY()
+
+	virtual const UScriptStruct& GetStruct() const override { return *StaticStruct(); }
+
+	void Reset()
+	{
+		LaneHandle.Reset();
+		NextLaneHandle.Reset();
+		NextExitLinkType = EZoneLaneLinkType::None;
+		bMoveReverse = false;
+		TargetDistance = 0.0f;
+		EndOfPathPosition.Reset();
+		AnticipationDistance = 50.0f;
+		EndOfPathIntent = EMassMovementAction::Move;
+	}
+
+	/** Current lane handle. (Could be debug only) */
+	FZoneGraphLaneHandle LaneHandle;
+	
+	/** If valid, the this lane will be set as current lane after the path follow is completed. */
+	FZoneGraphLaneHandle NextLaneHandle;
+	
+	/** Target distance along current lane. */
+	float TargetDistance = 0.0f;
+	
+	/** Optional end of path location. */
+	TOptional<FVector> EndOfPathPosition;
+
+	/** Optional end of path direction, used only if EndOfPathPosition is set. */
+	TOptional<FVector> EndOfPathDirection;
+
+	/** If start or end of path is off-lane, the distance along the lane is pushed forward/back along the lane to make smoother transition. */
+	float AnticipationDistance = 50.0f;
+
+	/** True, if we're moving reverse along the lane. */
+	bool bMoveReverse = false;
+	
+	/** Movement intent at the end of the path */
+	EMassMovementAction EndOfPathIntent = EMassMovementAction::Move;
+	
+	/** How the next lane handle is reached relative to the current lane. */
+	EZoneLaneLinkType NextExitLinkType = EZoneLaneLinkType::None;
+};
+
+/**
+ * Follows a path long the current lane to a specified point.
+ */
+USTRUCT(meta = (DisplayName = "ZG Path Follow"))
+struct MASSAIBEHAVIOR_API FMassZoneGraphPathFollowTask : public FMassStateTreeTaskBase
+{
+	GENERATED_BODY()
+
+protected:
+	virtual EStateTreeRunStatus EnterState(FStateTreeExecutionContext& Context, const EStateTreeStateChangeType ChangeType, const FStateTreeTransitionResult& Transition) override;
+	virtual EStateTreeRunStatus Tick(FStateTreeExecutionContext& Context, const float DeltaTime) override;
+
+	bool RequestPath(FMassStateTreeExecutionContext& Context, const FMassZoneGraphTargetLocation& TargetLocation) const;
+
+	UPROPERTY(meta=(BaseStruct="MassZoneGraphLaneLocationFragment"))
+	FStateTreeExternalItemHandle LocationHandle;
+
+	UPROPERTY(meta=(BaseStruct="MassMoveTargetFragment"))
+	FStateTreeExternalItemHandle MoveTargetHandle;
+
+	UPROPERTY(meta=(BaseStruct="MassMovementConfigFragment"))
+	FStateTreeExternalItemHandle MovementConfigHandle;
+
+	UPROPERTY(meta=(BaseStruct="MassZoneGraphPathRequestFragment"))
+	FStateTreeExternalItemHandle PathRequestHandle;
+
+	UPROPERTY(meta=(BaseStruct="MassZoneGraphShortPathFragment"))
+	FStateTreeExternalItemHandle ShortPathHandle;
+
+	UPROPERTY(meta=(BaseStruct="MassZoneGraphCachedLaneFragment"))
+	FStateTreeExternalItemHandle CachedLaneHandle;
+
+	UPROPERTY(meta=(BaseStruct="DataFragment_AgentRadius"))
+	FStateTreeExternalItemHandle AgentRadiusHandle;
+	
+	UPROPERTY(meta=(BaseClass="ZoneGraphSubsystem"))
+	FStateTreeExternalItemHandle ZoneGraphSubsystemHandle;
+
+	UPROPERTY(meta = (BaseStruct = "MassSimulationLODFragment"))
+	FStateTreeExternalItemHandle SimulationLODHandle;
+
+	UPROPERTY(EditAnywhere, Category = Parameters, meta=(Struct="MassZoneGraphTargetLocation", Bindable))
+	FStateTreeResultRef TargetLocation;
+
+	UPROPERTY(EditAnywhere, Category = Parameters)
+	FMassMovementStyleRef MovementStyle;
+
+	UPROPERTY(EditAnywhere, Category = Parameters)
+	float SpeedScale = 1.0f;
+};
diff --git a/Engine/Plugins/AI/MassAI/Source/MassAIBehavior/Public/Tasks/MassZoneGraphStandTask.h b/Engine/Plugins/AI/MassAI/Source/MassAIBehavior/Public/Tasks/MassZoneGraphStandTask.h
new file mode 100644
--- /dev/null
+++ b/Engine/Plugins/AI/MassAI/Source/MassAIBehavior/Public/Tasks/MassZoneGraphStandTask.h
@@ -0,0 +1,53 @@
+// Copyright Epic Games, Inc. All Rights Reserved.
+
+#pragma once
+
+#include "MassStateTreeTypes.h"
+#include "MassMovementTypes.h"
+#include "MassZoneGraphStandTask.generated.h"
+
+struct FStateTreeExecutionContext;
+
+/**
+ * Stop, and stand on current ZoneGraph location
+ */
+USTRUCT(meta = (DisplayName = "ZG Stand"))
+struct MASSAIBEHAVIOR_API FMassZoneGraphStandTask : public FMassStateTreeTaskBase
+{
+	GENERATED_BODY()
+
+protected:
+	virtual EStateTreeRunStatus EnterState(FStateTreeExecutionContext& Context, const EStateTreeStateChangeType ChangeType, const FStateTreeTransitionResult& Transition) override;
+	virtual EStateTreeRunStatus Tick(FStateTreeExecutionContext& Context, const float DeltaTime) override;
+
+	UPROPERTY(meta=(BaseStruct="MassZoneGraphLaneLocationFragment"))
+	FStateTreeExternalItemHandle LocationHandle;
+
+	UPROPERTY(meta=(BaseStruct="MassMoveTargetFragment"))
+	FStateTreeExternalItemHandle MoveTargetHandle;
+
+	UPROPERTY(meta=(BaseStruct="MassZoneGraphShortPathFragment"))
+	FStateTreeExternalItemHandle ShortPathHandle;
+
+	UPROPERTY(meta=(BaseStruct="MassZoneGraphCachedLaneFragment"))
+	FStateTreeExternalItemHandle CachedLaneHandle;
+
+	UPROPERTY(meta=(BaseClass="ZoneGraphSubsystem"))
+	FStateTreeExternalItemHandle ZoneGraphSubsystemHandle;
+
+	UPROPERTY(meta=(BaseClass="MassSignalSubsystem"))
+	FStateTreeExternalItemHandle MassSignalSubsystemHandle;
+
+	UPROPERTY(meta = (BaseStruct = "MassSimulationLODFragment"))
+	FStateTreeExternalItemHandle SimulationLODHandle;
+
+	UPROPERTY(meta=(BaseStruct="MassMovementConfigFragment"))
+	FStateTreeExternalItemHandle MovementConfigHandle;
+
+	UPROPERTY()
+	float Time = 0.0f;
+
+	/** Delay before the task ends. Default (0 or any negative) will run indefinitely so it requires a transition in the state tree to stop it. */
+	UPROPERTY(EditAnywhere, Category = Parameters, meta = (Bindable))
+	float Duration = 0.0f;
+};
diff --git a/Engine/Plugins/AI/MassAI/Source/MassAIDebug/MassAIDebug.Build.cs b/Engine/Plugins/AI/MassAI/Source/MassAIDebug/MassAIDebug.Build.cs
new file mode 100644
--- /dev/null
+++ b/Engine/Plugins/AI/MassAI/Source/MassAIDebug/MassAIDebug.Build.cs
@@ -0,0 +1,66 @@
+// Copyright Epic Games, Inc. All Rights Reserved.
+
+namespace UnrealBuildTool.Rules
+{
+	public class MassAIDebug : ModuleRules
+	{
+		public MassAIDebug(ReadOnlyTargetRules Target) : base(Target)
+		{
+			PCHUsage = ModuleRules.PCHUsageMode.UseExplicitOrSharedPCHs;
+
+			PublicIncludePaths.AddRange(
+				new string[] {
+					"Runtime/AIModule/Public",
+					ModuleDirectory + "/Public",
+				}
+			);
+
+
+			PrivateIncludePaths.AddRange(
+				new string[] {
+				}
+			);
+
+			PublicDependencyModuleNames.AddRange(
+				new string[] {
+					"Core",
+					"CoreUObject",
+					"Engine",
+					"InputCore",
+					"MassEntity",
+					"NavigationSystem",
+					"StateTreeModule",
+					"MassGameplayDebug",
+					"MassActors",
+					"MassAIBehavior",
+					"MassCommon",
+					"MassMovement",
+					"MassAIMovement",
+					"MassAIReplication",
+					"MassSmartObjects",
+					"MassSpawner",
+					"MassSimulation",
+					"MassRepresentation",
+					"MassLOD",
+					"StructUtils",
+				}
+			);
+
+			if (Target.bBuildEditor == true)
+			{
+				PrivateDependencyModuleNames.Add("UnrealEd");
+				PublicDependencyModuleNames.Add("MassEntityEditor");
+			}
+
+			if (Target.bBuildDeveloperTools || (Target.Configuration != UnrealTargetConfiguration.Shipping && Target.Configuration != UnrealTargetConfiguration.Test))
+			{
+				PrivateDependencyModuleNames.Add("GameplayDebugger");
+				PublicDefinitions.Add("WITH_GAMEPLAY_DEBUGGER=1");
+			}
+			else
+			{
+				PublicDefinitions.Add("WITH_GAMEPLAY_DEBUGGER=0");
+			}
+		}
+	}
+}
