AggregateTicking and Pipe merged into Engine/Plugins/Runtime/MassEntity

diff --git a/Engine/Plugins/AI/MassAI/Source/MassAIDebug/Private/GameplayDebuggerCategory_Mass.cpp b/Engine/Plugins/AI/MassAI/Source/MassAIDebug/Private/GameplayDebuggerCategory_Mass.cpp
new file mode 100644
--- /dev/null
+++ b/Engine/Plugins/AI/MassAI/Source/MassAIDebug/Private/GameplayDebuggerCategory_Mass.cpp
@@ -0,0 +1,670 @@
+// Copyright Epic Games, Inc. All Rights Reserved.
+
+#include "GameplayDebuggerCategory_Mass.h"
+
+#if WITH_GAMEPLAY_DEBUGGER && WITH_MASS_DEBUG
+#include "EntityView.h"
+#include "GameplayDebuggerCategoryReplicator.h"
+#include "GameplayDebuggerPlayerManager.h"
+#include "MassDebuggerSubsystem.h"
+#include "MassActorSubsystem.h"
+#include "GameFramework/PlayerController.h"
+#include "MassAgentComponent.h"
+#include "MassEntitySystem.h"
+#include "MassAIMovementFragments.h"
+#include "MassLookAtFragments.h"
+#include "MassStateTreeFragments.h"
+#include "MassStateTreeExecutionContext.h"
+#include "MassZoneGraphMovementFragments.h"
+#include "Util/ColorConstants.h"
+#include "MassSimulationLOD.h"
+#include "CanvasItem.h"
+
+//----------------------------------------------------------------------//
+//  FGameplayDebuggerCategory_Mass
+//----------------------------------------------------------------------//
+FGameplayDebuggerCategory_Mass::FGameplayDebuggerCategory_Mass()
+{
+	CachedDebugActor = nullptr;
+	bShowOnlyWithDebugActor = false;
+
+	// @todo would be nice to have these saved in per-user settings 
+	bShowArchetypes = false;
+	bShowShapes = false;
+	bShowAgentFragments = false;
+	bPickEntity = false;
+	bShowEntityDetails = false;
+	bShowNearEntityOverview = true;
+	bShowNearEntityAvoidance = false;
+	bShowNearEntityPath = false;
+
+	BindKeyPress(EKeys::A.GetFName(), FGameplayDebuggerInputModifier::Shift, this, &FGameplayDebuggerCategory_Mass::OnToggleArchetypes, EGameplayDebuggerInputMode::Replicated);
+	BindKeyPress(EKeys::S.GetFName(), FGameplayDebuggerInputModifier::Shift, this, &FGameplayDebuggerCategory_Mass::OnToggleShapes, EGameplayDebuggerInputMode::Replicated);
+	BindKeyPress(EKeys::G.GetFName(), FGameplayDebuggerInputModifier::Shift, this, &FGameplayDebuggerCategory_Mass::OnToggleAgentFragments, EGameplayDebuggerInputMode::Replicated);
+	BindKeyPress(EKeys::P.GetFName(), FGameplayDebuggerInputModifier::Shift, this, &FGameplayDebuggerCategory_Mass::OnPickEntity, EGameplayDebuggerInputMode::Replicated);
+	BindKeyPress(EKeys::D.GetFName(), FGameplayDebuggerInputModifier::Shift, this, &FGameplayDebuggerCategory_Mass::OnToggleEntityDetails, EGameplayDebuggerInputMode::Replicated);
+	BindKeyPress(EKeys::O.GetFName(), FGameplayDebuggerInputModifier::Shift, this, &FGameplayDebuggerCategory_Mass::OnToggleNearEntityOverview, EGameplayDebuggerInputMode::Replicated);
+	BindKeyPress(EKeys::V.GetFName(), FGameplayDebuggerInputModifier::Shift, this, &FGameplayDebuggerCategory_Mass::OnToggleNearEntityAvoidance, EGameplayDebuggerInputMode::Replicated);
+	BindKeyPress(EKeys::C.GetFName(), FGameplayDebuggerInputModifier::Shift, this, &FGameplayDebuggerCategory_Mass::OnToggleNearEntityPath, EGameplayDebuggerInputMode::Replicated);
+}
+
+void FGameplayDebuggerCategory_Mass::SetCachedEntity(FLWEntity Entity, UMassDebuggerSubsystem& Debugger)
+{
+	CachedEntity = Entity;
+	Debugger.SetSelectedEntity(Entity);
+}
+
+void FGameplayDebuggerCategory_Mass::PickEntity(const APlayerController& OwnerPC, UWorld& World, UMassDebuggerSubsystem& Debugger)
+{
+	FVector ViewLocation = FVector::ZeroVector;
+	FVector ViewDirection = FVector::ForwardVector;
+	AGameplayDebuggerPlayerManager::GetViewPoint(OwnerPC, ViewLocation, ViewDirection);
+
+	// Reusing similar algorithm as UGameplayDebuggerLocalController for now 
+	const float MaxScanDistance = 25000.0f;
+	const float MinViewDirDot = 0.707f; // 45 degrees
+	
+	float BestScore = MinViewDirDot;
+	FLWEntity BestEntity;
+	FVector BestLocation = FVector::ZeroVector;
+
+	TConstArrayView<FLWEntity> Entities = Debugger.GetEntities();
+	TConstArrayView<FVector> Locations = Debugger.GetLocations();
+	checkf(Entities.Num() == Locations.Num(), TEXT("Both Entities and Locations lists are expected to be of the same size: %d vs %d"), Entities.Num(), Locations.Num());
+
+	for (int32 i = 0; i < Locations.Num(); ++i)
+	{
+		const FVector& Location = Locations[i];
+		FVector DirToEntity = (Location - ViewLocation);
+		const float DistToEntity = DirToEntity.Size();
+		if (DistToEntity > MaxScanDistance)
+		{
+			continue;
+		}
+
+		DirToEntity = (FMath::IsNearlyZero(DistToEntity)) ? ViewDirection : DirToEntity /= DistToEntity;
+		const float ViewDot = FVector::DotProduct(ViewDirection, DirToEntity);
+		if (ViewDot > BestScore)
+		{
+			BestScore = ViewDot;
+			BestEntity = Entities[i];
+			BestLocation = Location;
+		}
+	}
+
+	AActor* BestActor = nullptr;
+	if (BestEntity.IsSet())
+	{
+		// Use this new entity
+		SetCachedEntity(BestEntity, Debugger);
+		UMassActorSubsystem* ActorManager = World.GetSubsystem<UMassActorSubsystem>();
+		if (ActorManager != nullptr)
+		{
+			BestActor = ActorManager->GetActorFromHandle(FMassHandle(CachedEntity));
+			CachedDebugActor = BestActor;
+		}
+	}
+	GetReplicator()->SetDebugActor(BestActor);
+}
+
+TSharedRef<FGameplayDebuggerCategory> FGameplayDebuggerCategory_Mass::MakeInstance()
+{
+	return MakeShareable(new FGameplayDebuggerCategory_Mass());
+}
+
+void FGameplayDebuggerCategory_Mass::CollectData(APlayerController* OwnerPC, AActor* DebugActor)
+{
+	UWorld* World = GetDataWorld(OwnerPC, DebugActor);
+	check(World);
+	UMassDebuggerSubsystem* Debugger = World->GetSubsystem<UMassDebuggerSubsystem>();
+	if (Debugger == nullptr)
+	{
+		AddTextLine(FString::Printf(TEXT("{Red}MassDebuggerSubsystem instance is missing")));
+		return;
+	}
+	Debugger->SetCollectingData();
+
+	// Ideally we would have a way to register in the main picking flow but that would require more changes to
+	// also support client-server picking. For now, we handle explicit mass picking requests on the authority
+	if (bPickEntity)
+	{
+		PickEntity(*OwnerPC, *World, *Debugger);
+		bPickEntity = false;
+	}
+
+	auto GetEntityFromActorFunc = [](const AActor& Actor, const UMassAgentComponent** OutMassAgentComponent = nullptr)
+	{
+		FLWEntity LWEntity;
+		UMassAgentComponent* AgentComp = Actor.FindComponentByClass<UMassAgentComponent>();
+		if (AgentComp)
+		{
+			LWEntity = AgentComp->GetEntityHandle().GetLWEntity();
+			if (OutMassAgentComponent)
+			{
+				*OutMassAgentComponent = AgentComp;
+			}
+		}
+		else
+		{
+			UMassActorSubsystem* ActorManager = UWorld::GetSubsystem<UMassActorSubsystem>(Actor.GetWorld());
+			if (ActorManager != nullptr)
+			{
+				LWEntity = ActorManager->GetHandleFromActor(&Actor).GetLWEntity();
+			}
+		}
+		return LWEntity;
+	};
+
+	if (CachedDebugActor != DebugActor)
+	{
+		CachedDebugActor = DebugActor;
+		if (DebugActor != nullptr)
+		{
+			SetCachedEntity(GetEntityFromActorFunc(*DebugActor), *Debugger);
+		}
+	}
+
+	UPipeEntitySubsystem* EntitySystem = UPipeEntitySubsystem::GetCurrent(World);
+	if (EntitySystem)
+	{
+		AddTextLine(FString::Printf(TEXT("{Green}Entities count active{grey}/all: {white}%d{grey}/%d"), EntitySystem->DebugGetEntityCount(), EntitySystem->DebugGetEntityCount()));
+		AddTextLine(FString::Printf(TEXT("{Green}Registered Archetypes count: {white}%d {green}data ver: {white}%d"), EntitySystem->DebugGetArchetypesCount(), EntitySystem->GetArchetypeDataVersion()));
+
+		if (bShowArchetypes)
+		{
+			FStringOutputDevice Ar;
+			Ar.SetAutoEmitLineTerminator(true);
+			EntitySystem->DebugPrintArchetypes(Ar);
+
+			AddTextLine(Ar);
+		}
+	}
+	else
+	{
+		AddTextLine(FString::Printf(TEXT("{Red}EntitySystem instance is missing")));
+	}
+
+	if (CachedEntity.IsSet())
+	{
+		AddTextLine(Debugger->GetSelectedEntityInfo());
+	}
+
+	//@todo could shave off some perf cost if UMassDebuggerSubsystem used FGameplayDebuggerShape directly
+	if (bShowShapes)
+	{
+		const TArray<UMassDebuggerSubsystem::FShapeDesc>* Shapes = Debugger->GetShapes();
+		check(Shapes);
+		// EMassEntityDebugShape::Box
+		for (const UMassDebuggerSubsystem::FShapeDesc& Desc : Shapes[uint8(EMassEntityDebugShape::Box)])
+		{
+			AddShape(FGameplayDebuggerShape::MakeBox(Desc.Location, FVector(Desc.Size), FColor::Blue));
+		}
+		// EMassEntityDebugShape::Cone
+		// note that we're modifying the Size here because MakeCone is using the third param as Cone's "height", while all mass debugger shapes are created with agent radius
+		// FGameplayDebuggerShape::Draw is using 0.25 rad for cone angle, so that's what we'll use here
+		static const float Tan025Rad = FMath::Tan(0.25f);
+		for (const UMassDebuggerSubsystem::FShapeDesc& Desc : Shapes[uint8(EMassEntityDebugShape::Cone)])
+		{
+			AddShape(FGameplayDebuggerShape::MakeCone(Desc.Location, FVector::UpVector, Desc.Size / Tan025Rad, FColor::Orange));
+		}
+		// EMassEntityDebugShape::Cylinder
+		for (const UMassDebuggerSubsystem::FShapeDesc& Desc : Shapes[uint8(EMassEntityDebugShape::Cylinder)])
+		{
+			AddShape(FGameplayDebuggerShape::MakeCylinder(Desc.Location, Desc.Size, Desc.Size * 2, FColor::Yellow));
+		}
+		// EMassEntityDebugShape::Capsule
+		for (const UMassDebuggerSubsystem::FShapeDesc& Desc : Shapes[uint8(EMassEntityDebugShape::Capsule)])
+		{
+			AddShape(FGameplayDebuggerShape::MakeCapsule(Desc.Location, Desc.Size, Desc.Size * 2, FColor::Green));
+		}
+	}
+
+	if (bShowAgentFragments)
+	{
+		const UMassAgentComponent* AgentComp = nullptr;
+		if (!CachedEntity.IsSet() && DebugActor != nullptr)
+		{
+			SetCachedEntity(GetEntityFromActorFunc(*DebugActor, &AgentComp), *Debugger);
+		}
+
+		if (CachedEntity.IsSet() && EntitySystem)
+		{
+			// CachedEntity can become invalid if the entity "dies" or in editor mode when related actor gets moved 
+			// (which causes the MassAgentComponent destruction and recreation).
+			if (EntitySystem->IsEntityActive(CachedEntity))
+			{
+				AddTextLine(FString::Printf(TEXT("{Green}Entity: {White}%s"), *CachedEntity.DebugGetDescription()));
+				AddTextLine(FString::Printf(TEXT("{Green}Type: {White}%s"), (AgentComp == nullptr) ? TEXT("N/A") : AgentComp->IsPuppet() ? TEXT("PUPPET") : TEXT("AGENT")));
+
+				if (bShowEntityDetails)
+				{
+					FStringOutputDevice FragmentsDesc;
+					FragmentsDesc.SetAutoEmitLineTerminator(true);
+					const TCHAR* PrefixToRemove = TEXT("DataFragment_");
+					EntitySystem->DebugPrintEntity(CachedEntity, FragmentsDesc, PrefixToRemove);
+					AddTextLine(FString::Printf(TEXT("{Green}Fragments:\n{White}%s"), *FragmentsDesc));
+				}
+				else
+				{
+					const FArchetypeHandle Archetype = EntitySystem->GetArchetypeForEntity(CachedEntity);
+					TArray<FName> ComponentNames;
+					TArray<FName> TagNames;
+					EntitySystem->DebugGetArchetypeStrings(Archetype, ComponentNames, TagNames);
+
+					FString Tags;
+					for (const FName Name : TagNames)
+					{
+						Tags += FString::Printf(TEXT("%s, "), *Name.ToString());
+					}
+					AddTextLine(FString::Printf(TEXT("{Green}Tags:\n{White}%s"), *Tags));
+
+					AddTextLine(FString::Printf(TEXT("{Green}Fragments:{White}")));
+					constexpr int ColumnsCount = 2;
+					int i = 0;
+					while (i + ColumnsCount < ComponentNames.Num())
+					{
+						AddTextLine(FString::Printf(TEXT("%-42s, %-42s"), *ComponentNames[i].ToString(), *ComponentNames[i + 1].ToString()));
+						i += ColumnsCount;
+					}
+					if (i < ComponentNames.Num())
+					{
+						AddTextLine(FString::Printf(TEXT("%s"), *ComponentNames[i].ToString()));
+					}
+
+				}
+
+				FDataFragment_Transform& TransformFragment = EntitySystem->GetComponentDataChecked<FDataFragment_Transform>(CachedEntity);
+				const float CapsuleRadius = 50.f;
+				AddShape(FGameplayDebuggerShape::MakeCapsule(TransformFragment.GetTransform().GetLocation() + 2.f * CapsuleRadius * FVector::UpVector, CapsuleRadius, CapsuleRadius * 2.f, FColor::Orange));
+			}
+			else
+			{
+				CachedEntity.Reset();
+			}
+		}
+		else
+		{
+			AddTextLine(FString::Printf(TEXT("{Green}Entity: {Red}INACTIVE")));
+		}
+	}
+
+	NearEntityDescriptions.Reset();
+	if (bShowNearEntityOverview)
+	{
+		FVector ViewLocation = FVector::ZeroVector;
+		FVector ViewDirection = FVector::ForwardVector;
+		AGameplayDebuggerPlayerManager::GetViewPoint(*OwnerPC, ViewLocation, ViewDirection);
+
+		FLWComponentQuery EntityQuery;
+		EntityQuery.AddRequirement<FMassStateTreeFragment>(ELWComponentAccess::ReadOnly);
+		EntityQuery.AddRequirement<FDataFragment_Transform>(ELWComponentAccess::ReadOnly);
+		EntityQuery.AddRequirement<FDataFragment_AgentRadius>(ELWComponentAccess::ReadOnly);
+		EntityQuery.AddRequirement<FMassSteeringFragment>(ELWComponentAccess::ReadOnly);
+		EntityQuery.AddRequirement<FMassSteeringGhostFragment>(ELWComponentAccess::ReadOnly);
+		EntityQuery.AddRequirement<FMassVelocityFragment>(ELWComponentAccess::ReadOnly);
+		EntityQuery.AddRequirement<FMassMoveTargetFragment>(ELWComponentAccess::ReadOnly);
+		EntityQuery.AddRequirement<FMassLookAtFragment>(ELWComponentAccess::ReadOnly);
+		EntityQuery.AddRequirement<FMassLookAtTrajectoryFragment>(ELWComponentAccess::ReadOnly);
+		EntityQuery.AddRequirement<FMassSimulationLODFragment>(ELWComponentAccess::ReadOnly);
+		EntityQuery.AddRequirement<FMassZoneGraphShortPathFragment>(ELWComponentAccess::ReadOnly);
+
+		const float CurrentTime = World->GetTimeSeconds();
+		
+		UMassStateTreeSubsystem* MassStateTreeSubsystem = World->GetSubsystem<UMassStateTreeSubsystem>();
+		if (MassStateTreeSubsystem != nullptr)
+		{
+			FLWComponentSystemExecutionContext Context(0.0f);
+		
+			EntityQuery.ForEachEntityChunk(*EntitySystem, Context, [this, Debugger, MassStateTreeSubsystem, EntitySystem, OwnerPC, ViewLocation, ViewDirection, CurrentTime](FLWComponentSystemExecutionContext& Context)
+			{
+				const int32 NumEntities = Context.GetEntitiesNum();
+				const TConstArrayView<FMassStateTreeFragment> StateTreeList = Context.GetComponentView<FMassStateTreeFragment>();
+				const TConstArrayView<FDataFragment_Transform> TransformList = Context.GetComponentView<FDataFragment_Transform>();
+				const TConstArrayView<FDataFragment_AgentRadius> RadiusList = Context.GetComponentView<FDataFragment_AgentRadius>();
+				const TConstArrayView<FMassSteeringFragment> SteeringList = Context.GetComponentView<FMassSteeringFragment>();
+				const TConstArrayView<FMassSteeringGhostFragment> GhostList = Context.GetComponentView<FMassSteeringGhostFragment>();
+				const TConstArrayView<FMassVelocityFragment> VelocityList = Context.GetComponentView<FMassVelocityFragment>();
+				const TConstArrayView<FMassMoveTargetFragment> MoveTargetList = Context.GetComponentView<FMassMoveTargetFragment>();
+				const TConstArrayView<FMassLookAtFragment> LookAtList = Context.GetComponentView<FMassLookAtFragment>();
+				const TConstArrayView<FMassLookAtTrajectoryFragment> LookAtTrajectoryList = Context.GetComponentView<FMassLookAtTrajectoryFragment>();
+				const TConstArrayView<FMassSimulationLODFragment> SimLODList = Context.GetComponentView<FMassSimulationLODFragment>();
+				const TConstArrayView<FMassZoneGraphShortPathFragment> ShortPathList = Context.GetComponentView<FMassZoneGraphShortPathFragment>();
+
+				constexpr float MaxViewDistance = 25000.0f;
+				constexpr float MinViewDirDot = 0.707f; // 45 degrees
+
+				const UStateTree* StateTree = MassStateTreeSubsystem->GetRegisteredStateTreeAsset(StateTreeList[0].StateTreeHandle);
+
+				for (int32 EntityIndex = 0; EntityIndex < NumEntities; ++EntityIndex)
+				{
+					const FDataFragment_Transform& Transform = TransformList[EntityIndex];
+					const FVector EntityLocation = Transform.GetTransform().GetLocation();
+					
+					// Cull entities
+					const FVector DirToEntity = EntityLocation - ViewLocation;
+					const float DistanceToEntitySq = DirToEntity.SquaredLength();
+					if (DistanceToEntitySq > FMath::Square(MaxViewDistance))
+					{
+						continue;
+					}
+					const float ViewDot = FVector::DotProduct(DirToEntity.GetSafeNormal(), ViewDirection);
+					if (ViewDot < MinViewDirDot)
+					{
+						continue;
+					}
+
+					const FDataFragment_AgentRadius& Radius = RadiusList[EntityIndex];
+					const FMassSteeringFragment& Steering = SteeringList[EntityIndex];
+					const FMassSteeringGhostFragment& Ghost = GhostList[EntityIndex];
+					const FMassVelocityFragment& Velocity = VelocityList[EntityIndex];
+					const FMassMoveTargetFragment& MoveTarget = MoveTargetList[EntityIndex];
+					const FMassLookAtFragment& LookAt = LookAtList[EntityIndex];
+					const FMassSimulationLODFragment& SimLOD = SimLODList[EntityIndex];
+					const FMassZoneGraphShortPathFragment& ShortPath = ShortPathList[EntityIndex];
+
+					const FVector EntityForward = Transform.GetTransform().GetRotation().GetForwardVector();
+
+					const float Height = 180.0f; // @todo: add height to agent.
+					const float EyeHeight = 160.0f; // @todo: add eye height to agent.
+					
+					// Draw entity position and orientation.
+					FVector BasePos = EntityLocation + FVector(0.0f ,0.0f ,25.0f );
+
+					AddShape(FGameplayDebuggerShape::MakeCircle(BasePos, FVector::UpVector, Radius.Radius, FColor::White));
+					AddShape(FGameplayDebuggerShape::MakeSegment(BasePos, BasePos + EntityForward * Radius.Radius * 1.25f, FColor::White));
+
+					// Velocity and steering target
+					BasePos += FVector(0.0f ,0.0f ,5.0f );
+					AddShape(FGameplayDebuggerShape::MakeArrow(BasePos, BasePos + Velocity.Value, 10.0f, 2.0f, FColor::Yellow));
+					BasePos += FVector(0.0f ,0.0f ,5.0f );
+					AddShape(FGameplayDebuggerShape::MakeArrow(BasePos, BasePos + Steering.DesiredVelocity, 10.0f, 1.0f, FColorList::Pink));
+
+					// Move target
+					const FVector MoveBasePos = MoveTarget.Center + FVector(0,0,5);
+					AddShape(FGameplayDebuggerShape::MakeArrow(MoveBasePos - MoveTarget.Forward * Radius.Radius, MoveBasePos + MoveTarget.Forward * Radius.Radius, 10.0f, 2.0f, FColorList::MediumVioletRed));
+
+					// Look at
+					constexpr float LookArrowLength = 100.0f;
+					BasePos = EntityLocation + FVector(0,0,EyeHeight);
+					const FVector WorldLookDirection = Transform.GetTransform().TransformVector(LookAt.Direction);
+					bool bLookArrowDrawn = false;
+					if (LookAt.LookAtMode == EMassLookAtMode::LookAtEntity && EntitySystem->IsEntityValid(LookAt.TrackedEntity))
+					{
+						if (const FDataFragment_Transform* TargetTransform = EntitySystem->GetComponentDataPtr<FDataFragment_Transform>(LookAt.TrackedEntity))
+						{
+							FVector TargetPosition = TargetTransform->GetTransform().GetLocation();
+							TargetPosition.Z = BasePos.Z;
+							AddShape(FGameplayDebuggerShape::MakeCircle(TargetPosition, FVector::UpVector, Radius.Radius, FColor::Red));
+							
+							const float TargetDistance = FMath::Max(LookArrowLength, FVector::DotProduct(WorldLookDirection, TargetPosition - BasePos));
+							AddShape(FGameplayDebuggerShape::MakeSegment(BasePos, BasePos + WorldLookDirection * TargetDistance, FColorList::LightGrey));
+							bLookArrowDrawn = true;
+						}
+					}
+					
+					if (LookAt.bRandomGazeEntities && EntitySystem->IsEntityValid(LookAt.GazeTrackedEntity))
+					{
+						if (const FDataFragment_Transform* TargetTransform = EntitySystem->GetComponentDataPtr<FDataFragment_Transform>(LookAt.GazeTrackedEntity))
+						{
+							FVector TargetPosition = TargetTransform->GetTransform().GetLocation();
+							TargetPosition.Z = BasePos.Z;
+							AddShape(FGameplayDebuggerShape::MakeCircle(TargetPosition, FVector::UpVector, Radius.Radius, FColor::Turquoise));
+						}
+					}
+
+					if (!bLookArrowDrawn)
+					{
+						AddShape(FGameplayDebuggerShape::MakeArrow(BasePos, BasePos + WorldLookDirection * LookArrowLength, 10.0f, 1.0f, FColor::Turquoise));
+					}
+
+
+					// Path
+					if (bShowNearEntityPath)
+					{
+						const FVector ZOffset = FVector(0.0f , 0.0f , 25.0f );
+						for (uint8 PointIndex = 0; PointIndex < ShortPath.NumPoints - 1; PointIndex++)
+						{
+							const FMassZoneGraphPathPoint& CurrPoint = ShortPath.Points[PointIndex];
+							const FMassZoneGraphPathPoint& NextPoint = ShortPath.Points[PointIndex + 1];
+							AddShape(FGameplayDebuggerShape::MakeSegment(CurrPoint.Position + ZOffset, NextPoint.Position + ZOffset, 3.0f, FColorList::Grey));
+						}
+					
+						for (uint8 PointIndex = 0; PointIndex < ShortPath.NumPoints; PointIndex++)
+						{
+							const FMassZoneGraphPathPoint& CurrPoint = ShortPath.Points[PointIndex];
+							const FVector CurrBase = CurrPoint.Position + ZOffset;
+							// Lane tangents
+							AddShape(FGameplayDebuggerShape::MakeSegment(CurrBase, CurrBase + CurrPoint.Tangent.GetVector() * 50.0f, 1.0f, FColorList::LightGrey));
+						}
+					}
+					
+					if (bShowNearEntityAvoidance)
+					{
+						// Standing avoidance.
+						if (Ghost.IsValid(MoveTarget.GetCurrentActionID()))
+						{
+							FVector GhostBasePos = Ghost.Location + FVector(0.0f ,0.0f ,25.0f );
+							AddShape(FGameplayDebuggerShape::MakeCircle(GhostBasePos, FVector::UpVector, Radius.Radius, FColorList::LightGrey));
+							GhostBasePos += FVector(0,0,5);
+							AddShape(FGameplayDebuggerShape::MakeArrow(GhostBasePos, GhostBasePos + Ghost.Velocity, 10.0f, 2.0f, FColorList::LightGrey));
+
+							const FVector GhostTargetBasePos = Ghost.SteerTarget + FVector(0.0f ,0.0f ,25.0f );
+							AddShape(FGameplayDebuggerShape::MakeCircle(GhostTargetBasePos, FVector::UpVector, Radius.Radius * 0.75f, FColorList::Orange));
+						}
+					}
+					
+					// Status
+					if (DistanceToEntitySq < FMath::Square(MaxViewDistance * 0.5f))
+					{
+						FString Status;
+
+						// Entity name
+						FLWEntity Entity = Context.GetEntity(EntityIndex);
+						Status += TEXT("{orange}");
+						Status += Entity.DebugGetDescription();
+						Status += TEXT(" {white}LOD ");
+						switch (SimLOD.LOD) {
+						case EMassLOD::High:
+							Status += TEXT("High");
+							break;
+						case EMassLOD::Medium:
+							Status += TEXT("Med");
+							break;
+						case EMassLOD::Low:
+							Status += TEXT("Low");
+							break;
+						case EMassLOD::Off:
+							Status += TEXT("Off");
+							break;
+						default:
+							Status += TEXT("?");
+							break;
+						}
+						Status += TEXT("\n");
+						
+						// Current StateTree task
+						if (StateTree != nullptr)
+						{
+							FMassStateTreeExecutionContext StateTreeContext(*EntitySystem, Context);
+							StateTreeContext.Init(*OwnerPC, *StateTree, EStateTreeStorage::External);
+							StateTreeContext.SetEntity(Entity);
+							
+							FStructView Storage = EntitySystem->GetComponentDataStruct(Entity, StateTree->GetRuntimeStorageStruct());
+							
+							Status += StateTreeContext.GetActiveStateName(Storage);
+							Status += TEXT("\n");
+						}
+
+						// Movement info
+						Status += FString::Printf(TEXT("{yellow}%s/%03d {white}%.1f cm/s\n"),
+							*UEnum::GetDisplayValueAsText(MoveTarget.GetCurrentAction()).ToString(), MoveTarget.GetCurrentActionID(), Velocity.Value.Length());
+						Status += FString::Printf(TEXT("{pink}-> %s {white}Dist: %.1f\n"),
+							*UEnum::GetDisplayValueAsText(MoveTarget.IntentAtGoal).ToString(), MoveTarget.DistanceToGoal);
+
+						// Look
+						const float RemainingTime = LookAt.GazeDuration - (CurrentTime - LookAt.GazeStartTime);
+						Status += FString::Printf(TEXT("{turquoise}%s/%s {lightgrey}%.1f\n"),
+							*UEnum::GetDisplayValueAsText(LookAt.LookAtMode).ToString(), *UEnum::GetDisplayValueAsText(LookAt.RandomGazeMode).ToString(), RemainingTime);
+						
+						if (!Status.IsEmpty())
+						{
+							BasePos += FVector(0,0,50);
+							constexpr float ViewWeight = 0.6f; // Higher the number the more the view angle affects the score.
+							const float ViewScale = 1.f - (ViewDot / MinViewDirDot); // Zero at center of screen
+							NearEntityDescriptions.Emplace(DistanceToEntitySq * ((1.0f - ViewWeight) + ViewScale * ViewWeight), BasePos, Status);
+						}
+					}
+				}
+			});
+		}
+
+		if (bShowNearEntityAvoidance)
+		{
+			FLWComponentQuery EntityColliderQuery;
+			EntityColliderQuery.AddRequirement<FMassAvoidanceColliderFragment>(ELWComponentAccess::ReadOnly);
+			EntityColliderQuery.AddRequirement<FDataFragment_Transform>(ELWComponentAccess::ReadOnly);
+			FLWComponentSystemExecutionContext Context(0.f);
+			EntityColliderQuery.ForEachEntityChunk(*EntitySystem, Context, [this, Debugger, EntitySystem, OwnerPC, ViewLocation, ViewDirection](FLWComponentSystemExecutionContext& Context)
+			{
+				const int32 NumEntities = Context.GetEntitiesNum();
+				const TConstArrayView<FDataFragment_Transform> TransformList = Context.GetComponentView<FDataFragment_Transform>();
+				const TConstArrayView<FMassAvoidanceColliderFragment> CollidersList = Context.GetComponentView<FMassAvoidanceColliderFragment>();
+	
+				constexpr float MaxViewDistance = 25000.0f;
+				constexpr float MinViewDirDot = 0.707f; // 45 degrees
+				
+				for (int32 EntityIndex = 0; EntityIndex < NumEntities; ++EntityIndex)
+				{
+					const FDataFragment_Transform& Transform = TransformList[EntityIndex];
+					const FVector EntityLocation = Transform.GetTransform().GetLocation();
+					const FVector EntityForward = Transform.GetTransform().GetRotation().GetForwardVector();
+					
+					FVector BasePos = EntityLocation + FVector(0.0f ,0.0f ,25.0f );
+
+					// Cull entities
+					const FVector DirToEntity = EntityLocation - ViewLocation;
+					const float DistanceToEntitySq = DirToEntity.SquaredLength();
+					if (DistanceToEntitySq > FMath::Square(MaxViewDistance))
+					{
+						continue;
+					}
+					const float ViewDot = FVector::DotProduct(DirToEntity.GetSafeNormal(), ViewDirection);
+					if (ViewDot < MinViewDirDot)
+					{
+						continue;
+					}
+					
+					// Display colliders
+					const FMassAvoidanceColliderFragment& Collider = CollidersList[EntityIndex];
+					if (Collider.Type == EMassColliderType::Circle)
+					{
+						AddShape(FGameplayDebuggerShape::MakeCircle(BasePos, FVector::UpVector, Collider.GetCircleCollider().Radius, FColor::Blue));
+					}
+					else if (Collider.Type == EMassColliderType::Pill)
+					{
+						const FMassPillCollider& Pill = Collider.GetPillCollider();
+						AddShape(FGameplayDebuggerShape::MakeCircle(BasePos + Pill.HalfLength * EntityForward, FVector::UpVector, Pill.Radius, FColor::Blue));
+						AddShape(FGameplayDebuggerShape::MakeCircle(BasePos - Pill.HalfLength * EntityForward, FVector::UpVector, Pill.Radius, FColor::Blue));
+					}
+				}
+			});
+		}
+		
+		// Cap labels to closest ones.
+		NearEntityDescriptions.Sort([](const FEntityDescription& LHS, const FEntityDescription& RHS){ return LHS.Score < RHS.Score; });
+		constexpr int32 MaxLabels = 15;
+		if (NearEntityDescriptions.Num() > MaxLabels)
+		{
+			NearEntityDescriptions.RemoveAt(MaxLabels, NearEntityDescriptions.Num() - MaxLabels);
+		}
+	}
+	
+}
+
+void FGameplayDebuggerCategory_Mass::DrawData(APlayerController* OwnerPC, FGameplayDebuggerCanvasContext& CanvasContext)
+{
+	CanvasContext.Printf(TEXT("\n[{yellow}%s{white}] %s Archetypes"), *GetInputHandlerDescription(0), bShowArchetypes ? TEXT("Hide") : TEXT("Show"));
+	CanvasContext.Printf(TEXT("[{yellow}%s{white}] %s Shapes"), *GetInputHandlerDescription(1), bShowShapes ? TEXT("Hide") : TEXT("Show"));
+	CanvasContext.Printf(TEXT("[{yellow}%s{white}] %s Agent Fragments"), *GetInputHandlerDescription(2), bShowAgentFragments ? TEXT("Hide") : TEXT("Show"));
+	CanvasContext.Printf(TEXT("[{yellow}%s{white}] Pick Entity"), *GetInputHandlerDescription(3));
+	CanvasContext.Printf(TEXT("[{yellow}%s{white}] %s Entity details"), *GetInputHandlerDescription(4), bShowEntityDetails ? TEXT("Hide") : TEXT("Show"));
+	CanvasContext.Printf(TEXT("[{yellow}%s{white}] %s Entity overview"), *GetInputHandlerDescription(5), bShowNearEntityOverview ? TEXT("Hide") : TEXT("Show"));
+	CanvasContext.Printf(TEXT("[{yellow}%s{white}] %s Entity avoidance"), *GetInputHandlerDescription(6), bShowNearEntityAvoidance ? TEXT("Hide") : TEXT("Show"));
+	CanvasContext.Printf(TEXT("[{yellow}%s{white}] %s Entity path"), *GetInputHandlerDescription(7), bShowNearEntityPath ? TEXT("Hide") : TEXT("Show"));
+
+	struct FEntityLayoutRect
+	{
+		FVector2D Min = FVector2D::ZeroVector;
+		FVector2D Max = FVector2D::ZeroVector;
+		int32 Index = INDEX_NONE;
+		float Alpha = 1.0f;
+	};
+
+	TArray<FEntityLayoutRect> Layout;
+
+	// The loop below is O(N^2), make sure to keep the N small.
+	constexpr int32 MaxDesc = 20;
+	const int32 NumDescs = FMath::Min(NearEntityDescriptions.Num(), MaxDesc);
+	
+	// The labels are assumed to have been ordered in order of importance (i.e. front to back).
+	for (int32 Index = 0; Index < NumDescs; Index++)
+	{
+		const FEntityDescription& Desc = NearEntityDescriptions[Index];
+		if (Desc.Description.Len() && CanvasContext.IsLocationVisible(Desc.Location))
+		{
+			float SizeX = 0, SizeY = 0;
+			const FVector2D ScreenLocation = CanvasContext.ProjectLocation(Desc.Location);
+			CanvasContext.MeasureString(Desc.Description, SizeX, SizeY);
+			
+			FEntityLayoutRect Rect;
+			Rect.Min = ScreenLocation + FVector2D(0, -SizeY * 0.5f);
+			Rect.Max = Rect.Min + FVector2D(SizeX, SizeY);
+			Rect.Index = Index;
+			Rect.Alpha = 0.0f;
+
+			// Calculate transparency based on how much more important rects are overlapping the new rect.
+			const float Area = FMath::Max(0.0f, Rect.Max.X - Rect.Min.X) * FMath::Max(0.0f, Rect.Max.Y - Rect.Min.Y);
+			const float InvArea = Area > KINDA_SMALL_NUMBER ? 1.0f / Area : 0.0f;
+			float Coverage = 0.0;
+
+			for (const FEntityLayoutRect& Other : Layout)
+			{
+				// Calculate rect intersection
+				const float MinX = FMath::Max(Rect.Min.X, Other.Min.X);
+				const float MinY = FMath::Max(Rect.Min.Y, Other.Min.Y);
+				const float MaxX = FMath::Min(Rect.Max.X, Other.Max.X);
+				const float MaxY = FMath::Min(Rect.Max.Y, Other.Max.Y);
+
+				// return zero area if not overlapping
+				const float IntersectingArea = FMath::Max(0.0f, MaxX - MinX) * FMath::Max(0.0f, MaxY - MinY);
+				Coverage += (IntersectingArea * InvArea) * Other.Alpha;
+			}
+
+			Rect.Alpha = FMath::Square(1.0f - FMath::Min(Coverage, 1.0f));
+			
+			if (Rect.Alpha > KINDA_SMALL_NUMBER)
+			{
+				Layout.Add(Rect);
+			}
+		}
+	}
+
+	// Render back to front so that the most important item renders at top.
+	const FVector2D Padding(5, 5);
+	for (int32 Index = Layout.Num() - 1; Index >= 0; Index--)
+	{
+		const FEntityLayoutRect& Rect = Layout[Index];
+		const FEntityDescription& Desc = NearEntityDescriptions[Rect.Index];
+
+		const FVector2D BackgroundPosition(Rect.Min - Padding);
+		FCanvasTileItem Background(Rect.Min - Padding, Rect.Max - Rect.Min + Padding * 2.0f, FLinearColor(0.0f, 0.0f, 0.0f, 0.35f * Rect.Alpha));
+		Background.BlendMode = SE_BLEND_TranslucentAlphaOnly;
+		CanvasContext.DrawItem(Background, BackgroundPosition.X, BackgroundPosition.Y);
+		
+		CanvasContext.PrintAt(Rect.Min.X, Rect.Min.Y, FColor::White, Rect.Alpha, Desc.Description);
+	}
+
+	FGameplayDebuggerCategory::DrawData(OwnerPC, CanvasContext);
+}
+#endif // WITH_GAMEPLAY_DEBUGGER && WITH_MASS_DEBUG
+
diff --git a/Engine/Plugins/AI/MassAI/Source/MassAIDebug/Private/MassDebugStateTreeProcessor.cpp b/Engine/Plugins/AI/MassAI/Source/MassAIDebug/Private/MassDebugStateTreeProcessor.cpp
new file mode 100644
--- /dev/null
+++ b/Engine/Plugins/AI/MassAI/Source/MassAIDebug/Private/MassDebugStateTreeProcessor.cpp
@@ -0,0 +1,103 @@
+// Copyright Epic Games, Inc. All Rights Reserved.
+
+#include "MassDebugStateTreeProcessor.h"
+#include "MassDebuggerSubsystem.h"
+#include "MassStateTreeExecutionContext.h"
+#include "MassStateTreeFragments.h"
+#include "MassCommonFragments.h"
+#include "VisualLogger/VisualLogger.h"
+
+//----------------------------------------------------------------------//
+// UMassDebugStateTreeProcessor
+//----------------------------------------------------------------------//
+UMassDebugStateTreeProcessor::UMassDebugStateTreeProcessor()
+{
+	ExecutionOrder.ExecuteInGroup = UE::Mass::ProcessorGroupNames::Behavior;
+	ExecutionOrder.ExecuteAfter.Add(TEXT("MassStateTreeProcessor"));
+}
+
+void UMassDebugStateTreeProcessor::ConfigureQueries()
+{
+	EntityQuery.AddRequirement<FMassStateTreeFragment>(ELWComponentAccess::ReadOnly);
+	EntityQuery.AddRequirement<FDataFragment_Transform>(ELWComponentAccess::ReadOnly);
+}
+
+void UMassDebugStateTreeProcessor::Execute(UEntitySubsystem& EntitySubsystem, FLWComponentSystemExecutionContext& Context)
+{
+#if WITH_MASS_DEBUG
+	UWorld* World = GetWorld();
+	if (World == nullptr)
+	{
+		return;
+	}
+	
+	UMassDebuggerSubsystem* Debugger = World->GetSubsystem<UMassDebuggerSubsystem>();
+	if (Debugger == nullptr)
+	{
+		return;
+	}
+
+	UMassStateTreeSubsystem* MassStateTreeSubsystem = World->GetSubsystem<UMassStateTreeSubsystem>();
+	if (MassStateTreeSubsystem == nullptr)
+	{
+		return;
+	}
+
+	if (!Debugger->GetSelectedEntity().IsSet() && !UE::MassDebug::HasDebugEntities())
+	{
+		return;
+	}
+	
+	QUICK_SCOPE_CYCLE_COUNTER(UMassDebugStateTreeProcessor_Run);	
+	EntityQuery.ForEachEntityChunk(EntitySubsystem, Context, [this, Debugger, MassStateTreeSubsystem, &EntitySubsystem](FLWComponentSystemExecutionContext& Context)
+		{
+			const FLWEntity SelectedEntity = Debugger->GetSelectedEntity();
+			const int32 NumEntities = Context.GetEntitiesNum();
+			const TConstArrayView<FMassStateTreeFragment> StateTreeList = Context.GetComponentView<FMassStateTreeFragment>();
+			const TConstArrayView<FDataFragment_Transform> TransformList = Context.GetComponentView<FDataFragment_Transform>();
+			const UStateTree* StateTree = MassStateTreeSubsystem->GetRegisteredStateTreeAsset(StateTreeList[0].StateTreeHandle);
+
+			// Not reporting error since this processor is a debug tool 
+			if (StateTree == nullptr)
+			{
+				return;
+			}
+		
+			for (int32 i = 0; i < NumEntities; ++i)
+			{
+				FLWEntity Entity = Context.GetEntity(i);
+				if (Entity == SelectedEntity)
+				{
+					FMassStateTreeExecutionContext StateTreeContext(EntitySubsystem, Context);
+					StateTreeContext.Init(*this, *StateTree, EStateTreeStorage::External);
+					StateTreeContext.SetEntity(Entity);
+					
+#if WITH_GAMEPLAY_DEBUGGER
+					const FStructView Storage = EntitySubsystem.GetComponentDataStruct(Entity, StateTree->GetRuntimeStorageStruct());
+					Debugger->AppendSelectedEntityInfo(StateTreeContext.GetDebugInfoString(Storage));
+#endif // WITH_GAMEPLAY_DEBUGGER
+				}
+					
+				FColor EntityColor = FColor::White;
+				const bool bDisplayDebug = UE::MassDebug::IsDebuggingEntity(Entity, &EntityColor);
+				if (bDisplayDebug)
+				{
+					const FDataFragment_Transform& Transform = TransformList[i];
+					
+					const FVector ZOffset(0,0,50);
+					const FVector Position = Transform.GetTransform().GetLocation() + ZOffset;
+
+					FMassStateTreeExecutionContext StateTreeContext(EntitySubsystem, Context);
+					StateTreeContext.Init(*this, *StateTree, EStateTreeStorage::External);
+					StateTreeContext.SetEntity(Entity);
+
+					const FStructView Storage = EntitySubsystem.GetComponentDataStruct(Entity, StateTree->GetRuntimeStorageStruct());
+
+					// State
+					UE_VLOG_SEGMENT_THICK(this, LogStateTree, Log, Position, Position + FVector(0,0,50), EntityColor, 2.0f, TEXT("%s %s"),
+						*Entity.DebugGetDescription(), *StateTreeContext.GetActiveStateName(Storage));
+				}
+			}
+		});
+	#endif // WITH_MASS_DEBUG
+}
diff --git a/Engine/Plugins/AI/MassAI/Source/MassAIDebug/Public/GameplayDebuggerCategory_Mass.h b/Engine/Plugins/AI/MassAI/Source/MassAIDebug/Public/GameplayDebuggerCategory_Mass.h
new file mode 100644
--- /dev/null
+++ b/Engine/Plugins/AI/MassAI/Source/MassAIDebug/Public/GameplayDebuggerCategory_Mass.h
@@ -0,0 +1,65 @@
+// Copyright Epic Games, Inc. All Rights Reserved.
+
+#pragma once
+
+// Required first for WITH_MASS_DEBUG
+#include "MassCommonTypes.h"
+
+#if WITH_GAMEPLAY_DEBUGGER && WITH_MASS_DEBUG
+
+#include "GameplayDebuggerCategory.h"
+
+class UMassDebuggerSubsystem;
+class APlayerController;
+class AActor;
+
+class FGameplayDebuggerCategory_Mass : public FGameplayDebuggerCategory
+{
+public:
+	FGameplayDebuggerCategory_Mass();
+
+	static TSharedRef<FGameplayDebuggerCategory> MakeInstance();
+
+protected:
+	virtual void CollectData(APlayerController* OwnerPC, AActor* DebugActor) override;
+	virtual void DrawData(APlayerController* OwnerPC, FGameplayDebuggerCanvasContext& CanvasContext) override;
+	
+	void SetCachedEntity(FLWEntity BestEntity, UMassDebuggerSubsystem& Debugger);
+
+	void OnToggleArchetypes() { bShowArchetypes = !bShowArchetypes; }
+	void OnToggleShapes() { bShowShapes = !bShowShapes; }
+	void OnToggleAgentFragments() { bShowAgentFragments = !bShowAgentFragments; }
+	void OnPickEntity() { bPickEntity = true; }
+	void OnToggleEntityDetails() { bShowEntityDetails = !bShowEntityDetails; }
+	void OnToggleNearEntityOverview() { bShowNearEntityOverview = !bShowNearEntityOverview; }
+	void OnToggleNearEntityAvoidance() { bShowNearEntityAvoidance = !bShowNearEntityAvoidance; }
+	void OnToggleNearEntityPath() { bShowNearEntityPath = !bShowNearEntityPath; }
+	
+	void PickEntity(const APlayerController& OwnerPC, UWorld& World, UMassDebuggerSubsystem& Debugger);
+
+protected:
+	AActor* CachedDebugActor;
+	FLWEntity CachedEntity;
+	bool bShowArchetypes;
+	bool bShowShapes;
+	bool bShowAgentFragments;
+	bool bPickEntity;
+	bool bShowEntityDetails;
+	bool bShowNearEntityOverview;
+	bool bShowNearEntityAvoidance;
+	bool bShowNearEntityPath;
+
+	struct FEntityDescription
+	{
+		FEntityDescription() = default;
+		FEntityDescription(const float InScore, const FVector& InLocation, const FString& InDescription) : Score(InScore), Location(InLocation), Description(InDescription) {}
+
+		float Score = 0.0f;
+		FVector Location = FVector::ZeroVector;
+		FString Description;
+	};
+	TArray<FEntityDescription> NearEntityDescriptions;
+
+};
+
+#endif // WITH_GAMEPLAY_DEBUGGER && WITH_MASS_DEBUG
diff --git a/Engine/Plugins/AI/MassAI/Source/MassAIDebug/Public/MassDebugStateTreeProcessor.h b/Engine/Plugins/AI/MassAI/Source/MassAIDebug/Public/MassDebugStateTreeProcessor.h
new file mode 100644
--- /dev/null
+++ b/Engine/Plugins/AI/MassAI/Source/MassAIDebug/Public/MassDebugStateTreeProcessor.h
@@ -0,0 +1,25 @@
+// Copyright Epic Games, Inc. All Rights Reserved.
+
+#pragma once
+
+#include "MassProcessor.h"
+#include "MassDebugStateTreeProcessor.generated.h"
+
+class UEntitySubsystem;
+struct FMassStateTreeFragment;
+struct FLWComponentQuery;
+struct FLWComponentSystemExecutionContext;
+
+UCLASS()
+class MASSAIDEBUG_API UMassDebugStateTreeProcessor : public UPipeProcessor
+{
+	GENERATED_BODY()
+
+protected:
+	UMassDebugStateTreeProcessor();
+
+	virtual void ConfigureQueries() override;
+	virtual void Execute(UEntitySubsystem& EntitySubsystem, FLWComponentSystemExecutionContext& Context) override;
+
+	FLWComponentQuery EntityQuery;
+};
diff --git a/Engine/Plugins/AI/MassAI/Source/MassAIMovement/MassAIMovement.Build.cs b/Engine/Plugins/AI/MassAI/Source/MassAIMovement/MassAIMovement.Build.cs
new file mode 100644
--- /dev/null
+++ b/Engine/Plugins/AI/MassAI/Source/MassAIMovement/MassAIMovement.Build.cs
@@ -0,0 +1,52 @@
+// Copyright Epic Games, Inc. All Rights Reserved.
+
+namespace UnrealBuildTool.Rules
+{
+	public class MassAIMovement : ModuleRules
+	{
+		public MassAIMovement(ReadOnlyTargetRules Target) : base(Target)
+		{
+			PCHUsage = ModuleRules.PCHUsageMode.UseExplicitOrSharedPCHs;
+
+			PublicIncludePaths.AddRange(
+				new string[] {
+				"Runtime/AIModule/Public",
+				ModuleDirectory + "/Public",
+				}
+			);
+
+
+			PrivateIncludePaths.AddRange(
+				new string[] {
+				}
+			);
+
+			PublicDependencyModuleNames.AddRange(
+				new string[] {
+					"AIModule",
+					"MassEntity",
+					"Core",
+					"CoreUObject",
+					"Engine",
+					"MassCommon",
+					"MassLOD",
+					"MassSignals",
+					"MassSimulation",
+					"MassSpawner",
+					"MassMovement",
+					"NavigationSystem",
+					"StructUtils",
+					"ZoneGraph",
+					"ZoneGraphAnnotations",
+					"DeveloperSettings",
+				}
+			);
+
+			if (Target.bBuildEditor == true)
+			{
+				PrivateDependencyModuleNames.Add("EditorFramework");
+				PrivateDependencyModuleNames.Add("UnrealEd");
+			}
+		}
+	}
+}
diff --git a/Engine/Plugins/AI/MassAI/Source/MassAIMovement/Private/AssignRandomNavLocationProcessor.cpp b/Engine/Plugins/AI/MassAI/Source/MassAIMovement/Private/AssignRandomNavLocationProcessor.cpp
new file mode 100644
--- /dev/null
+++ b/Engine/Plugins/AI/MassAI/Source/MassAIMovement/Private/AssignRandomNavLocationProcessor.cpp
@@ -0,0 +1,71 @@
+// Copyright Epic Games, Inc. All Rights Reserved.
+
+#include "AssignRandomNavLocationProcessor.h"
+#include "NavigationData.h"
+#include "NavigationSystem.h"
+#include "AI/Navigation/NavAgentInterface.h"
+#include "MassCommonTypes.h"
+#include "MassCommonFragments.h"
+
+//----------------------------------------------------------------------//
+// UAssignRandomNavLocationProcessor 
+//----------------------------------------------------------------------//
+UAssignRandomNavLocationProcessor::UAssignRandomNavLocationProcessor()
+{
+	bAutoRegisterWithProcessingPhases = false;
+}
+
+void UAssignRandomNavLocationProcessor::ConfigureQueries()
+{
+	EntityQuery.AddRequirement<FDataFragment_Transform>(ELWComponentAccess::ReadWrite);
+	EntityQuery.AddRequirement<FDataFragment_NavLocation>(ELWComponentAccess::ReadWrite);
+}
+
+void UAssignRandomNavLocationProcessor::Execute(UEntitySubsystem& EntitySubsystem, FLWComponentSystemExecutionContext& Context)
+{
+	if (!ensure(NavigationSystem))
+	{
+		return;
+	}
+
+	EntityQuery.ForEachEntityChunk(EntitySubsystem, Context, [this](FLWComponentSystemExecutionContext& Context)
+		{
+			TArrayView<FDataFragment_Transform> LocationList = Context.GetMutableComponentView<FDataFragment_Transform>();
+			TArrayView<FDataFragment_NavLocation> NavLocationList = Context.GetMutableComponentView<FDataFragment_NavLocation>();
+			const int32 NumEntities = Context.GetEntitiesNum();
+			TArray<FNavLocation> Locations;
+			Locations.AddDefaulted(NumEntities);
+			for (FNavLocation& Location : Locations)
+			{
+				NavigationSystem->GetRandomReachablePointInRadius(Origin, Radius, Location);
+			}
+
+			// apply locations
+			for (int32 Index = 0; Index < NumEntities; ++Index)
+			{
+				LocationList[Index].GetMutableTransform().SetLocation(Locations[Index].Location);
+				NavLocationList[Index].NodeRef = Locations[Index].NodeRef;
+			}
+		});
+}
+
+void UAssignRandomNavLocationProcessor::Initialize(UObject& InOwner)
+{
+	Super::Initialize(InOwner);
+
+	if (INavAgentInterface* AsNavAgent = Cast<INavAgentInterface>(&InOwner))
+	{
+		Origin = AsNavAgent->GetNavAgentLocation();
+	}
+	else if (AActor* AsActor = Cast<AActor>(&InOwner))
+	{
+		Origin = AsActor->GetActorLocation();
+	}
+	else
+	{
+		// @todo add logging
+	}
+
+	NavigationSystem = FNavigationSystem::GetCurrent<UNavigationSystemV1>(InOwner.GetWorld());
+	UE_CLOG(NavigationSystem == nullptr, LogMass, Error, TEXT("UAssignRandomNavLocationProcessor used while no NavigationSystem present"));
+}
diff --git a/Engine/Plugins/AI/MassAI/Source/MassAIMovement/Private/MassAIMovementFragments.cpp b/Engine/Plugins/AI/MassAI/Source/MassAIMovement/Private/MassAIMovementFragments.cpp
new file mode 100644
--- /dev/null
+++ b/Engine/Plugins/AI/MassAI/Source/MassAIMovement/Private/MassAIMovementFragments.cpp
@@ -0,0 +1,36 @@
+// Copyright Epic Games, Inc. All Rights Reserved.
+
+#include "MassAIMovementFragments.h"
+#include "GameFramework/GameStateBase.h"
+#include "Engine/World.h"
+
+void FMassMoveTargetFragment::CreateNewAction(const EMassMovementAction InAction, const UWorld& InWorld)
+{
+	ensureMsgf(InWorld.GetNetMode() != NM_Client, TEXT("This version of SetDesiredAction should only be called on the authority."));
+	CurrentAction = InAction;
+	CurrentActionID++;
+	CurrentActionWorldStartTime = InWorld.TimeSeconds;
+
+	const AGameStateBase* GameState = InWorld.GetGameState();
+	CurrentActionServerStartTime = GameState != nullptr ? GameState->GetServerWorldTimeSeconds() : CurrentActionWorldStartTime;
+
+	MarkNetDirty();
+}
+
+void FMassMoveTargetFragment::CreateReplicatedAction(const EMassMovementAction InAction, const uint16 InActionID, const float InWorldStartTime, const float InServerStartTime)
+{
+	CurrentAction = InAction;
+	CurrentActionID = InActionID;
+	CurrentActionWorldStartTime = InWorldStartTime;
+	CurrentActionServerStartTime = InServerStartTime;
+}
+
+FString FMassMoveTargetFragment::ToString() const
+{
+	return FString::Printf(TEXT("ActionID:%d Action:%s StartTime: World:%.1f Server:%.1f DesiredSpeed:%.1f "),
+		CurrentActionID,
+		*UEnum::GetValueAsString(CurrentAction),
+		CurrentActionWorldStartTime,
+		CurrentActionServerStartTime,
+		DesiredSpeed.Get());
+}
\ No newline at end of file
diff --git a/Engine/Plugins/AI/MassAI/Source/MassAIMovement/Private/MassAIMovementTypes.cpp b/Engine/Plugins/AI/MassAI/Source/MassAIMovement/Private/MassAIMovementTypes.cpp
new file mode 100644
--- /dev/null
+++ b/Engine/Plugins/AI/MassAI/Source/MassAIMovement/Private/MassAIMovementTypes.cpp
@@ -0,0 +1,19 @@
+// Copyright Epic Games, Inc. All Rights Reserved.
+
+#include "MassAIMovementTypes.h"
+
+DEFINE_LOG_CATEGORY(LogMassNavigation);
+DEFINE_LOG_CATEGORY(LogMassDynamicObstacle);
+
+void FMassMovementConfig::Update()
+{
+	for (FMassMovementStyleConfig& Style : MovementStyles)
+	{
+		Style.MaxProbability = 0;
+		for (FMassMovementStyleSpeed& Speed : Style.DesiredSpeeds)
+		{
+			Style.MaxProbability += Speed.Probability;
+			Speed.ProbabilityThreshold = Style.MaxProbability;
+		}
+	}
+}
diff --git a/Engine/Plugins/AI/MassAI/Source/MassAIMovement/Private/MassAvoidanceObstacleTrait.cpp b/Engine/Plugins/AI/MassAI/Source/MassAIMovement/Private/MassAvoidanceObstacleTrait.cpp
new file mode 100644
--- /dev/null
+++ b/Engine/Plugins/AI/MassAI/Source/MassAIMovement/Private/MassAvoidanceObstacleTrait.cpp
@@ -0,0 +1,22 @@
+// Copyright Epic Games, Inc. All Rights Reserved.
+
+#include "MassAvoidanceObstacleTrait.h"
+
+#include "MassAvoidanceProcessors.h"
+#include "MassEntityTemplateRegistry.h"
+#include "Engine/World.h"
+#include "MassAIMovementFragments.h"
+#include "MovementProcessor.h"
+
+void UMassAvoidanceObstacleTrait::BuildTemplate(FMassEntityTemplateBuildContext& BuildContext, UWorld& World) const
+{
+	// UAvoidanceObstacleProcessor
+	BuildContext.AddFragmentWithDefaultInitializer<FMassAvoidanceObstacleGridCellLocationFragment>();
+	BuildContext.AddDeinitializer(*GetDefault<UMassAvoidanceObstacleRemoverFragmentDestructor>());
+	BuildContext.AddFragmentWithDefaultInitializer<FDataFragment_AgentRadius>();
+
+	if (bExtendToEdgeObstacle)
+	{
+		BuildContext.AddTag<FMassAvoidanceExtendToEdgeObstacleTag>();
+	}
+}
\ No newline at end of file
