AggregateTicking and Pipe merged into Engine/Plugins/Runtime/MassEntity

diff --git a/Engine/Plugins/AI/MassAI/Source/MassAIMovement/Private/MassAvoidanceProcessors.cpp b/Engine/Plugins/AI/MassAI/Source/MassAIMovement/Private/MassAvoidanceProcessors.cpp
new file mode 100644
--- /dev/null
+++ b/Engine/Plugins/AI/MassAI/Source/MassAIMovement/Private/MassAvoidanceProcessors.cpp
@@ -0,0 +1,2088 @@
+// Copyright Epic Games, Inc. All Rights Reserved.
+
+#include "MassAvoidanceProcessors.h"
+#include "NavigationSystem.h"
+#include "DrawDebugHelpers.h"
+#include "EntityView.h"
+#include "MassAIMovementFragments.h"
+#include "VisualLogger/VisualLogger.h"
+#include "Math/Vector2D.h"
+#include "NavMesh/RecastNavMesh.h"
+#include "Logging/LogMacros.h"
+#include "ZoneGraphTypes.h"
+#include "ZoneGraphData.h"
+#include "ZoneGraphSubsystem.h"
+#include "MassAvoidanceSettings.h"
+#include "MassMovementSettings.h"
+#include "ZoneGraphQuery.h"
+#include "MassZoneGraphMovementFragments.h"
+
+#define UNSAFE_FOR_MT 1
+
+DEFINE_LOG_CATEGORY(LogAvoidance);
+DEFINE_LOG_CATEGORY(LogAvoidanceVelocities);
+DEFINE_LOG_CATEGORY(LogAvoidanceAgents);
+DEFINE_LOG_CATEGORY(LogAvoidanceObstacles);
+
+namespace UE::MassAvoidance
+{
+	namespace Tweakables
+	{
+		float AgentDetectionDistance = 400.f;
+		bool bEnableAvoidance = true;
+		bool bEnableSettingsforExtendingColliders = true;
+		bool bStopAvoidingOthersAtDestination = true;
+		bool bUseAdjacentCorridors = true;
+		bool bUseDrawDebugHelpers = false;
+	} // Tweakables
+
+	FAutoConsoleVariableRef Vars[] = 
+	{
+		FAutoConsoleVariableRef(TEXT("ai.mass.avoidance.AgentDetectionDistance"), Tweakables::AgentDetectionDistance, TEXT("Distance to detect other agents in cm."), ECVF_Cheat),
+		FAutoConsoleVariableRef(TEXT("ai.mass.avoidance.EnableAvoidance"), Tweakables::bEnableAvoidance, TEXT("Set to false to disable avoidance forces (for debug purposes)."), ECVF_Cheat),
+		FAutoConsoleVariableRef(TEXT("ai.mass.avoidance.EnableSettingsforExtendingColliders"), Tweakables::bEnableSettingsforExtendingColliders, TEXT("Set to false to disable using different settings for extending obstacles (for debug purposes)."), ECVF_Cheat),
+		FAutoConsoleVariableRef(TEXT("ai.mass.avoidance.StopAvoidingOthersAtDestination"), Tweakables::bStopAvoidingOthersAtDestination, TEXT("Once destination is reached, ignore predictive avoidance forces caused by other agents."), ECVF_Cheat),
+		FAutoConsoleVariableRef(TEXT("ai.mass.avoidance.UseAdjacentCorridors"), Tweakables::bUseAdjacentCorridors, TEXT("Set to false to disable usage of adjacent lane width."), ECVF_Cheat),
+		FAutoConsoleVariableRef(TEXT("ai.mass.avoidance.UseDrawDebugHelpers"), Tweakables::bUseDrawDebugHelpers, TEXT("Use debug draw helpers in addition to visual logs."), ECVF_Cheat)
+	};
+
+	constexpr int32 MaxExpectedAgentsPerCell = 6;
+	constexpr int32 MinTouchingCellCount = 4;
+	constexpr int32 MaxAgentResults = MaxExpectedAgentsPerCell * MinTouchingCellCount;
+
+	static void FindCloseAgents(const FVector& Center, const FAvoidanceObstacleHashGrid2D& AvoidanceObstacleGrid, TArray<FMassAvoidanceObstacleItem, TFixedAllocator<MaxAgentResults>>& OutCloseEntities, const int32 MaxResults)
+	{
+		//QUICK_SCOPE_CYCLE_COUNTER(FindCloseAgents);
+
+		OutCloseEntities.Reset();
+		const FVector Extent(Tweakables::AgentDetectionDistance, Tweakables::AgentDetectionDistance, 0.f);
+		const FBox QueryBox = FBox(Center - Extent, Center + Extent);
+
+		struct FSortingCell
+		{
+			int32 X;
+			int32 Y;
+			int32 Level;
+			float SqDist;
+		};
+		TArray<FSortingCell, TInlineAllocator<64>> Cells;
+		const FVector QueryCenter = QueryBox.GetCenter();
+		
+		for (int32 Level = 0; Level < AvoidanceObstacleGrid.NumLevels; Level++)
+		{
+			const float CellSize = AvoidanceObstacleGrid.GetCellSize(Level);
+			const FAvoidanceObstacleHashGrid2D::FCellRect Rect = AvoidanceObstacleGrid.CalcQueryBounds(QueryBox, Level);
+			for (int32 Y = Rect.MinY; Y <= Rect.MaxY; Y++)
+			{
+				for (int32 X = Rect.MinX; X <= Rect.MaxX; X++)
+				{
+					const float CenterX = (X + 0.5f) * CellSize;
+					const float CenterY = (Y + 0.5f) * CellSize;
+					const float DX = CenterX - QueryCenter.X;
+					const float DY = CenterY - QueryCenter.Y;
+					const float SqDist = DX * DX + DY * DY;
+					FSortingCell SortCell;
+					SortCell.X = X;
+					SortCell.Y = Y;
+					SortCell.Level = Level;
+					SortCell.SqDist = SqDist;
+					Cells.Add(SortCell);
+				}
+			}
+		}
+
+		Cells.Sort([](const FSortingCell& A, const FSortingCell& B) { return A.SqDist < B.SqDist; });
+
+		for (const FSortingCell& SortedCell : Cells)
+		{
+			if (const FAvoidanceObstacleHashGrid2D::FCell* Cell = AvoidanceObstacleGrid.FindCell(SortedCell.X, SortedCell.Y, SortedCell.Level))
+			{
+				const TSparseArray<FAvoidanceObstacleHashGrid2D::FItem>&  Items = AvoidanceObstacleGrid.GetItems();
+				for (int32 Idx = Cell->First; Idx != INDEX_NONE; Idx = Items[Idx].Next)
+				{
+					OutCloseEntities.Add(Items[Idx].ID);
+					if (OutCloseEntities.Num() >= MaxResults)
+					{
+						return;
+					}
+				}
+			}
+		}
+	}
+
+	// Adapted from ray-capsule intersection: https://iquilezles.org/www/articles/intersectors/intersectors.htm
+	static float ComputeTimeOfImpact(const FVector2D Pos, const FVector2D Vel, const float Rad,
+		const FVector2D SegStart, const FVector2D SegEnd)
+	{
+		const FVector2D SegDir = SegEnd - SegStart;
+		const FVector2D RelPos = Pos - SegStart;
+		const float VelSq = FVector2D::DotProduct(Vel, Vel);
+		const float SegDirSq = FVector2D::DotProduct(SegDir, SegDir);
+		const float DirVelSq = FVector2D::DotProduct(SegDir, Vel);
+		const float DirRelPosSq = FVector2D::DotProduct(SegDir, RelPos);
+		const float VelRelPosSq = FVector2D::DotProduct(Vel, RelPos);
+		const float RelPosSq = FVector2D::DotProduct(RelPos, RelPos);
+		const float A = SegDirSq * VelSq - DirVelSq * DirVelSq;
+		const float B = SegDirSq * VelRelPosSq - DirRelPosSq * DirVelSq;
+		const float C = SegDirSq * RelPosSq - DirRelPosSq * DirRelPosSq - FMath::Square(Rad) * SegDirSq;
+		const float H = FMath::Max<float>(0.f, B*B - A*C); // b^2 - ac, Using max for closest point of arrival result when no hit.
+		const float T = FMath::Abs(A) > 0.f ? (-B - FMath::Sqrt(H)) / A : 0.f;
+		const float Y = DirRelPosSq + T * DirVelSq;
+		
+		if (Y > 0.f && Y < SegDirSq) 
+		{
+			// body
+			return T;
+		}
+		else 
+		{
+			// caps
+			const FVector2D CapRelPos = (Y <= 0.f) ? RelPos : Pos - SegEnd;
+			const float Cb = FVector2D::DotProduct(Vel, CapRelPos);
+			const float Cc = FVector2D::DotProduct(CapRelPos, CapRelPos) - FMath::Square(Rad);
+			const float Ch = FMath::Max<float>(0.0f, Cb * Cb - VelSq * Cc);
+			return VelSq > 0.f ? (-Cb - FMath::Sqrt(Ch)) / VelSq : 0.f;
+		}
+	}
+
+	static float ComputeTimeOfImpact(const FVector RelPos, const FVector RelVel, const float TotalRadius)
+	{
+		// Calculate time of impact based on relative agent positions and velocities.
+		const float A = FVector::DotProduct(RelVel, RelVel);
+		const float Inv2A = A > 0.f ? 1.f / (2.f * A) : 0.f;
+		const float B = FMath::Min(0.f, 2.f * FVector::DotProduct(RelVel, RelPos));
+		const float C = FVector::DotProduct(RelPos, RelPos) - FMath::Square(TotalRadius);
+		// Using max() here gives us CPA (closest point on arrival) when there is no hit.
+		const float Discr = FMath::Sqrt(FMath::Max(0.f, B * B - 4.f * A * C));
+		return (-B - Discr) * Inv2A;
+	}
+
+	static FVector Clamp(const FVector Vec, const float Mag)
+	{
+		const float Len = Vec.SizeSquared();
+		if (Len > FMath::Square(Mag)) {
+			return Vec * Mag / FMath::Sqrt(Len);
+		}
+		return Vec;
+	}
+
+	static float ProjectPtSeg(const FVector2D Point, const FVector2D Start, const FVector2D End)
+	{
+		const FVector2D Seg = End - Start;
+		const FVector2D Dir = Point - Start;
+		const float d = Seg.SizeSquared();
+		const float t = FVector2D::DotProduct(Seg, Dir);
+		if (t < 0.f) return 0;
+		if (t > d) return 1;
+		return d > 0.f ? (t / d) : 0.f;
+	}
+
+	static float Smoothf(float X) 
+	{
+		return X * X * (3 - 2 * X);
+	}
+
+	static FVector GetLeftDirection(const FVector Dir, const FVector Up)
+	{
+		return FVector::CrossProduct(Dir, Up);
+	}
+
+	static FVector ComputeMiterDirection(const FVector PointA, const FVector PointB)
+	{
+		FVector Mid = 0.5f * (PointA + PointB);
+		const float MidSquared = FVector::DotProduct(Mid, Mid);
+		if (MidSquared > KINDA_SMALL_NUMBER)
+		{
+			const float Scale = FMath::Min(1.f / MidSquared, 20.f);
+			Mid *= Scale;
+		}
+		return Mid;
+	}
+
+	static bool UseDrawDebugHelper()
+	{
+		return Tweakables::bUseDrawDebugHelpers;
+	}
+
+#if WITH_MASS_DEBUG	
+	
+	//----------------------------------------------------------------------//
+	// Begin MassDebugUtils
+	// @todo: Extract those generic debug functions to a separate location
+	//----------------------------------------------------------------------//
+	struct FDebugContext
+	{
+		FDebugContext(const UObject* InLogOwner, const FLogCategoryBase& InCategory, const UWorld* InWorld, const FLWEntity InEntity)
+			: LogOwner(InLogOwner)
+			, Category(InCategory)
+			, World(InWorld)
+			, Entity(InEntity)
+		{}
+
+		const UObject* LogOwner;
+		const FLogCategoryBase& Category;
+		const UWorld* World;
+		const FLWEntity Entity;
+	};
+
+	static bool DebugIsSelected(const FLWEntity& Entity)
+	{
+		FColor Color;
+		return UE::MassDebug::IsDebuggingEntity(Entity, &Color);
+	}
+
+	static void DebugDrawLine(const FDebugContext& Context, const FVector& Start, const FVector& End, const FColor& Color, const float Thickness = 0.f, const bool bPersistent = false)
+	{
+		if (!DebugIsSelected(Context.Entity))
+		{
+			return;
+		}
+
+		UE_VLOG_SEGMENT_THICK(Context.LogOwner, Context.Category, Log, Start, End, Color, Thickness, TEXT(""));
+
+		if (UseDrawDebugHelper() && Context.World)
+		{
+			DrawDebugLine(Context.World, Start, End, Color, bPersistent, /*LifeTime =*/ -1.f, /*DepthPriority =*/ 0, Thickness);
+		}
+	}
+
+	static void DebugDrawArrow(const FDebugContext& Context, const FVector& Start, const FVector& End, const FColor& Color, const float HeadSize = 8.f, const float Thickness = 1.5f)
+	{
+		if (!DebugIsSelected(Context.Entity))
+		{
+			return;
+		}
+
+		const float Pointyness = 1.8f;
+		const FVector Line = End - Start;
+		const FVector UnitV = Line.GetSafeNormal();
+		const FVector Perp = FVector::CrossProduct(UnitV, FVector::UpVector);
+		const FVector Left = Perp - (Pointyness*UnitV);
+		const FVector Right = -Perp - (Pointyness*UnitV);
+		UE_VLOG_SEGMENT_THICK(Context.LogOwner, Context.Category, Log, Start, End, Color, Thickness, TEXT(""));
+		UE_VLOG_SEGMENT_THICK(Context.LogOwner, Context.Category, Log, End, End + HeadSize * Left, Color, Thickness, TEXT(""));
+		UE_VLOG_SEGMENT_THICK(Context.LogOwner, Context.Category, Log, End, End + HeadSize * Right, Color, Thickness, TEXT(""));
+
+		if (UseDrawDebugHelper() && Context.World)
+		{
+			DrawDebugLine(Context.World, Start, End, Color, /*bPersistent=*/ false, /*LifeTime =*/ -1.f, /*DepthPriority =*/ 0, Thickness);
+			DrawDebugLine(Context.World, End, End + HeadSize * Left, Color, /*bPersistent=*/ false, /*LifeTime =*/ -1.f, /*DepthPriority =*/ 0, Thickness);
+			DrawDebugLine(Context.World, End, End + HeadSize * Right, Color, /*bPersistent=*/ false, /*LifeTime =*/ -1.f, /*DepthPriority =*/ 0, Thickness);
+		}
+	}
+
+	static void DebugDrawSphere(const FDebugContext& Context, const FVector& Center, const float Radius, const FColor& Color)
+	{
+		if (!DebugIsSelected(Context.Entity))
+		{
+			return;
+		}
+
+		UE_VLOG_LOCATION(Context.LogOwner, Context.Category, Log, Center, Radius, Color, TEXT(""));
+
+		if (UseDrawDebugHelper() && Context.World)
+		{
+			DrawDebugSphere(Context.World, Center, Radius, /*segments = */16, Color);
+		}
+	}
+
+	static void DebugDrawBox(const FDebugContext& Context, const FBox& Box, const FColor& Color)
+	{
+		if (!DebugIsSelected(Context.Entity))
+		{
+			return;
+		}
+
+		UE_VLOG_BOX(Context.LogOwner, Context.Category, Log, Box, Color, TEXT(""));
+		
+		if (UseDrawDebugHelper() && Context.World)
+		{
+			DrawDebugBox(Context.World, Box.GetCenter(), Box.GetExtent(), Color);
+		}
+	}
+	
+	static void DebugDrawCylinder(const FDebugContext& Context, const FVector& Bottom, const FVector& Top, const float Radius, const FColor& Color, const FString& Text = FString())
+	{
+		if (!DebugIsSelected(Context.Entity))
+		{
+			return;
+		}
+
+		UE_VLOG_CYLINDER(Context.LogOwner, Context.Category, Log, Bottom, Top, Radius, Color, TEXT("%s"), *Text);
+
+		if (UseDrawDebugHelper() && Context.World)
+		{
+			DrawDebugCylinder(Context.World, Bottom, Top, Radius, /*segments = */24, Color);
+		}
+	}
+	//----------------------------------------------------------------------//
+	// End MassDebugUtils
+	//----------------------------------------------------------------------//
+
+
+	// Local debug utils
+	static void DebugDrawVelocity(const FDebugContext& Context, const FVector& Start, const FVector& End, const FColor& Color)
+	{
+		// Different arrow than DebugDrawArrow()
+		if (!DebugIsSelected(Context.Entity))
+		{
+			return;
+		}
+
+		const float Thickness = 3.f;
+		const float Pointyness = 1.8f;
+		const FVector Line = End - Start;
+		const FVector UnitV = Line.GetSafeNormal();
+		const FVector Perp = FVector::CrossProduct(UnitV, FVector::UpVector);
+		const FVector Left = Perp - (Pointyness * UnitV);
+		const FVector Right = -Perp - (Pointyness * UnitV);
+		const float HeadSize = 0.08f * Line.Size();
+		UE_VLOG_SEGMENT_THICK(Context.LogOwner, Context.Category, Log, Start, End, Color, Thickness, TEXT(""));
+		UE_VLOG_SEGMENT_THICK(Context.LogOwner, Context.Category, Log, End, End + HeadSize * Left, Color, Thickness, TEXT(""));
+		UE_VLOG_SEGMENT_THICK(Context.LogOwner, Context.Category, Log, End, End + HeadSize * Right, Color, Thickness, TEXT(""));
+		UE_VLOG_SEGMENT_THICK(Context.LogOwner, Context.Category, Log, End + HeadSize * Left, End + HeadSize * Right, Color, Thickness, TEXT(""));
+
+		if (UseDrawDebugHelper() && Context.World)
+		{
+			DrawDebugLine(Context.World, Start, End, Color, /*bPersistent=*/ false, /*LifeTime =*/ -1.f, /*DepthPriority =*/ 0, Thickness);
+			DrawDebugLine(Context.World, End, End + HeadSize * Left, Color, /*bPersistent=*/ false, /*LifeTime =*/ -1.f, /*DepthPriority =*/ 0, Thickness);
+			DrawDebugLine(Context.World, End, End + HeadSize * Right, Color, /*bPersistent=*/ false, /*LifeTime =*/ -1.f, /*DepthPriority =*/ 0, Thickness);
+			DrawDebugLine(Context.World, End + HeadSize * Left, End + HeadSize * Right, Color, /*bPersistent=*/ false, /*LifeTime =*/ -1.f, /*DepthPriority =*/ 0, Thickness);
+		}
+	}
+
+	static void DebugDrawForce(const FDebugContext& Context, const FVector& Start, const FVector& End, const FColor& Color)
+	{
+		DebugDrawArrow(Context, Start, End, Color, /*HeadSize*/4.f, /*Thickness*/3.f);
+	}
+
+	static void DebugDrawSummedForce(const FDebugContext& Context, const FVector& Start, const FVector& End, const FColor& Color)
+	{
+		DebugDrawArrow(Context, Start + FVector(0.f,0.f,1.f), End + FVector(0.f, 0.f, 1.f), Color, /*HeadSize*/8.f, /*Thickness*/6.f);
+	}
+
+	static void DebugDrawLane(const FDebugContext& Context, const FZoneGraphStorage& ZoneStorage, const FZoneGraphLaneHandle& LaneHandle, const FColor Color)
+	{
+		if (!ensure(LaneHandle.IsValid()))
+		{
+			return;
+		}
+
+		const FVector OffsetZ(0.f, 0.f, 1.f);
+		const FZoneLaneData& Lane = ZoneStorage.Lanes[LaneHandle.Index];
+		FVector PrevPoint = ZoneStorage.LanePoints[Lane.PointsBegin];
+		for (int32 i = Lane.PointsBegin + 1; i < Lane.PointsEnd; i++)
+		{
+			const FVector Point = ZoneStorage.LanePoints[i];
+			DebugDrawLine(Context, PrevPoint + OffsetZ, Point + OffsetZ, Color, /*Thickness*/2.f);
+			PrevPoint = Point;
+		}
+	}
+
+#endif // WITH_MASS_DEBUG
+
+} // namespace UE::MassAvoidance
+
+
+//----------------------------------------------------------------------//
+//  UMassAvoidanceProcessor
+//----------------------------------------------------------------------//
+UMassAvoidanceProcessor::UMassAvoidanceProcessor()
+{
+	bAutoRegisterWithProcessingPhases = true;
+	ExecutionFlags = (int32)EProcessorExecutionFlags::All;
+	ExecutionOrder.ExecuteInGroup = UE::Mass::ProcessorGroupNames::Avoidance;
+	ExecutionOrder.ExecuteAfter.Add(UE::Mass::ProcessorGroupNames::LOD);
+}
+
+void UMassAvoidanceProcessor::ConfigureQueries()
+{
+	EntityQuery.AddRequirement<FMassSteeringFragment>(ELWComponentAccess::ReadWrite);
+	EntityQuery.AddRequirement<FMassNavigationEdgesFragment>(ELWComponentAccess::ReadOnly);
+	EntityQuery.AddRequirement<FMassSimulationLODFragment>(ELWComponentAccess::ReadOnly);
+	EntityQuery.AddRequirement<FMassMoveTargetFragment>(ELWComponentAccess::ReadOnly);
+	EntityQuery.AddRequirement<FDataFragment_Transform>(ELWComponentAccess::ReadOnly);
+	EntityQuery.AddRequirement<FMassVelocityFragment>(ELWComponentAccess::ReadOnly);
+	EntityQuery.AddRequirement<FDataFragment_AgentRadius>(ELWComponentAccess::ReadOnly);
+	EntityQuery.AddRequirement<FMassMovementConfigFragment>(ELWComponentAccess::ReadOnly);
+	EntityQuery.AddTagRequirement<FMassHighLODTag>(ELWComponentPresence::All);
+}
+
+void UMassAvoidanceProcessor::Initialize(UObject& Owner)
+{
+	Super::Initialize(Owner);
+
+	WeakWorld = Owner.GetWorld();
+	WeakMovementSubsystem = UWorld::GetSubsystem<UMassMovementSubsystem>(Owner.GetWorld());
+}
+
+void UMassAvoidanceProcessor::Execute(UEntitySubsystem& EntitySubsystem, FLWComponentSystemExecutionContext& Context)
+{
+	QUICK_SCOPE_CYCLE_COUNTER(UMassAvoidanceProcessor);
+
+	const UWorld* World = WeakWorld.Get();
+	UMassMovementSubsystem* MovementSubsystem = WeakMovementSubsystem.Get();
+	const UMassAvoidanceSettings* Settings = UMassAvoidanceSettings::Get();
+	const UMassMovementSettings* MovementSettings = UMassMovementSettings::Get();
+	if (!World || !MovementSubsystem || !Settings || !MovementSettings)
+	{
+		return;
+	}
+
+	const float TimeDelta = Context.GetDeltaTimeSeconds();
+	const float CurrentTime = World->GetTimeSeconds();
+
+	// Naming notes:
+	// While computing forces on an agents, for close agents and environment edges, there are separation forces and avoidance forces.
+	// We aim to call them:
+	//		AgentSeparationForce
+	//		AgentAvoidForce
+	//		ObstacleSeparationForce
+	//		ObstacleAvoidForce
+
+	const float TimeHoriz = FMath::Max(Settings->TimeHorizon, KINDA_SMALL_NUMBER);
+	const float InvTimeHoriz = (1.f / TimeHoriz);
+
+	// In range 5-25. Comp/damp combined should be around 30 for firm collision response.
+	// More damping makes the agents to slow down more during collision, which results smoother sim 
+	const float AgentKSeparation = Settings->AgentSeparation;
+	const float AgentKSeparationForExtendingColliders = Settings->AgentSeparationForExtendingColliders;
+	const float AgentInset = Settings->AgentCollisionInset;
+	const float AgentBufferSeparation = FMath::Max(Settings->AgentSeparationBuffer, KINDA_SMALL_NUMBER);
+	const float AgentBufferSeparationAtEnd = FMath::Max(Settings->AgentSeparationBufferAtEnd, KINDA_SMALL_NUMBER);
+	const float AgentBufferSeparationForExtendingColliders = FMath::Max(Settings->AgentSeparationBufferForExtendingColliders, KINDA_SMALL_NUMBER);
+	
+	const float NearTargetLocDist = FMath::Max(Settings->NearTargetLocationDistance, KINDA_SMALL_NUMBER);
+
+	// Making the prediction smaller allows it to pass through things more easily,
+	// and making the buffer bigger makes the reaction smoother, yet firm. 
+	const float AvoidKAgent = Settings->AgentAvoidanceStiffness;
+	const float AvoidKObstacle = Settings->ObstacleAvoidanceStiffness;
+	const float AvoidInset = Settings->AvoidanceInset;
+	const float AvoidBuffer = FMath::Max(Settings->AvoidanceBuffer, KINDA_SMALL_NUMBER);
+	const float AvoidBufferForExtendingColliders = FMath::Max(Settings->AvoidanceBufferForExtendingColliders, KINDA_SMALL_NUMBER);
+	const float AvoidanceBufferAtEnd = FMath::Max(Settings->AvoidanceBufferAtEnd, KINDA_SMALL_NUMBER);
+
+	// Obstacle collision coeffs can be much bigger than agent-to-agent, they should be never violated.
+	// Damping is almost more important than compression 
+	const float ObstacleKSeparation = Settings->ObstacleSeparation;
+	const float ObstacleInset = Settings->ObstacleCollisionInset;
+	const float ObstacleBufferSeparation = FMath::Max(Settings->ObstacleSeparationBuffer, KINDA_SMALL_NUMBER);
+
+	// Colors
+	static const FColor CurrentAgentColor = FColor::Emerald;
+
+	static const FColor VelocityColor = FColor::Black;
+	static const FColor PrefVelocityColor = FColor::Red;
+	static const FColor DesiredVelocityColor = FColor::Yellow;
+	static const FColor FinalSteeringForceColor = FColor::Cyan;
+	static const float	BigArrowThickness = 6.f;
+	static const float	BigArrowHeadSize = 12.f;
+
+	// Agents colors
+	static const FColor AgentsColor = FColor::Orange;
+	static const FColor AgentSeparationForceColor = FColor(255, 145, 71);	// Orange red
+	static const FColor AgentAvoidForceColor = AgentsColor;
+	
+	// Obstacles colors
+	static const FColor ObstacleColor = FColor::Blue;
+	static const FColor ObstacleContactNormalColor = FColor::Silver;
+	static const FColor ObstacleAvoidForceColor = FColor::Magenta;
+	static const FColor ObstacleSeparationForceColor = FColor(255, 66, 66);	// Bright red
+	
+	static const FVector DebugAgentHeightOffset = FVector(0.f, 0.f, 185.f);
+	static const FVector DebugLowCylinderOffset = FVector(0.f, 0.f, 20.f);
+
+	EntityQuery.ForEachEntityChunk(EntitySubsystem, Context, [=, &EntitySubsystem](FLWComponentSystemExecutionContext& Ctx)
+	{
+		const int32 NumEntities = Ctx.GetEntitiesNum();
+		const float DistanceCutOffSquare = FMath::Square(UE::MassAvoidance::Tweakables::AgentDetectionDistance);
+
+		const TArrayView<FMassSteeringFragment> SteeringList = Ctx.GetMutableComponentView<FMassSteeringFragment>();
+		const TConstArrayView<FMassNavigationEdgesFragment> NavEdgesList = Ctx.GetComponentView<FMassNavigationEdgesFragment>();
+		const TConstArrayView<FDataFragment_Transform> LocationList = Ctx.GetComponentView<FDataFragment_Transform>();
+		const TConstArrayView<FMassVelocityFragment> VelocityList = Ctx.GetComponentView<FMassVelocityFragment>();
+		const TConstArrayView<FDataFragment_AgentRadius> RadiusList = Ctx.GetComponentView<FDataFragment_AgentRadius>();
+		const TConstArrayView<FMassSimulationLODFragment> SimLODList = Ctx.GetComponentView<FMassSimulationLODFragment>();
+		const TConstArrayView<FMassMoveTargetFragment> MoveTargetList = Ctx.GetComponentView<FMassMoveTargetFragment>();
+		const TConstArrayView<FMassMovementConfigFragment> MovementConfigList = Ctx.GetComponentView<FMassMovementConfigFragment>();
+
+		// Arrays used to store close agents
+		TArray<FMassAvoidanceObstacleItem, TFixedAllocator<UE::MassAvoidance::MaxAgentResults>> CloseEntities;
+
+		struct FSortingAgent
+		{
+			FVector LocationCached;
+			FVector Forward;
+			FMassAvoidanceObstacleItem ObstacleItem;
+			float SqDist;
+		};
+		TArray<FSortingAgent, TFixedAllocator<UE::MassAvoidance::MaxAgentResults>> ClosestAgents;
+
+		struct FContact
+		{
+			FVector Position = FVector::ZeroVector;
+			FVector Normal = FVector::ZeroVector;
+			float Distance = 0.f;
+		};
+		TArray<FContact, TInlineAllocator<16>> Contacts;
+
+		struct FCollider
+		{
+			FVector Location = FVector::ZeroVector;
+			FVector Velocity = FVector::ZeroVector;
+			float Radius = 0.f;
+			bool bExtendToEdge = false;
+			bool bIsMoving = false;
+		};
+		TArray<FCollider, TInlineAllocator<16>> Colliders;
+
+		// Get the default movement config.
+		FMassMovementConfigHandle CurrentConfigHandle;
+		const FMassMovementConfig* CurrentMovementConfig = nullptr;
+
+		// Steps are:
+		//	1. Prepare agents
+		//	2. Avoid environment: add edge avoidance force and edge separation force
+		//  3. Avoid close agents: add agent avoidance force and agent separation force
+		//  4. Add noise (TBD)
+		//  5. Integrate and orient
+
+		for (int32 EntityIndex = 0; EntityIndex < NumEntities; ++EntityIndex)
+		{
+			// @todo: this should eventually be part of the query.
+			const FMassMoveTargetFragment& MoveTarget = MoveTargetList[EntityIndex];
+			if (MoveTarget.GetCurrentAction() == EMassMovementAction::Animate || MoveTarget.GetCurrentAction() == EMassMovementAction::Stand)
+			{
+				continue;
+			}
+
+			FLWEntity Entity = Ctx.GetEntity(EntityIndex);
+			
+			const FMassMovementConfigFragment& MovementConfig = MovementConfigList[EntityIndex];
+			if (MovementConfig.ConfigHandle != CurrentConfigHandle)
+			{
+				CurrentMovementConfig = MovementSettings->GetMovementConfigByHandle(MovementConfig.ConfigHandle);
+				CurrentConfigHandle = MovementConfig.ConfigHandle;
+			}
+			if (!CurrentMovementConfig)
+			{
+#if WITH_MASS_DEBUG && UNSAFE_FOR_MT
+				UE_VLOG(this, LogMassNavigation, Log, TEXT("%s Invalid movement config."), *Entity.DebugGetDescription());
+#endif
+				continue;
+			}
+
+			FMassSteeringFragment& Steering = SteeringList[EntityIndex];
+			const FMassNavigationEdgesFragment& NavEdges = NavEdgesList[EntityIndex];
+			const FDataFragment_Transform& Location = LocationList[EntityIndex];
+			const FMassVelocityFragment& Velocity = VelocityList[EntityIndex];
+			const FDataFragment_AgentRadius& Radius = RadiusList[EntityIndex];
+			const FMassSimulationLODFragment& SimLOD = SimLODList[EntityIndex];
+
+			// Smaller steering max accel makes the steering more "calm" but less opportunistic, may not find solution, or gets stuck.
+			// Max contact accel should be quite a big bigger than steering so that collision response is firm. 
+			const float MaxSteerAccel = CurrentMovementConfig->Steering.MaxAcceleration;
+			const float MaximumSpeed = CurrentMovementConfig->MaximumSpeed;
+
+			const FVector AgentLocation = Location.GetTransform().GetTranslation();
+			const FVector AgentVelocity = FVector(Velocity.Value.X, Velocity.Value.Y, 0.0f);
+			const float AgentRadius = Radius.Radius;
+			const FVector PrefVelocity = Steering.DesiredVelocity;
+			FVector SteeringForce = Steering.SteeringForce;
+
+			const bool bFadeAvoidingAtDestination = UE::MassAvoidance::Tweakables::bStopAvoidingOthersAtDestination && MoveTarget.IntentAtGoal == EMassMovementAction::Stand;
+			const float NearEndFade = bFadeAvoidingAtDestination ? FMath::Clamp(MoveTargetList[EntityIndex].DistanceToGoal / NearTargetLocDist, 0.f, 1.f) : 1.0f;
+			const bool bAgentIsMoving = MoveTarget.GetCurrentAction() == EMassMovementAction::Move;
+			
+#if WITH_MASS_DEBUG && UNSAFE_FOR_MT
+			const UE::MassAvoidance::FDebugContext BaseDebugContext(this, LogAvoidance, World, Entity);
+			const UE::MassAvoidance::FDebugContext VelocitiesDebugContext(this, LogAvoidanceVelocities, World, Entity);
+			const UE::MassAvoidance::FDebugContext ObstacleDebugContext(this, LogAvoidanceObstacles, World, Entity);
+			const UE::MassAvoidance::FDebugContext AgentDebugContext(this, LogAvoidanceAgents, World, Entity);
+			
+			if (UE::MassAvoidance::DebugIsSelected(Entity))
+			{
+				// Draw agent
+				const FString Text = FString::Printf(TEXT("%i"), Entity.Index);
+				DebugDrawCylinder(BaseDebugContext, AgentLocation, AgentLocation + DebugAgentHeightOffset, AgentRadius+1.f, CurrentAgentColor, Text);
+
+				const FColor AgentColor = (SimLOD.LOD == EMassLOD::High) ? CurrentAgentColor : FColor::Red;
+				DebugDrawSphere(BaseDebugContext, AgentLocation, 10.f, AgentColor);
+
+				// Draw current velocity (black)
+				DebugDrawVelocity(VelocitiesDebugContext, AgentLocation + DebugAgentHeightOffset, AgentLocation + DebugAgentHeightOffset + AgentVelocity, VelocityColor);
+
+				// Draw preferred velocity (red)
+				DebugDrawVelocity(VelocitiesDebugContext, AgentLocation + DebugAgentHeightOffset, AgentLocation + DebugAgentHeightOffset + PrefVelocity, PrefVelocityColor);
+
+				// Draw initial steering force
+				DebugDrawArrow(BaseDebugContext, AgentLocation + DebugAgentHeightOffset, AgentLocation + DebugAgentHeightOffset + SteeringForce, CurrentAgentColor, BigArrowHeadSize, BigArrowThickness);
+
+				// Draw center
+				DebugDrawSphere(BaseDebugContext, AgentLocation, /*Radius*/2.f, CurrentAgentColor);
+			}
+#endif // WITH_MASS_DEBUG
+
+			FVector OldSteeringForce = FVector::ZeroVector;
+			
+			if (!MoveTarget.bOffBoundaries && UE::MassAvoidance::Tweakables::bEnableAvoidance)
+			{
+				const FVector DesiredAcceleration = UE::MassAvoidance::Clamp(SteeringForce, MaxSteerAccel);
+				const FVector DesiredVelocity = UE::MassAvoidance::Clamp(AgentVelocity + DesiredAcceleration * TimeDelta, MaximumSpeed);
+
+#if WITH_MASS_DEBUG && UNSAFE_FOR_MT
+				// Draw desired velocity (yellow)
+				UE::MassAvoidance::DebugDrawVelocity(VelocitiesDebugContext, AgentLocation + DebugAgentHeightOffset, AgentLocation + DebugAgentHeightOffset + DesiredVelocity, DesiredVelocityColor);
+#endif // WITH_MASS_DEBUG
+
+				OldSteeringForce = SteeringForce;
+				Contacts.Reset();
+
+				for (const FNavigationAvoidanceEdge& Edge : NavEdges.AvoidanceEdges)
+				{
+					const FVector EdgeDiff = Edge.End - Edge.Start;
+					FVector EdgeDir = FVector::ZeroVector;
+					float EdgeLength = 0.0f;
+					EdgeDiff.ToDirectionAndLength(EdgeDir, EdgeLength);
+
+					const FVector AgentToEdgeStart = AgentLocation - Edge.Start;
+					const float DistAlongEdge = FVector::DotProduct(EdgeDir, AgentToEdgeStart);
+					const float DistAwayFromEdge = FVector::DotProduct(Edge.LeftDir, AgentToEdgeStart);
+
+					float ConDist = 0.0f;
+					FVector ConNorm = FVector::ForwardVector;
+					FVector ConPos = FVector::ZeroVector;
+					bool bDirectlyBehindEdge = false;
+					
+					if (DistAwayFromEdge < 0.0f)
+					{
+						// Inside or behind the edge
+						if (DistAlongEdge < 0.0f)
+						{
+							ConPos = Edge.Start;
+							ConNorm = -EdgeDir;
+							ConDist = -DistAlongEdge;
+						}
+						else if (DistAlongEdge > EdgeLength)
+						{
+							ConPos = Edge.End;
+							ConNorm = EdgeDir;
+							ConDist = DistAlongEdge;
+						}
+						else
+						{
+							ConPos = Edge.Start + EdgeDir * DistAlongEdge;
+							ConNorm = Edge.LeftDir;
+							ConDist = 0.0f;
+							bDirectlyBehindEdge = true;
+						}
+					}
+					else
+					{
+						if (DistAlongEdge < 0.0f)
+						{
+							// Start Corner
+							ConPos = Edge.Start;
+							const FVector RelPos = AgentLocation - Edge.Start;
+							EdgeDiff.ToDirectionAndLength(ConNorm, ConDist);
+						}
+						else if (DistAlongEdge > EdgeLength)
+						{
+							// End Corner
+							ConPos = Edge.End;
+							const FVector RelPos = AgentLocation - Edge.End;
+							EdgeDiff.ToDirectionAndLength(ConNorm, ConDist);
+						}
+						else
+						{
+							// Front
+							ConPos = Edge.Start + EdgeDir * DistAlongEdge;
+							ConNorm = Edge.LeftDir;
+							ConDist = DistAwayFromEdge;
+						}
+					}
+					
+					// Check to merge contacts
+					bool bAdd = true;
+					for (int ContactIndex = 0; ContactIndex < Contacts.Num(); ContactIndex++)
+					{
+						if (FVector::DotProduct(Contacts[ContactIndex].Normal, ConNorm) > 0.f && FMath::Abs(FVector::DotProduct(ConNorm, Contacts[ContactIndex].Position - ConPos)) < (10.f/*cm*/))
+						{
+							// Contacts are on same place, merge
+							if (ConDist < Contacts[ContactIndex].Distance)
+							{
+								// New is closer, override.
+								Contacts[ContactIndex].Position = ConPos;
+								Contacts[ContactIndex].Normal = ConNorm;
+								Contacts[ContactIndex].Distance = ConDist;
+							}
+							bAdd = false;
+							break;
+						}
+					}
+
+					// Not found, add new contact
+					if (bAdd)
+					{
+						FContact Contact;
+						Contact.Position = ConPos;
+						Contact.Normal = ConNorm;
+						Contact.Distance = ConDist;
+						Contacts.Add(Contact);
+					}
+
+					// Skip predictive avoidance when behind the edge.
+					if (!bDirectlyBehindEdge)
+					{
+						// Avoid edges
+						float TOI = UE::MassAvoidance::ComputeTimeOfImpact(FVector2D(AgentLocation), FVector2D(DesiredVelocity), AgentRadius, FVector2D(Edge.Start), FVector2D(Edge.End));
+						TOI = FMath::Clamp(TOI, 0.0f, TimeHoriz);
+						const FVector HitAgentPos = AgentLocation + DesiredVelocity * TOI;
+						const float T2 = UE::MassAvoidance::ProjectPtSeg(FVector2D(HitAgentPos), FVector2D(Edge.Start), FVector2D(Edge.End));
+						const FVector HitObPos = FMath::Lerp(Edge.Start, Edge.End, T2);
+
+						// Calculate penetration at CPA
+						FVector AvoidRelPos = HitAgentPos - HitObPos;
+						AvoidRelPos.Z = 0.f;	// @todo AT: ignore the z component for now until we clamp the height of obstacles
+						const float AvoidDist = AvoidRelPos.Size();
+						const FVector AvoidNormal = AvoidDist > 0.0f ? (AvoidRelPos / AvoidDist) : FVector::ForwardVector;
+
+						const float AvoidPen = (AgentRadius - AvoidInset + AvoidBuffer) - AvoidDist;
+						const float AvoidMag = FMath::Square(FMath::Clamp(AvoidPen / AvoidBuffer, 0.f, 1.f));
+						const float AvoidMagDist = 1.f + FMath::Square(1.f - TOI * InvTimeHoriz);
+						const FVector AvoidForce = AvoidNormal * AvoidMag * AvoidMagDist * AvoidKObstacle * NearEndFade; // Predictive avoidance against environment is tuned down towards the end of the path
+
+						SteeringForce += AvoidForce;
+
+#if WITH_MASS_DEBUG && UNSAFE_FOR_MT
+						// Draw contact normal
+						UE::MassAvoidance::DebugDrawArrow(ObstacleDebugContext, ConPos, ConPos + 50.f * ConNorm, ObstacleContactNormalColor, /*HeadSize=*/ 5.f);
+						UE::MassAvoidance::DebugDrawSphere(ObstacleDebugContext, ConPos, 2.5f, ObstacleContactNormalColor);
+
+						// Draw hit pos with edge
+						UE::MassAvoidance::DebugDrawLine(ObstacleDebugContext, AgentLocation, HitAgentPos, ObstacleAvoidForceColor);
+						UE::MassAvoidance::DebugDrawCylinder(ObstacleDebugContext, HitAgentPos, HitAgentPos + DebugAgentHeightOffset, AgentRadius, ObstacleAvoidForceColor);
+
+						// Draw avoid obstacle force
+						UE::MassAvoidance::DebugDrawForce(ObstacleDebugContext, HitObPos, HitObPos + AvoidForce, ObstacleAvoidForceColor);
+#endif // WITH_MASS_DEBUG
+					}
+				} // edge loop
+
+#if WITH_MASS_DEBUG && UNSAFE_FOR_MT
+				// Draw total steering force to avoid obstacles
+				const FVector EnvironmentAvoidSteeringForce = SteeringForce - OldSteeringForce;
+				UE::MassAvoidance::DebugDrawSummedForce(ObstacleDebugContext,
+					AgentLocation + DebugAgentHeightOffset,
+					AgentLocation + DebugAgentHeightOffset + EnvironmentAvoidSteeringForce,
+					ObstacleAvoidForceColor);
+#endif // WITH_MASS_DEBUG
+
+				// Process contacts to add edge separation force
+				const FVector SteeringForceBeforeSeparation = SteeringForce;
+				for (int ContactIndex = 0; ContactIndex < Contacts.Num(); ContactIndex++) 
+				{
+					const FVector ConNorm = Contacts[ContactIndex].Normal.GetSafeNormal();
+					const float ContactDist = Contacts[ContactIndex].Distance;
+
+					// Separation force (stay away from obstacles if possible)
+					const float SeparationPenalty = (AgentRadius - ObstacleInset + ObstacleBufferSeparation) - ContactDist;
+					const float SeparationMag = UE::MassAvoidance::Smoothf(FMath::Clamp(SeparationPenalty / ObstacleBufferSeparation, 0.f, 1.f));
+					const FVector SeparationForce = ConNorm * ObstacleKSeparation * SeparationMag;
+
+					SteeringForce += SeparationForce;
+
+#if WITH_MASS_DEBUG && UNSAFE_FOR_MT
+					// Draw individual contact forces
+					DebugDrawForce(ObstacleDebugContext, Contacts[ContactIndex].Position + DebugAgentHeightOffset, Contacts[ContactIndex].Position + SeparationForce + DebugAgentHeightOffset, ObstacleSeparationForceColor);
+#endif // WITH_MASS_DEBUG
+				}
+				
+#if WITH_MASS_DEBUG && UNSAFE_FOR_MT
+				// Draw total steering force to separate from close edges
+				const FVector TotalSeparationForce = SteeringForce - SteeringForceBeforeSeparation;
+				DebugDrawSummedForce(ObstacleDebugContext,
+					AgentLocation + DebugAgentHeightOffset,
+					AgentLocation + DebugAgentHeightOffset + TotalSeparationForce,
+					ObstacleSeparationForceColor);
+
+				// Display close obstacle edges
+				if (UE::MassAvoidance::DebugIsSelected(Entity))
+				{
+					for (const FNavigationAvoidanceEdge& Edge : NavEdges.AvoidanceEdges)
+					{
+						DebugDrawLine(ObstacleDebugContext, DebugAgentHeightOffset + Edge.Start, DebugAgentHeightOffset + Edge.End, ObstacleColor, /*Thickness=*/2.f);
+						const FVector Middle = DebugAgentHeightOffset + 0.5f * (Edge.Start + Edge.End);
+						DebugDrawArrow(ObstacleDebugContext, Middle, Middle + 10.f * FVector::CrossProduct((Edge.End - Edge.Start), FVector::UpVector).GetSafeNormal(), ObstacleColor, /*HeadSize=*/2.f);
+					}
+				}
+#endif // WITH_MASS_DEBUG
+			}
+
+			//////////////////////////////////////////////////////////////////////////
+			// 3. Avoid close agents
+			// Desired velocity
+			const FVector DesAcc = UE::MassAvoidance::Clamp(SteeringForce, MaxSteerAccel);
+			const FVector DesVel = UE::MassAvoidance::Clamp(AgentVelocity + DesAcc * TimeDelta, MaximumSpeed);
+
+			// Find close obstacles
+			const FAvoidanceObstacleHashGrid2D& AvoidanceObstacleGrid = MovementSubsystem->GetGridMutable();
+			UE::MassAvoidance::FindCloseAgents(AgentLocation, AvoidanceObstacleGrid, CloseEntities, UE::MassAvoidance::MaxAgentResults);
+
+			// Remove unwanted and find the closests in the CloseEntities
+			ClosestAgents.Reset();
+			for (const FAvoidanceObstacleHashGrid2D::ItemIDType OtherEntity : CloseEntities)
+			{
+				// Skip self
+				if (OtherEntity.Entity == Entity)
+				{
+					continue;
+				}
+
+				// Skip invalid entities.
+				if (!EntitySubsystem.IsEntityValid(OtherEntity.Entity))
+				{
+					UE_LOG(LogAvoidanceObstacles, VeryVerbose, TEXT("Close entity is invalid, skipped."));
+					continue;
+				}
+				
+				// Skip too far
+				const FTransform& Transform = EntitySubsystem.GetComponentDataChecked<FDataFragment_Transform>(OtherEntity.Entity).GetTransform();
+				const FVector OtherLocation = Transform.GetLocation();
+				
+				const float SqDist = FVector::DistSquared(AgentLocation, OtherLocation);
+				if (SqDist > DistanceCutOffSquare)
+				{
+					continue;
+				}
+
+				FSortingAgent SortAgent;
+				SortAgent.LocationCached = OtherLocation;
+				SortAgent.Forward = Transform.GetRotation().GetForwardVector();
+				SortAgent.ObstacleItem = OtherEntity;
+				SortAgent.SqDist = SqDist;
+				ClosestAgents.Add(SortAgent);
+			}
+			ClosestAgents.Sort([](const FSortingAgent& A, const FSortingAgent& B) { return A.SqDist < B.SqDist; });
+
+			// Compute forces
+			OldSteeringForce = SteeringForce;
+			FVector TotalAgentSeparationForce = FVector::ZeroVector;
+
+			// Fill collider's list ouf of close agents
+			Colliders.Reset();
+			constexpr int32 MaxColliders = 6;
+			for (int32 Index = 0; Index < ClosestAgents.Num(); Index++)
+			{
+				if(Colliders.Num() >= MaxColliders)
+				{
+					break;
+				}
+
+				FSortingAgent& OtherAgent = ClosestAgents[Index];
+				FEntityView OtherEntityView(EntitySubsystem, OtherAgent.ObstacleItem.Entity);
+
+				const FMassVelocityFragment* OtherVelocityFragment = OtherEntityView.GetComponentDataPtr<FMassVelocityFragment>();
+				const FVector OtherVelocity = OtherVelocityFragment != nullptr ? OtherVelocityFragment->Value : FVector::ZeroVector; // Get velocity from FAvoidanceComponent
+				const bool bExtendToEdge = OtherEntityView.HasTag<FMassAvoidanceExtendToEdgeObstacleTag>();
+
+				// @todo: this is heavy fragment to access, see if we could handle this differently.
+				const FMassMoveTargetFragment* OtherMoveTarget = OtherEntityView.GetComponentDataPtr<FMassMoveTargetFragment>();
+				const bool bOtherIsMoving = OtherMoveTarget ? OtherMoveTarget->GetCurrentAction() == EMassMovementAction::Move : true; // Assume moving if other does not have move target.
+				
+				// Check for colliders data
+				if (EnumHasAnyFlags(OtherAgent.ObstacleItem.ItemFlags, EMassAvoidanceObstacleItemFlags::HasColliderData))
+				{
+					const FMassAvoidanceColliderFragment* ColliderFragment = OtherEntityView.GetComponentDataPtr<FMassAvoidanceColliderFragment>();
+					if (ColliderFragment)
+					{
+						if (ColliderFragment->Type == EMassColliderType::Circle)
+						{
+							FCollider Collider;
+							Collider.Velocity = OtherVelocity;
+							Collider.bExtendToEdge = bExtendToEdge;
+							Collider.bIsMoving = bOtherIsMoving;
+							Collider.Radius = ColliderFragment->GetCircleCollider().Radius;
+							Collider.Location = OtherAgent.LocationCached;
+							Colliders.Add(Collider);
+						}
+						else if (ColliderFragment->Type == EMassColliderType::Pill)
+						{
+							FCollider Collider;
+							Collider.Velocity = OtherVelocity;
+							Collider.bExtendToEdge = bExtendToEdge;
+							Collider.bIsMoving = bOtherIsMoving;
+							const FMassPillCollider& Pill = ColliderFragment->GetPillCollider(); 
+							Collider.Radius = Pill.Radius;
+
+							Collider.Location = OtherAgent.LocationCached + (Pill.HalfLength * OtherAgent.Forward);
+							Colliders.Add(Collider);
+
+							if(Colliders.Num() >= MaxColliders)
+							{
+								break;
+							}
+
+							Collider.Location = OtherAgent.LocationCached + (-Pill.HalfLength * OtherAgent.Forward);
+							Colliders.Add(Collider);
+						}
+					}
+				}
+				else
+				{
+					FCollider Collider;
+					Collider.Location = OtherAgent.LocationCached;
+					Collider.Velocity = OtherVelocity;
+					Collider.Radius = OtherEntityView.GetComponentData<FDataFragment_AgentRadius>().Radius;
+					Collider.bExtendToEdge = bExtendToEdge;
+					Collider.bIsMoving = bOtherIsMoving;
+					Colliders.Add(Collider);
+				}
+			}
+
+			// Process colliders for avoidance
+			for (FCollider Collider : Colliders)
+			{
+				// Increases radius and offset agent position to ease avoidance for obstacle near edges.
+				bool bDebugIsOtherAgentUpdated = false;
+
+				bool bHasForcedNormal = false;
+				FVector ForcedNormal = FVector::ZeroVector;
+
+				if (Collider.bExtendToEdge)
+				{
+					// If the space between edge and collider is less than MinClearance, make the agent to avoid the gap.
+					constexpr float ClearanceScale = 0.7f; // @todo: Make configurable
+					const float MinClearance = 2.f*AgentRadius * ClearanceScale;
+					
+					// Find the maximum distance from edges that are too close.
+					float MaxDist = -1.f;
+					FVector ClosestPoint;
+					for (const FNavigationAvoidanceEdge& Edge : NavEdges.AvoidanceEdges)
+					{
+						const FVector Point = FMath::ClosestPointOnSegment(Collider.Location, Edge.Start, Edge.End);
+						const FVector Offset = Collider.Location - Point;
+						if (FVector::DotProduct(Offset, Edge.LeftDir) < 0.f)
+						{
+							// Behind the edge, ignore.
+							continue;
+						}
+
+						const float OffsetLength = Offset.Length();
+						const bool bTooNarrow = (OffsetLength - Collider.Radius) < MinClearance; 
+						if (bTooNarrow)
+						{
+							MaxDist = FMath::Max(OffsetLength, MaxDist);
+							ClosestPoint = Point;
+						}
+					}
+
+					if (MaxDist != -1.f)
+					{
+						// Set up forced normal to avoid the gap between collider and edge.
+						ForcedNormal = (Collider.Location - ClosestPoint).GetSafeNormal();
+						bHasForcedNormal = true;
+					}
+				}
+
+				const float TotalRadius = AgentRadius + Collider.Radius;
+
+				FVector RelPos = AgentLocation - Collider.Location;
+				RelPos.Z = 0.f; // we assume we work on a flat plane for now
+				const FVector RelVel = DesVel - Collider.Velocity;
+				const float ConDist = RelPos.Size();
+				const FVector ConNorm = ConDist > 0.f ? RelPos / ConDist : FVector::ForwardVector;
+
+				FVector SeparationNormal = ConNorm;
+				if (bHasForcedNormal)
+				{
+					// The more head on the collisions is, the more we should avoid towards the forced direction.
+					const FVector RelVelNorm = RelVel.GetSafeNormal();
+					const float Blend = FMath::Max(0.0f, -FVector::DotProduct(ConNorm, RelVelNorm));
+					SeparationNormal = FMath::Lerp(ConNorm, ForcedNormal, Blend).GetSafeNormal();
+				}
+				
+				// @todo: Make configurable
+				const float StandingScaling = Collider.bIsMoving ? 1.0f : 0.65f; // Care less about standing agents so that we can push through standing crowd.
+
+				const bool bUseExtendingCollidersSettings = Collider.bExtendToEdge && UE::MassAvoidance::Tweakables::bEnableSettingsforExtendingColliders;
+				const float AgentSeparationForce = bUseExtendingCollidersSettings ? AgentKSeparationForExtendingColliders : AgentKSeparation;
+				const float Buffer = bUseExtendingCollidersSettings ? AgentBufferSeparationForExtendingColliders : AgentBufferSeparation;
+				const float AgentAvoidBuffer = bUseExtendingCollidersSettings ? AvoidBufferForExtendingColliders : AvoidBuffer;
+				const float ContextualAgentBufferSeparation = FMath::Lerp(AgentBufferSeparationAtEnd, Buffer, NearEndFade);
+
+				// Separation force (stay away from agents if possible)
+				const float PenSep = (TotalRadius - AgentInset + ContextualAgentBufferSeparation) - ConDist;
+				const float SeparationMag = FMath::Square(FMath::Clamp(PenSep / ContextualAgentBufferSeparation, 0.f, 1.f));
+				const FVector SepForce = SeparationNormal * AgentSeparationForce;
+				const FVector SeparationForce = SepForce * SeparationMag * StandingScaling;
+
+				SteeringForce += SeparationForce;
+				TotalAgentSeparationForce += SeparationForce;
+
+				// Agent avoidance
+				const float ContextualAvoidBuffer = FMath::Lerp(AvoidanceBufferAtEnd, AgentAvoidBuffer, NearEndFade);
+
+				// Calculate time of impact based on relative agent positions and velocities.
+				const float A = FVector::DotProduct(RelVel, RelVel);
+				const float Inv2A = A > 0.f ? 1.f / (2.f * A) : 0.f;
+				const float B = FMath::Min(0.f, 2.f * FVector::DotProduct(RelVel, RelPos));
+				const float C = FVector::DotProduct(RelPos, RelPos) - FMath::Square(TotalRadius - AvoidInset);
+				// Using max() here gives us CPA (closest point on arrival) when there is no hit.
+				const float Discr = FMath::Sqrt(FMath::Max(0.f, B * B - 4.f * A * C));
+				const float T0 = (-B - Discr) * Inv2A;
+				const float TOI = FMath::Clamp(T0, 0.f, TimeHoriz);
+
+				// Calculate penetration at CPA
+				const FVector AvoidRelPos = RelPos + RelVel * TOI;
+				const float AvoidDist = AvoidRelPos.Size();
+				const FVector AvoidConNormal = AvoidDist > 0.0f ? (AvoidRelPos / AvoidDist) : FVector::ForwardVector;
+
+				FVector AvoidNormal = AvoidConNormal;
+				if (bHasForcedNormal)
+				{
+					// The more head on the predicted collisions is, the more we should avoid towards the forced direction.
+					const FVector RelVelNorm = RelVel.GetSafeNormal();
+					const float Blend = FMath::Max(0.0f, -FVector::DotProduct(AvoidConNormal, RelVelNorm));
+					AvoidNormal = FMath::Lerp(AvoidConNormal, ForcedNormal, Blend).GetSafeNormal();
+				}
+				
+				const float AvoidPenetration = (TotalRadius - AvoidInset + ContextualAvoidBuffer) - AvoidDist; // Based on future agents distance
+				const float AvoidMag = FMath::Square(FMath::Clamp(AvoidPenetration / ContextualAvoidBuffer, 0.f, 1.f));
+				const float AvoidMagDist = (1.f - (TOI / TimeHoriz)); // No clamp, TOI is between 0 and TimeHoriz
+				const float AvoidReactMag = 1.f; // FMath::Min(closeAgent.seenTime / avoidReactionTime, 1.0f);			@todo: no seen time for now
+				const FVector AvoidForce = AvoidNormal * AvoidMag * AvoidMagDist * AvoidReactMag * AvoidKAgent * StandingScaling;
+
+				SteeringForce += AvoidForce;
+
+#if WITH_MASS_DEBUG && UNSAFE_FOR_MT
+				// Display close agent
+				UE::MassAvoidance::DebugDrawCylinder(AgentDebugContext, Collider.Location, Collider.Location + DebugLowCylinderOffset, Collider.Radius, AgentsColor);
+
+				if (bDebugIsOtherAgentUpdated)
+				{
+					UE::MassAvoidance::DebugDrawCylinder(BaseDebugContext, Collider.Location, Collider.Location + DebugAgentHeightOffset, Collider.Radius, FColor::Red);
+				}
+
+				// Draw agent contact separation force
+				UE::MassAvoidance::DebugDrawSummedForce(AgentDebugContext,
+					Collider.Location + DebugAgentHeightOffset,
+					Collider.Location + DebugAgentHeightOffset + SeparationForce,
+					AgentSeparationForceColor); 
+				
+				if (AvoidForce.Size() > 0.f)
+				{
+					// Draw agent vs agent hit positions
+					const FVector HitPosition = AgentLocation + (DesVel * TOI);
+					const FVector LeftOffset = AgentRadius * UE::MassAvoidance::GetLeftDirection(DesVel.GetSafeNormal(), FVector::UpVector);
+					UE::MassAvoidance::DebugDrawLine(AgentDebugContext, AgentLocation + DebugAgentHeightOffset + LeftOffset, HitPosition + DebugAgentHeightOffset + LeftOffset, CurrentAgentColor, 1.5f);
+					UE::MassAvoidance::DebugDrawLine(AgentDebugContext, AgentLocation + DebugAgentHeightOffset - LeftOffset, HitPosition + DebugAgentHeightOffset - LeftOffset, CurrentAgentColor, 1.5f);
+					UE::MassAvoidance::DebugDrawCylinder(AgentDebugContext, HitPosition, HitPosition + DebugAgentHeightOffset, AgentRadius, CurrentAgentColor);
+
+					const FVector OtherHitPosition = Collider.Location + (Collider.Velocity * TOI);
+					const FVector OtherLeftOffset = Collider.Radius * UE::MassAvoidance::GetLeftDirection(Collider.Velocity.GetSafeNormal(), FVector::UpVector);
+					const FVector Left = DebugAgentHeightOffset + OtherLeftOffset;
+					const FVector Right = DebugAgentHeightOffset - OtherLeftOffset;
+					UE::MassAvoidance::DebugDrawLine(AgentDebugContext, Collider.Location + Left, OtherHitPosition + Left, AgentsColor, 1.5f);
+					UE::MassAvoidance::DebugDrawLine(AgentDebugContext, Collider.Location + Right, OtherHitPosition + Right, AgentsColor, 1.5f);
+					UE::MassAvoidance::DebugDrawCylinder(AgentDebugContext, Collider.Location, Collider.Location + DebugAgentHeightOffset, AgentRadius, AgentsColor);
+					UE::MassAvoidance::DebugDrawCylinder(AgentDebugContext, OtherHitPosition, OtherHitPosition + DebugAgentHeightOffset, AgentRadius, AgentsColor);
+
+					// Draw agent avoid force
+					UE::MassAvoidance::DebugDrawForce(AgentDebugContext,
+						OtherHitPosition + DebugAgentHeightOffset,
+						OtherHitPosition + DebugAgentHeightOffset + AvoidForce,
+						AgentAvoidForceColor);
+				}
+#endif // WITH_MASS_DEBUG
+			} // close entities loop
+
+			if (MoveTarget.GetPreviousAction() != EMassMovementAction::Move)
+			{
+				// Fade in avoidance when transitioning from other than move action.
+				// I.e. the standing behavior may move the agents so close to each,
+				// and that causes the separation to push them out quickly when avoidance is activated. 
+				constexpr float FadeInTime = 1.0f; // @todo: make configurable.
+				const float AvoidanceFade = FMath::Min((CurrentTime - MoveTarget.GetCurrentActionStartTime()) / FadeInTime, 1.0f);
+				SteeringForce *= AvoidanceFade; 
+			}
+			
+			Steering.SteeringForce = UE::MassAvoidance::Clamp(SteeringForce, MaxSteerAccel); // Assume unit mass
+
+#if WITH_MASS_DEBUG && UNSAFE_FOR_MT
+			const FVector AgentAvoidSteeringForce = SteeringForce - OldSteeringForce;
+
+			// Draw total steering force to separate agents
+			UE::MassAvoidance::DebugDrawSummedForce(AgentDebugContext,
+				AgentLocation + DebugAgentHeightOffset,
+				AgentLocation + DebugAgentHeightOffset + TotalAgentSeparationForce,
+				AgentSeparationForceColor);
+
+			// Draw total steering force to avoid agents
+			UE::MassAvoidance::DebugDrawSummedForce(AgentDebugContext,
+				AgentLocation + DebugAgentHeightOffset,
+				AgentLocation + DebugAgentHeightOffset + AgentAvoidSteeringForce,
+				AgentAvoidForceColor);
+
+			// Draw final steering force adding to the agent velocity
+			UE::MassAvoidance::DebugDrawArrow(BaseDebugContext, 
+				AgentLocation + AgentVelocity + DebugAgentHeightOffset,
+				AgentLocation + AgentVelocity + DebugAgentHeightOffset + SteeringList[EntityIndex].SteeringForce,
+				FinalSteeringForceColor, BigArrowHeadSize, BigArrowThickness);
+#endif // WITH_MASS_DEBUG
+		}
+	});
+}
+
+//----------------------------------------------------------------------//
+//  UMassStandingAvoidanceProcessor
+//----------------------------------------------------------------------//
+UMassStandingAvoidanceProcessor::UMassStandingAvoidanceProcessor()
+{
+	bAutoRegisterWithProcessingPhases = true;
+	ExecutionFlags = (int32)EProcessorExecutionFlags::All;
+	ExecutionOrder.ExecuteInGroup = UE::Mass::ProcessorGroupNames::Avoidance;
+	ExecutionOrder.ExecuteAfter.Add(UE::Mass::ProcessorGroupNames::LOD);
+}
+
+void UMassStandingAvoidanceProcessor::ConfigureQueries()
+{
+	EntityQuery.AddRequirement<FMassSteeringGhostFragment>(ELWComponentAccess::ReadWrite);
+	EntityQuery.AddRequirement<FMassNavigationEdgesFragment>(ELWComponentAccess::ReadOnly);
+	EntityQuery.AddRequirement<FMassMoveTargetFragment>(ELWComponentAccess::ReadOnly);
+	EntityQuery.AddRequirement<FDataFragment_Transform>(ELWComponentAccess::ReadOnly);
+	EntityQuery.AddRequirement<FDataFragment_AgentRadius>(ELWComponentAccess::ReadOnly);
+	EntityQuery.AddTagRequirement<FMassHighLODTag>(ELWComponentPresence::All);
+}
+
+void UMassStandingAvoidanceProcessor::Initialize(UObject& Owner)
+{
+	Super::Initialize(Owner);
+
+	WeakWorld = Owner.GetWorld();
+	WeakMovementSubsystem = UWorld::GetSubsystem<UMassMovementSubsystem>(Owner.GetWorld());
+}
+
+void UMassStandingAvoidanceProcessor::Execute(UEntitySubsystem& EntitySubsystem, FLWComponentSystemExecutionContext& Context)
+{
+	QUICK_SCOPE_CYCLE_COUNTER(UMassStandingAvoidanceProcessor);
+
+	const UWorld* World = WeakWorld.Get();
+	UMassMovementSubsystem* MovementSubsystem = WeakMovementSubsystem.Get();
+	if (!World || !MovementSubsystem)
+	{
+		return;
+	}
+
+	// Avoidance while standing
+	EntityQuery.ForEachEntityChunk(EntitySubsystem, Context, [&EntitySubsystem, MovementSubsystem, World](FLWComponentSystemExecutionContext& Context)
+	{
+		const int32 NumEntities = Context.GetEntitiesNum();
+		const float DeltaTime = Context.GetDeltaTimeSeconds();
+		const float DistanceCutOffSquare = FMath::Square(UE::MassAvoidance::Tweakables::AgentDetectionDistance);
+
+		const TArrayView<FMassSteeringGhostFragment> GhostList = Context.GetMutableComponentView<FMassSteeringGhostFragment>();
+		const TConstArrayView<FDataFragment_Transform> LocationList = Context.GetComponentView<FDataFragment_Transform>();
+		const TConstArrayView<FDataFragment_AgentRadius> RadiusList = Context.GetComponentView<FDataFragment_AgentRadius>();
+		const TConstArrayView<FMassMoveTargetFragment> MoveTargetList = Context.GetComponentView<FMassMoveTargetFragment>();
+
+		// Arrays used to store close agents
+		TArray<FMassAvoidanceObstacleItem, TFixedAllocator<UE::MassAvoidance::MaxAgentResults>> CloseEntities;
+
+		struct FSortingAgent
+		{
+			FSortingAgent() = default;
+			FSortingAgent(const FLWEntity InEntity, const FVector InLocation, const FVector InForward, const float InDistSq) : Entity(InEntity), Location(InLocation), Forward(InForward), DistSq(InDistSq) {}
+			
+			FLWEntity Entity;
+			FVector Location = FVector::ZeroVector;
+			FVector Forward = FVector::ForwardVector;
+			float DistSq = 0.0f;
+		};
+		TArray<FSortingAgent, TFixedAllocator<UE::MassAvoidance::MaxAgentResults>> ClosestAgents;
+
+		for (int32 EntityIndex = 0; EntityIndex < NumEntities; ++EntityIndex)
+		{
+			// @todo: this should eventually be part of the query.
+			const FMassMoveTargetFragment& MoveTarget = MoveTargetList[EntityIndex];
+			if (MoveTarget.GetCurrentAction() != EMassMovementAction::Stand)
+			{
+				continue;
+			}
+			
+			FMassSteeringGhostFragment& Ghost = GhostList[EntityIndex];
+			// Skip if the ghost is not valid for this movement action yet.
+			if (Ghost.IsValid(MoveTarget.GetCurrentActionID()) == false)
+			{
+				continue;
+			}
+
+			const FDataFragment_Transform& Location = LocationList[EntityIndex];
+			const FDataFragment_AgentRadius& Radius = RadiusList[EntityIndex];
+
+			FLWEntity Entity = Context.GetEntity(EntityIndex);
+			const FVector AgentLocation = Location.GetTransform().GetTranslation();
+			const float AgentRadius = Radius.Radius;
+
+			// Steer ghost to move target.
+			constexpr float MaxSteerAccel = 300.0f;
+			constexpr float MaximumSpeed = 250.0f;
+			constexpr float StandDeadZoneRadius = 5.0f;
+			constexpr float StandSlowdownRadius = 15.0f;
+			constexpr float SteeringReactionTime = 2.0f;
+			constexpr float SteerK = 1.f / SteeringReactionTime;
+
+			FVector SteerDirection = FVector::ZeroVector;
+			FVector Delta = MoveTarget.Center - Ghost.Location;
+			Delta.Z = 0.0f;
+			const float Distance = Delta.Size();
+			if (Distance > KINDA_SMALL_NUMBER)
+			{
+				SteerDirection = Delta / Distance;
+			}
+			const float SpeedFade = FMath::Clamp((Distance - StandDeadZoneRadius) / FMath::Max(KINDA_SMALL_NUMBER, StandSlowdownRadius - StandDeadZoneRadius), 0.0f, 1.0f);
+							
+			const FVector DesiredVelocity = SteerDirection * MaximumSpeed * SpeedFade;
+			FVector SteeringForce = SteerK * (DesiredVelocity - Ghost.Velocity); // Goal force
+			
+			const FVector DesAcc = SteeringForce.GetClampedToMaxSize2D(MaxSteerAccel);
+
+			// Find close obstacles
+			const FAvoidanceObstacleHashGrid2D& AvoidanceObstacleGrid = MovementSubsystem->GetGridMutable();
+			UE::MassAvoidance::FindCloseAgents(AgentLocation, AvoidanceObstacleGrid, CloseEntities, UE::MassAvoidance::MaxAgentResults);
+
+			// Remove unwanted and find the closest in the CloseEntities
+			ClosestAgents.Reset();
+			for (const FAvoidanceObstacleHashGrid2D::ItemIDType OtherEntity : CloseEntities)
+			{
+				// Skip self
+				if (OtherEntity.Entity == Entity)
+				{
+					continue;
+				}
+
+				// Skip invalid entities.
+				if (!EntitySubsystem.IsEntityValid(OtherEntity.Entity))
+				{
+					UE_LOG(LogAvoidanceObstacles, VeryVerbose, TEXT("Close entity is invalid, skipped."));
+					continue;
+				}
+
+				// Skip too far
+				const FDataFragment_Transform& OtherTransform = EntitySubsystem.GetComponentDataChecked<FDataFragment_Transform>(OtherEntity.Entity);
+				const FVector OtherLocation = OtherTransform.GetTransform().GetLocation();
+				const float DistSq = FVector::DistSquared(AgentLocation, OtherLocation);
+				if (DistSq > DistanceCutOffSquare)
+				{
+					continue;
+				}
+
+				ClosestAgents.Emplace(OtherEntity.Entity, OtherLocation, OtherTransform.GetTransform().GetRotation().GetForwardVector(), DistSq);
+			}
+			ClosestAgents.Sort([](const FSortingAgent& A, const FSortingAgent& B) { return A.DistSq < B.DistSq; });
+
+			// Compute forces
+			const int32 MaxCloseAgentTreated = 6;
+			const int32 NumCloseAgents = FMath::Min(ClosestAgents.Num(), MaxCloseAgentTreated);
+			for (int32 Index = 0; Index < NumCloseAgents; Index++)
+			{
+				FSortingAgent& OtherAgent = ClosestAgents[Index];
+				FEntityView OtherEntityView(EntitySubsystem, OtherAgent.Entity);
+
+				const FMassVelocityFragment* OtherVelocityFragment = OtherEntityView.GetComponentDataPtr<FMassVelocityFragment>();
+				const float OtherRadius = OtherEntityView.GetComponentData<FDataFragment_AgentRadius>().Radius;
+
+				const float TotalRadius = AgentRadius + OtherRadius;
+
+				constexpr float GhostInset = 10.0f;
+				constexpr float MovingInset = -5.0f;
+
+				constexpr float AgentGhostSeparationBuffer = 20.0f;
+				constexpr float AgentMovingSeparationBuffer = 50.0f;
+				
+				constexpr float AgentKGhostSeparation = 200.0f;
+				constexpr float AgentKMovingSeparation = 500.0f;
+				constexpr float DirectionScaleStrength = 0.9f; // How strongly the direction scaling affects [0..1]
+
+				// @todo: this is heavy fragment to access, see if we could handle this differently.
+				const FMassMoveTargetFragment* OtherMoveTarget = OtherEntityView.GetComponentDataPtr<FMassMoveTargetFragment>();
+				const FMassSteeringGhostFragment* OtherGhost = OtherEntityView.GetComponentDataPtr<FMassSteeringGhostFragment>();
+
+				const bool bOtherHasGhost = OtherMoveTarget != nullptr && OtherGhost != nullptr
+											&& OtherMoveTarget->GetCurrentAction() == EMassMovementAction::Stand
+											&& OtherGhost->IsValid(OtherMoveTarget->GetCurrentActionID());
+				
+				// If other has ghost active, avoid that, else avoid the actual agent.
+				if (bOtherHasGhost)
+				{
+					// Avoid the other agent more, when it is further away from it's goal location.
+					const float ApproachDistance = 100.0f;
+					const float OtherDistanceToGoal = FVector::Distance(OtherGhost->Location, OtherMoveTarget->Center);
+					const float OtherSteerFade = FMath::Clamp(OtherDistanceToGoal / ApproachDistance, 0.0f, 1.0f);
+					const float SeparationK = FMath::Lerp(AgentKGhostSeparation, AgentKMovingSeparation, OtherSteerFade);
+					
+					// Ghost separation
+					FVector RelPos = Ghost.Location - OtherGhost->Location;
+					RelPos.Z = 0.f; // we assume we work on a flat plane for now
+					const float ConDist = RelPos.Size();
+					const FVector ConNorm = ConDist > 0.f ? RelPos / ConDist : FVector::ForwardVector;
+
+					// Separation force (stay away from agents if possible)
+					const float PenSep = (TotalRadius - GhostInset + AgentGhostSeparationBuffer) - ConDist;
+					const float SeparationMag = UE::MassAvoidance::Smoothf(FMath::Clamp(PenSep / AgentGhostSeparationBuffer, 0.f, 1.f));
+					const FVector SepForce = ConNorm * SeparationK;
+					const FVector SeparationForce = SepForce * SeparationMag;
+
+					SteeringForce += SeparationForce;
+				}
+				else
+				{
+					// Avoid more when the avoidance other is in front,
+					const FVector DirToOther = (OtherAgent.Location - Ghost.Location).GetSafeNormal();
+					const float DirectionScale = (1.0f - DirectionScaleStrength) + DirectionScaleStrength * FMath::Square(FMath::Max(0.0f, FVector::DotProduct(MoveTarget.Forward, DirToOther)));
+
+					// Treat the other agent as a capsule.
+					constexpr float RadiusToPersonalSpaceScale = 3.0f;
+					const FVector OtherBasePosition = OtherAgent.Location;
+					const FVector OtherPersonalSpacePosition = OtherAgent.Location + OtherAgent.Forward * OtherRadius * RadiusToPersonalSpaceScale * DirectionScale;
+					const FVector OtherLocation = FMath::ClosestPointOnSegment(Ghost.Location, OtherBasePosition, OtherPersonalSpacePosition);
+
+					FVector RelPos = Ghost.Location - OtherLocation;
+					RelPos.Z = 0.f;
+					const float ConDist = RelPos.Size();
+					const FVector ConNorm = ConDist > 0.f ? RelPos / ConDist : FVector::ForwardVector;
+
+					// Separation force (stay away from agents if possible)
+					const float PenSep = (TotalRadius - MovingInset + AgentMovingSeparationBuffer) - ConDist;
+					const float SeparationMag = UE::MassAvoidance::Smoothf(FMath::Clamp(PenSep / AgentMovingSeparationBuffer, 0.f, 1.f));
+					const FVector SepForce = ConNorm * AgentKMovingSeparation;
+					const FVector SeparationForce = SepForce * SeparationMag;
+
+					SteeringForce += SeparationForce;
+				}
+			}
+
+			SteeringForce.Z = 0.0f;
+			SteeringForce = UE::MassAvoidance::Clamp(SteeringForce, MaxSteerAccel); // Assume unit mass
+			Ghost.Velocity += SteeringForce * DeltaTime;
+			Ghost.Velocity.Z = 0.0f;
+			
+			// Damping
+			constexpr float VelocityDecayTime = 0.4f;
+			FMath::ExponentialSmoothingApprox(Ghost.Velocity, FVector::ZeroVector, DeltaTime, VelocityDecayTime);
+			
+			Ghost.Location += Ghost.Velocity * DeltaTime;
+
+			// Dont let the ghost location too far from move target center.
+			const float MaxDeviation = AgentRadius * 1.5f;
+			const FVector DirToCenter = Ghost.Location - MoveTarget.Center;
+			const float DistToCenter = DirToCenter.Length();
+			if (DistToCenter > MaxDeviation)
+			{
+				Ghost.Location = MoveTarget.Center + DirToCenter * (MaxDeviation / DistToCenter);
+			}
+
+		}
+	});
+	
+}
+
+//----------------------------------------------------------------------//
+//  UMassAvoidanceObstacleProcessor
+//----------------------------------------------------------------------//
+UMassAvoidanceObstacleProcessor::UMassAvoidanceObstacleProcessor()
+{
+	ExecutionFlags = (int32)EProcessorExecutionFlags::All;
+	ExecutionOrder.ExecuteInGroup = UE::Mass::ProcessorGroupNames::Avoidance;
+	ExecutionOrder.ExecuteBefore.Add(TEXT("MassAvoidanceProcessor"));
+}
+
+void UMassAvoidanceObstacleProcessor::ConfigureQueries()
+{
+	AddToGridEntityQuery.AddRequirement<FDataFragment_Transform>(ELWComponentAccess::ReadOnly);
+	AddToGridEntityQuery.AddRequirement<FDataFragment_AgentRadius>(ELWComponentAccess::ReadOnly);
+	AddToGridEntityQuery.AddRequirement<FMassAvoidanceObstacleGridCellLocationFragment>(ELWComponentAccess::ReadWrite);
+	UpdateGridEntityQuery = AddToGridEntityQuery;
+	RemoveFromGridEntityQuery = AddToGridEntityQuery;
+
+	AddToGridEntityQuery.AddTagRequirement<FMassOffLODTag>(ELWComponentPresence::None);
+	AddToGridEntityQuery.AddTagRequirement<FMassInAvoidanceObstacleGridTag>(ELWComponentPresence::None);
+
+	UpdateGridEntityQuery.AddTagRequirement<FMassOffLODTag>(ELWComponentPresence::None);
+	UpdateGridEntityQuery.AddTagRequirement<FMassInAvoidanceObstacleGridTag>(ELWComponentPresence::All);
+
+	RemoveFromGridEntityQuery.AddTagRequirement<FMassOffLODTag>(ELWComponentPresence::All);
+	RemoveFromGridEntityQuery.AddTagRequirement<FMassInAvoidanceObstacleGridTag>(ELWComponentPresence::All);
+}
+
+void UMassAvoidanceObstacleProcessor::Initialize(UObject& Owner)
+{
+	Super::Initialize(Owner);
+
+	WeakMovementSubsystem = UWorld::GetSubsystem<UMassMovementSubsystem>(Owner.GetWorld());
+}
+
+void UMassAvoidanceObstacleProcessor::Execute(UEntitySubsystem& EntitySubsystem, FLWComponentSystemExecutionContext& Context)
+{
+	using namespace UE::MassAvoidance;
+
+	UMassMovementSubsystem* MovementSubsystem = WeakMovementSubsystem.Get();
+	if (!MovementSubsystem)
+	{
+		return;
+	}
+
+	// can't be ParallelFor due to MovementSubsystem->GetGridMutable().Move not being thread-safe
+	AddToGridEntityQuery.ForEachEntityChunk(EntitySubsystem, Context, [this, MovementSubsystem, &EntitySubsystem](FLWComponentSystemExecutionContext& Context)
+		{
+			const int32 NumEntities = Context.GetEntitiesNum();
+
+			TConstArrayView<FDataFragment_Transform> LocationList = Context.GetComponentView<FDataFragment_Transform>();
+			TConstArrayView<FDataFragment_AgentRadius> RadiiList = Context.GetComponentView<FDataFragment_AgentRadius>();
+			TArrayView<FMassAvoidanceObstacleGridCellLocationFragment> AvoidanceObstacleCellLocationList = Context.GetMutableComponentView<FMassAvoidanceObstacleGridCellLocationFragment>();
+
+			for (int32 EntityIndex = 0; EntityIndex < NumEntities; ++EntityIndex)
+			{
+				// Add to the grid
+				const FVector NewPos = LocationList[EntityIndex].GetTransform().GetLocation();
+				const float Radius = RadiiList[EntityIndex].Radius;
+
+				FMassAvoidanceObstacleItem ObstacleItem;
+				ObstacleItem.Entity = Context.GetEntity(EntityIndex);
+				FEntityView EntityView(EntitySubsystem, ObstacleItem.Entity);
+				const FMassAvoidanceColliderFragment* Collider = EntityView.GetComponentDataPtr<FMassAvoidanceColliderFragment>();
+				if (Collider)
+				{
+					ObstacleItem.ItemFlags |= EMassAvoidanceObstacleItemFlags::HasColliderData;
+				}
+				
+				const FBox NewBounds(NewPos - FVector(Radius, Radius, 0.f), NewPos + FVector(Radius, Radius, 0.f));
+				AvoidanceObstacleCellLocationList[EntityIndex].CellLoc = MovementSubsystem->GetGridMutable().Add(ObstacleItem, NewBounds);
+
+				Context.Defer().AddTag<FMassInAvoidanceObstacleGridTag>(ObstacleItem.Entity);
+			}
+		});
+
+	UpdateGridEntityQuery.ForEachEntityChunk(EntitySubsystem, Context, [this, MovementSubsystem, &EntitySubsystem](FLWComponentSystemExecutionContext& Context)
+		{
+			const int32 NumEntities = Context.GetEntitiesNum();
+
+			TConstArrayView<FDataFragment_Transform> LocationList = Context.GetComponentView<FDataFragment_Transform>();
+			TConstArrayView<FDataFragment_AgentRadius> RadiiList = Context.GetComponentView<FDataFragment_AgentRadius>();
+			TArrayView<FMassAvoidanceObstacleGridCellLocationFragment> AvoidanceObstacleCellLocationList = Context.GetMutableComponentView<FMassAvoidanceObstacleGridCellLocationFragment>();
+
+			for (int32 EntityIndex = 0; EntityIndex < NumEntities; ++EntityIndex)
+			{
+				// Update position in grid
+				const FVector NewPos = LocationList[EntityIndex].GetTransform().GetLocation();
+				const float Radius = RadiiList[EntityIndex].Radius;
+				FMassAvoidanceObstacleItem ObstacleItem;
+				ObstacleItem.Entity = Context.GetEntity(EntityIndex);
+				FEntityView EntityView(EntitySubsystem, ObstacleItem.Entity);
+				const FMassAvoidanceColliderFragment* Collider = EntityView.GetComponentDataPtr<FMassAvoidanceColliderFragment>();
+				if (Collider)
+				{
+					ObstacleItem.ItemFlags |= EMassAvoidanceObstacleItemFlags::HasColliderData;
+				}
+				const FBox NewBounds(NewPos - FVector(Radius, Radius, 0.f), NewPos + FVector(Radius, Radius, 0.f));
+				AvoidanceObstacleCellLocationList[EntityIndex].CellLoc = MovementSubsystem->GetGridMutable().Move(ObstacleItem, AvoidanceObstacleCellLocationList[EntityIndex].CellLoc, NewBounds);
+
+#if WITH_MASS_DEBUG && 0
+				const FDebugContext BaseDebugContext(this, LogAvoidance, nullptr, ObstacleItem.Entity);
+				if (DebugIsSelected(ObstacleItem.Entity))
+				{
+					FBox Box = MovementSubsystem->GetGridMutable().CalcCellBounds(AvoidanceObstacleCellLocationList[EntityIndex].CellLoc);
+					Box.Max.Z += 200.f;					
+					DebugDrawBox(BaseDebugContext, Box, FColor::Yellow);
+				}
+#endif // WITH_MASS_DEBUG
+			}
+		});
+
+	RemoveFromGridEntityQuery.ForEachEntityChunk(EntitySubsystem, Context, [this, MovementSubsystem](FLWComponentSystemExecutionContext& Context)
+	{
+		const int32 NumEntities = Context.GetEntitiesNum();
+
+		TArrayView<FMassAvoidanceObstacleGridCellLocationFragment> AvoidanceObstacleCellLocationList = Context.GetMutableComponentView<FMassAvoidanceObstacleGridCellLocationFragment>();
+		for (int32 EntityIndex = 0; EntityIndex < NumEntities; ++EntityIndex)
+		{
+			FMassAvoidanceObstacleItem ObstacleItem;
+			ObstacleItem.Entity = Context.GetEntity(EntityIndex);
+			MovementSubsystem->GetGridMutable().Remove(ObstacleItem, AvoidanceObstacleCellLocationList[EntityIndex].CellLoc);
+			AvoidanceObstacleCellLocationList[EntityIndex].CellLoc = FAvoidanceObstacleHashGrid2D::FCellLocation();
+
+			Context.Defer().RemoveTag<FMassInAvoidanceObstacleGridTag>(ObstacleItem.Entity);
+		}
+	});
+}
+
+//----------------------------------------------------------------------//
+//  UMassNavigationBoundaryProcessor
+//----------------------------------------------------------------------//
+UMassNavigationBoundaryProcessor::UMassNavigationBoundaryProcessor()
+{
+	bAutoRegisterWithProcessingPhases = false;
+	ExecutionOrder.ExecuteInGroup = UE::Mass::ProcessorGroupNames::Avoidance;
+	ExecutionOrder.ExecuteAfter.Add(UE::Mass::ProcessorGroupNames::LOD);
+	ExecutionOrder.ExecuteBefore.Add(TEXT("MassAvoidanceProcessor"));
+}
+
+void UMassNavigationBoundaryProcessor::ConfigureQueries()
+{
+	EntityQuery.AddRequirement<FDataFragment_Transform>(ELWComponentAccess::ReadOnly);
+	EntityQuery.AddRequirement<FDataFragment_NavLocation>(ELWComponentAccess::ReadOnly);
+	EntityQuery.AddRequirement<FMassEdgeDetectionParamsFragment>(ELWComponentAccess::ReadOnly);
+	EntityQuery.AddRequirement<FMassNavigationEdgesFragment>(ELWComponentAccess::ReadWrite);
+}
+
+void UMassNavigationBoundaryProcessor::Initialize(UObject& Owner)
+{
+	Super::Initialize(Owner);
+
+	WeakWorld = Owner.GetWorld();
+	WeakMovementSubsystem = UWorld::GetSubsystem<UMassMovementSubsystem>(Owner.GetWorld());
+
+	ANavigationData* NavData = Cast<ANavigationData>(&Owner);
+	if (NavData == nullptr)
+	{
+		UNavigationSystemV1* NavSys = FNavigationSystem::GetCurrent<UNavigationSystemV1>(Owner.GetWorld());
+		NavData = NavSys ? NavSys->GetDefaultNavDataInstance() : nullptr;
+	}
+	WeakNavData = NavData;
+}
+
+void UMassNavigationBoundaryProcessor::Execute(UEntitySubsystem& EntitySubsystem, FLWComponentSystemExecutionContext& Context)
+{
+	using namespace UE::MassAvoidance;
+
+	const UWorld* World = WeakWorld.Get();
+	UMassMovementSubsystem* MovementSubsystem = WeakMovementSubsystem.Get();
+	ANavigationData* NavData = WeakNavData.Get();
+	if (!World || !MovementSubsystem || !NavData)
+	{
+		return;
+	}
+
+	EntityQuery.ForEachEntityChunk(EntitySubsystem, Context, [this, NavData](FLWComponentSystemExecutionContext& Context)
+		{
+			const int32 NumEntities = Context.GetEntitiesNum();
+			TConstArrayView<FDataFragment_Transform> LocationList = Context.GetComponentView<FDataFragment_Transform>();
+			TConstArrayView<FDataFragment_NavLocation> NavLocationList = Context.GetComponentView<FDataFragment_NavLocation>();
+			TConstArrayView<FMassEdgeDetectionParamsFragment> EdgeDetectionParamsList = Context.GetComponentView<FMassEdgeDetectionParamsFragment>();
+			TArrayView<FMassNavigationEdgesFragment> EdgesList = Context.GetMutableComponentView<FMassNavigationEdgesFragment>();
+
+			const ARecastNavMesh* RecastNavMesh = Cast<const ARecastNavMesh>(NavData);
+			if (!RecastNavMesh)
+			{
+				return;
+			}
+
+			TArray<FNavigationWallEdge> Edges;
+			Edges.Reserve(64);
+
+			for (int32 EntityIndex = 0; EntityIndex < NumEntities; ++EntityIndex)
+			{
+				RecastNavMesh->FindEdges(NavLocationList[EntityIndex].NodeRef, LocationList[EntityIndex].GetTransform().GetLocation(), EdgeDetectionParamsList[EntityIndex].EdgeDetectionRange, nullptr, Edges);
+
+				EdgesList[EntityIndex].AvoidanceEdges.Reset();
+				for (int32 Index = 0; Index < Edges.Num() && Index < FMassNavigationEdgesFragment::MaxEdgesCount; Index++)
+				{
+					const FNavigationWallEdge& Edge = Edges[Index];
+					EdgesList[EntityIndex].AvoidanceEdges.Add(FNavigationAvoidanceEdge(Edge.Start, Edge.End));
+				}
+			}
+		});
+}
+
+//----------------------------------------------------------------------//
+//  UMassLaneBoundaryProcessor
+//----------------------------------------------------------------------//
+UMassLaneBoundaryProcessor::UMassLaneBoundaryProcessor()
+{
+	ExecutionFlags = (int32)EProcessorExecutionFlags::All;
+
+	bAutoRegisterWithProcessingPhases = false;
+	ExecutionOrder.ExecuteInGroup = UE::Mass::ProcessorGroupNames::Avoidance;
+	ExecutionOrder.ExecuteAfter.Add(UE::Mass::ProcessorGroupNames::LOD);
+	ExecutionOrder.ExecuteBefore.Add(TEXT("MassAvoidanceProcessor"));
+}
+
+void UMassLaneBoundaryProcessor::ConfigureQueries()
+{
+	EntityQuery.AddRequirement<FDataFragment_Transform>(ELWComponentAccess::ReadOnly);							// need agent position to get closest point on lane
+	EntityQuery.AddRequirement<FMassSimulationLODFragment>(ELWComponentAccess::ReadOnly);						// LOD level
+	EntityQuery.AddRequirement<FMassNavigationEdgesFragment>(ELWComponentAccess::ReadWrite);					// output edges
+	EntityQuery.AddRequirement<FMassLastUpdatePositionFragment>(ELWComponentAccess::ReadWrite);					// to keep position when boundaries where last updated
+	EntityQuery.AddRequirement<FMassZoneGraphLaneLocationFragment>(ELWComponentAccess::ReadOnly);				// current lane location
+	EntityQuery.AddRequirement<FMassAvoidanceBoundaryLastLaneHandleFragment>(ELWComponentAccess::ReadWrite);	// keep track of the last used lane
+	EntityQuery.AddTagRequirement<FMassHighLODTag>(ELWComponentPresence::All);
+}
+
+void UMassLaneBoundaryProcessor::Initialize(UObject& Owner)
+{
+	Super::Initialize(Owner);
+
+	WeakWorld = Owner.GetWorld();
+	WeakZoneGraph = UWorld::GetSubsystem<UZoneGraphSubsystem>(Owner.GetWorld());
+}
+
+void UMassLaneBoundaryProcessor::Execute(UEntitySubsystem& EntitySubsystem, FLWComponentSystemExecutionContext& Context)
+{
+	using namespace UE::MassAvoidance;
+
+	const UWorld* World = WeakWorld.Get();
+	if (!World)
+	{
+		return;
+	}
+
+	const UZoneGraphSubsystem* ZoneGraphSubsystem = WeakZoneGraph.Get();
+	if (!ZoneGraphSubsystem)
+	{
+		return;
+	}
+
+	static const FColor PaleTurquoise = FColor(175, 238, 238);
+	static const FColor LaneColor = PaleTurquoise;
+
+	EntityQuery.ForEachEntityChunk(EntitySubsystem, Context, [this, World, ZoneGraphSubsystem](FLWComponentSystemExecutionContext& Context)
+		{
+			const int32 NumEntities = Context.GetEntitiesNum();
+
+			TConstArrayView<FDataFragment_Transform> LocationList = Context.GetComponentView<FDataFragment_Transform>();
+			TArrayView<FMassNavigationEdgesFragment> EdgesList = Context.GetMutableComponentView<FMassNavigationEdgesFragment>();
+			TArrayView<FMassLastUpdatePositionFragment> LastUpdatePositionList = Context.GetMutableComponentView<FMassLastUpdatePositionFragment>();
+			TConstArrayView<FMassZoneGraphLaneLocationFragment> LaneLocationList = Context.GetComponentView<FMassZoneGraphLaneLocationFragment>();
+			TArrayView<FMassAvoidanceBoundaryLastLaneHandleFragment> LastLaneHandleList = Context.GetMutableComponentView<FMassAvoidanceBoundaryLastLaneHandleFragment>();
+
+			TArray<FZoneGraphLinkedLane> LinkedLanes;
+			LinkedLanes.Reserve(4);
+		
+			for (int32 EntityIndex = 0; EntityIndex < NumEntities; ++EntityIndex)
+			{
+				// First check if we moved enough for an update
+				const FVector& Location = LocationList[EntityIndex].GetTransform().GetLocation();
+				const float DeltaDistSquared = FVector::DistSquared(Location, LastUpdatePositionList[EntityIndex].Value);
+				const float UpdateDistanceThresholdSquared = FMath::Square(50.f);
+
+				FZoneGraphLaneHandle& LastLaneHandle = LastLaneHandleList[EntityIndex].LaneHandle;
+				const FMassZoneGraphLaneLocationFragment& LaneLocationFragment = LaneLocationList[EntityIndex];
+
+				if (DeltaDistSquared < UpdateDistanceThresholdSquared && LaneLocationFragment.LaneHandle == LastLaneHandle)
+				{
+					// Not moved enough
+					continue;
+				}
+				else
+				{
+					LastUpdatePositionList[EntityIndex].Value = Location;
+					LastLaneHandle = LaneLocationFragment.LaneHandle;
+				}
+
+				// If we are skipping the update we don't want to reset the edges, we just want to execute up to the display of the lane.
+				FMassNavigationEdgesFragment& EdgesFragment = EdgesList[EntityIndex];
+				EdgesFragment.AvoidanceEdges.Reset();
+
+#if WITH_MASS_DEBUG && UNSAFE_FOR_MT
+				const FLWEntity Entity = Context.GetEntity(EntityIndex);
+				const FDebugContext BaseDebugContext(this, LogAvoidance, World, Entity);
+				const FDebugContext ObstacleDebugContext(this, LogAvoidanceObstacles, World, Entity);
+
+				if (DebugIsSelected(Entity))
+				{
+					if (!LaneLocationFragment.LaneHandle.IsValid())
+					{
+						DebugDrawSphere(ObstacleDebugContext, Location, /*Radius=*/100.f, FColor(128,128,128));
+					}
+				}
+#endif // WITH_MASS_DEBUG
+
+				// @todo: Fix transition between lanes is not smooth. WanderFragment.CurrentLaneLocation is in front of the actual
+				//		  position so when there is a lane switch picking the closest position on the lane jumps forward.
+				if (LaneLocationFragment.LaneHandle.IsValid())
+				{
+					const AZoneGraphData* Data = ZoneGraphSubsystem->GetZoneGraphData(LaneLocationFragment.LaneHandle.DataHandle);
+
+					if (!ensureMsgf(Data, TEXT("ZoneGraphData not found!")))
+					{
+						continue;
+					}
+
+					const FZoneGraphStorage& Storage = Data->GetStorage();
+
+					// Get nearest location on the current lane.
+
+					const FZoneGraphLaneHandle& LaneHandle = LaneLocationList[EntityIndex].LaneHandle;
+					FZoneGraphLaneLocation LaneLocation;
+					UE::ZoneGraph::Query::CalculateLocationAlongLane(Storage, LaneHandle, LaneLocationList[EntityIndex].DistanceAlongLane, LaneLocation);
+
+					if (LaneLocation.IsValid())
+					{
+#if WITH_MASS_DEBUG && UNSAFE_FOR_MT
+						if (DebugIsSelected(Entity))
+						{
+							// Draw the location found and the lane from that position to the end.
+							DebugDrawSphere(BaseDebugContext, LaneLocation.Position, /*Radius=*/10.f, LaneColor);
+							const FZoneLaneData& Lane = Storage.Lanes[LaneLocation.LaneHandle.Index];
+							float Progression = LaneLocation.DistanceAlongLane;
+							int32 LaneSegment = LaneLocation.LaneSegment;
+							const float DrawDistance = 0.5f * Tweakables::AgentDetectionDistance;
+							while ((Progression - LaneLocation.DistanceAlongLane) < DrawDistance && (LaneSegment < (Lane.PointsEnd - 1)))
+							{
+								Progression = Storage.LanePointProgressions[LaneSegment];
+								DebugDrawLine(BaseDebugContext, Storage.LanePoints[LaneSegment], Storage.LanePoints[LaneSegment + 1], LaneColor, /*Thickness=*/3.f);
+								LaneSegment++;
+							}
+						}
+#endif // WITH_MASS_DEBUG
+
+						// Get width of adjacent lanes.
+						float AdjacentLeftWidth = 0.f;
+						float AdjacentRightWidth = 0.f;
+						if (Tweakables::bUseAdjacentCorridors)
+						{
+							LinkedLanes.Reset();
+							UE::ZoneGraph::Query::GetLinkedLanes(Storage, LaneLocation.LaneHandle, EZoneLaneLinkType::Adjacent, EZoneLaneLinkFlags::Left|EZoneLaneLinkFlags::Right, EZoneLaneLinkFlags::None, LinkedLanes);
+							
+							for (const FZoneGraphLinkedLane& LinkedLane : LinkedLanes)
+							{
+								if (LinkedLane.HasFlags(EZoneLaneLinkFlags::Left))
+								{
+									const FZoneLaneData& Lane = Storage.Lanes[LinkedLane.DestLane.Index];
+									AdjacentLeftWidth += Lane.Width;
+								}
+								else if (LinkedLane.HasFlags(EZoneLaneLinkFlags::Right))
+								{
+									const FZoneLaneData& Lane = Storage.Lanes[LinkedLane.DestLane.Index];
+									AdjacentRightWidth += Lane.Width;
+								}
+							}
+						}
+
+						const FZoneLaneData& Lane = Storage.Lanes[LaneLocation.LaneHandle.Index];
+						const float HalfWidth = 0.5*Lane.Width;
+						const int32 Segment = LaneLocation.LaneSegment;
+
+						static const int32 MaxPoints = 4;
+						FVector Points[MaxPoints];
+						FVector SegmentDirections[MaxPoints];
+						FVector LeftDirections[MaxPoints];
+						FVector MiterDirections[MaxPoints];
+						const int32 FirstSegment = FMath::Max(Lane.PointsBegin, Segment - 1); // Segment should always be <= Lane.PointsEnd - 2
+						check(Lane.GetNumPoints() >= 2);
+						const int32 LastSegment = FMath::Min(Segment + 1, Lane.GetLastPoint() - 1); // Lane.PointsEnd - 1 is the lane last point, Lane.PointsEnd - 2 is the lane last segment
+						const int32 NumPoints = (LastSegment - FirstSegment + 1) + 1; //NumPoint = NumSegment + 1
+						check(NumPoints >= 2);
+						check(NumPoints <= MaxPoints);
+						
+						// Get points
+						for (int32 Index = 0; Index < NumPoints; Index++)
+						{
+							Points[Index] = Storage.LanePoints[FirstSegment + Index];
+						}
+						
+						// Calculate segment forward and left directions.
+						for (int32 Index = 0; Index < NumPoints - 1; Index++)
+						{
+							SegmentDirections[Index] = (Points[Index + 1] - Points[Index]).GetSafeNormal();
+							const FVector Up = Storage.LaneUpVectors[FirstSegment + Index];
+							LeftDirections[Index] = GetLeftDirection(SegmentDirections[Index], Up);
+						}
+
+						// Last point inherits the direction from the last segment.
+						SegmentDirections[NumPoints - 1] = SegmentDirections[NumPoints - 2];
+						LeftDirections[NumPoints - 1] = LeftDirections[NumPoints - 2];
+
+						// Calculate miter directions at inner corners.
+						// Note, mitered direction is average of the adjacent edge left directions, and scaled so that the expanded edges are parallel to the stem.
+						// First and last point dont have adjacent segments, and not mitered.
+						MiterDirections[0] = LeftDirections[0];
+						MiterDirections[NumPoints - 1] = LeftDirections[NumPoints - 1];
+						for (int32 Index = 1; Index < NumPoints - 1; Index++)
+						{
+							MiterDirections[Index] = ComputeMiterDirection(LeftDirections[Index - 1], LeftDirections[Index]);
+						}
+
+						// Compute left and right positions from lane width and miter directions.
+						const float LeftWidth = HalfWidth + AdjacentLeftWidth;
+						const float RightWidth = HalfWidth + AdjacentRightWidth;
+						FVector LeftPositions[MaxPoints];
+						FVector RightPositions[MaxPoints];
+						for (int32 Index = 0; Index < NumPoints; Index++)
+						{
+							const FVector MiterDir = MiterDirections[Index];
+							LeftPositions[Index] = Points[Index] + LeftWidth * MiterDir;
+							RightPositions[Index] = Points[Index] - RightWidth * MiterDir;
+						}
+						int32 NumLeftPositions = NumPoints;
+						int32 NumRightPositions = NumPoints;
+
+
+#if 0 && WITH_MASS_DEBUG // Detailed debug disabled
+						if (DebugIsSelected(Entity))
+						{
+							float Radius = 2.f;
+							for (int32 Index = 0; Index < NumPoints; Index++)
+							{
+								if (Index < NumPoints - 1)
+								{
+									DebugDrawLine(ObstacleDebugContext, Points[Index], Points[Index + 1], FColor::Blue, /*Thickness=*/6.f);
+								}
+								DebugDrawSphere(ObstacleDebugContext, Points[Index], Radius, FColor::Blue);
+								DebugDrawSphere(ObstacleDebugContext, LeftPositions[Index], Radius, FColor::Green);
+								DebugDrawSphere(ObstacleDebugContext, RightPositions[Index], Radius, FColor::Red);
+								Radius += 4.f;
+							}
+						}
+#endif //WITH_MASS_DEBUG
+
+						// Remove edges crossing when there are 3 eges.
+						if (NumPoints == 4)
+						{
+							FVector Intersection = FVector::ZeroVector;
+							if (FMath::SegmentIntersection2D(LeftPositions[0], LeftPositions[1], LeftPositions[2], LeftPositions[3], Intersection))
+							{
+								LeftPositions[1] = Intersection;
+								LeftPositions[2] = LeftPositions[3];
+								NumLeftPositions--;
+							}
+
+							Intersection = FVector::ZeroVector;
+							if (FMath::SegmentIntersection2D(RightPositions[0], RightPositions[1], RightPositions[2], RightPositions[3], Intersection))
+							{
+								RightPositions[1] = Intersection;
+								RightPositions[2] = RightPositions[3];
+								NumRightPositions--;
+							}
+						}
+
+						// Add edges
+						for (int32 Index = 0; Index < NumLeftPositions - 1; Index++)
+						{
+							EdgesFragment.AvoidanceEdges.Add(FNavigationAvoidanceEdge(LeftPositions[Index + 1], LeftPositions[Index])); // Left side: reverse start and end to keep the normal inside.
+						}
+
+						for (int32 Index = 0; Index < NumRightPositions - 1; Index++)
+						{
+							EdgesFragment.AvoidanceEdges.Add(FNavigationAvoidanceEdge(RightPositions[Index], RightPositions[Index + 1]));
+						}
+					}
+				}
+			}
+		});
+}
+
+//----------------------------------------------------------------------//
+//  UMassLaneCacheBoundaryProcessor
+//----------------------------------------------------------------------//
+UMassLaneCacheBoundaryProcessor::UMassLaneCacheBoundaryProcessor()
+{
+	ExecutionFlags = (int32)EProcessorExecutionFlags::All;
+
+	bAutoRegisterWithProcessingPhases = true;
+	ExecutionOrder.ExecuteInGroup = UE::Mass::ProcessorGroupNames::Avoidance;
+	ExecutionOrder.ExecuteAfter.Add(UE::Mass::ProcessorGroupNames::LOD);
+	ExecutionOrder.ExecuteBefore.Add(TEXT("MassAvoidanceProcessor"));
+}
+
+void UMassLaneCacheBoundaryProcessor::ConfigureQueries()
+{
+	EntityQuery.AddRequirement<FMassZoneGraphCachedLaneFragment>(ELWComponentAccess::ReadOnly);
+	EntityQuery.AddRequirement<FMassMoveTargetFragment>(ELWComponentAccess::ReadOnly);
+	EntityQuery.AddRequirement<FMassZoneGraphLaneLocationFragment>(ELWComponentAccess::ReadOnly);
+	EntityQuery.AddRequirement<FMassSimulationLODFragment>(ELWComponentAccess::ReadOnly);
+	EntityQuery.AddRequirement<FMassLaneCacheBoundaryFragment>(ELWComponentAccess::ReadWrite);
+	EntityQuery.AddRequirement<FMassNavigationEdgesFragment>(ELWComponentAccess::ReadWrite);	// output edges
+	EntityQuery.AddTagRequirement<FMassOffLODTag>(ELWComponentPresence::None);
+}
+
+void UMassLaneCacheBoundaryProcessor::Initialize(UObject& Owner)
+{
+	Super::Initialize(Owner);
+
+	WeakWorld = Owner.GetWorld();
+}
+
+void UMassLaneCacheBoundaryProcessor::Execute(UEntitySubsystem& EntitySubsystem, FLWComponentSystemExecutionContext& Context)
+{
+	using namespace UE::MassAvoidance;
+
+	QUICK_SCOPE_CYCLE_COUNTER(MassLaneCacheBoundaryProcessor);
+
+	const UWorld* World = WeakWorld.Get();
+	if (!World)
+	{
+		return;
+	}
+
+	EntityQuery.ForEachEntityChunk(EntitySubsystem, Context, [this, World](FLWComponentSystemExecutionContext& Context)
+	{
+		const int32 NumEntities = Context.GetEntitiesNum();
+
+		TConstArrayView<FMassZoneGraphCachedLaneFragment> CachedLaneList = Context.GetComponentView<FMassZoneGraphCachedLaneFragment>();
+		TConstArrayView<FMassZoneGraphLaneLocationFragment> LaneLocationList = Context.GetComponentView<FMassZoneGraphLaneLocationFragment>();
+		TConstArrayView<FMassMoveTargetFragment> MovementTargetList = Context.GetComponentView<FMassMoveTargetFragment>();
+		TArrayView<FMassLaneCacheBoundaryFragment> LaneCacheBoundaryList = Context.GetMutableComponentView<FMassLaneCacheBoundaryFragment>();
+		TArrayView<FMassNavigationEdgesFragment> EdgesList = Context.GetMutableComponentView<FMassNavigationEdgesFragment>();
+
+		TArray<FZoneGraphLinkedLane> LinkedLanes;
+		LinkedLanes.Reserve(4);	
+	
+		for (int32 EntityIndex = 0; EntityIndex < NumEntities; ++EntityIndex)
+		{
+			const FMassZoneGraphCachedLaneFragment& CachedLane = CachedLaneList[EntityIndex];
+			const FMassZoneGraphLaneLocationFragment& LaneLocation = LaneLocationList[EntityIndex];
+			const FMassMoveTargetFragment& MovementTarget = MovementTargetList[EntityIndex];
+			FMassNavigationEdgesFragment& Edges = EdgesList[EntityIndex];
+			FMassLaneCacheBoundaryFragment& LaneCacheBoundary = LaneCacheBoundaryList[EntityIndex];
+			const FLWEntity Entity = Context.GetEntity(EntityIndex);
+
+			// First check if we moved enough for an update
+			const float DeltaDistSquared = FVector::DistSquared(MovementTarget.Center, LaneCacheBoundary.LastUpdatePosition);
+			const float UpdateDistanceThresholdSquared = FMath::Square(50.f);
+
+#if WITH_MASS_DEBUG
+			const FDebugContext ObstacleDebugContext(this, LogAvoidanceObstacles, World, Entity);
+			if (DebugIsSelected(Entity))
+			{
+				DebugDrawSphere(ObstacleDebugContext, LaneCacheBoundary.LastUpdatePosition, /*Radius=*/10.f, FColor(128,128,128));
+				DebugDrawSphere(ObstacleDebugContext, MovementTarget.Center, /*Radius=*/10.f, FColor(255,255,255));
+			}
+#endif
+
+			if (DeltaDistSquared < UpdateDistanceThresholdSquared && CachedLane.CacheID == LaneCacheBoundary.LastUpdateCacheID)
+			{
+				// Not moved enough
+				continue;
+			}
+
+			LaneCacheBoundary.LastUpdatePosition = MovementTarget.Center;
+			LaneCacheBoundary.LastUpdateCacheID = CachedLane.CacheID;
+
+			// If we are skipping the update we don't want to reset the edges, we just want to execute up to the display of the lane.
+			Edges.AvoidanceEdges.Reset();
+			if (CachedLane.NumPoints < 2)
+			{
+				// Nothing to do
+				continue;
+			}
+			
+
+#if WITH_MASS_DEBUG
+			if (DebugIsSelected(Entity))
+			{
+				DebugDrawSphere(ObstacleDebugContext, MovementTarget.Center, /*Radius=*/100.f, FColor(128,128,128));
+			}
+#endif // WITH_MASS_DEBUG
+
+			const float HalfWidth = 0.5 * CachedLane.LaneWidth.Get();
+
+			static const int32 MaxPoints = 4;
+			FVector Points[MaxPoints];
+			FVector SegmentDirections[MaxPoints];
+			FVector LeftDirections[MaxPoints];
+			FVector MiterDirections[MaxPoints];
+
+			const int32 CurrentSegment = CachedLane.FindSegmentIndexAtDistance(LaneLocation.DistanceAlongLane);
+			const int32 FirstSegment = FMath::Max(0, CurrentSegment - 1); // Segment should always be <= CachedLane.NumPoints - 2
+			const int32 LastSegment = FMath::Min(CurrentSegment + 1, (int32)CachedLane.NumPoints - 2); // CachedLane.NumPoints - 1 is the lane last point, CachedLane.NumPoints - 2 is the lane last segment
+			const int32 NumPoints = (LastSegment - FirstSegment + 1) + 1; // NumPoint = NumSegment + 1
+			check(NumPoints >= 2);
+			check(NumPoints <= MaxPoints);
+			
+			// Get points
+			for (int32 Index = 0; Index < NumPoints; Index++)
+			{
+				Points[Index] = CachedLane.LanePoints[Index];
+			}
+			
+			// Calculate segment forward and left directions.  
+			for (int32 Index = 0; Index < NumPoints - 1; Index++)
+			{
+				SegmentDirections[Index] = (Points[Index + 1] - Points[Index]).GetSafeNormal();
+				LeftDirections[Index] = GetLeftDirection(SegmentDirections[Index], FVector::UpVector);
+			}
+
+			// Last point inherits the direction from the last segment.
+			SegmentDirections[NumPoints - 1] = SegmentDirections[NumPoints - 2];
+			LeftDirections[NumPoints - 1] = LeftDirections[NumPoints - 2];
+
+			// Calculate miter directions at inner corners.
+			// Note, mitered direction is average of the adjacent edge left directions, and scaled so that the expanded edges are parallel to the stem.
+			// First and last point dont have adjacent segments, and not mitered.
+			MiterDirections[0] = LeftDirections[0];
+			MiterDirections[NumPoints - 1] = LeftDirections[NumPoints - 1];
+			for (int32 Index = 1; Index < NumPoints - 1; Index++)
+			{
+				MiterDirections[Index] = ComputeMiterDirection(LeftDirections[Index - 1], LeftDirections[Index]);
+			}
+
+			// Compute left and right positions from lane width and miter directions.
+			const float LeftWidth = HalfWidth + CachedLane.LaneLeftSpace.Get();
+			const float RightWidth = HalfWidth + CachedLane.LaneRightSpace.Get();
+			FVector LeftPositions[MaxPoints];
+			FVector RightPositions[MaxPoints];
+			for (int32 Index = 0; Index < NumPoints; Index++)
+			{
+				const FVector MiterDir = MiterDirections[Index];
+				LeftPositions[Index] = Points[Index] + LeftWidth * MiterDir;
+				RightPositions[Index] = Points[Index] - RightWidth * MiterDir;
+			}
+			int32 NumLeftPositions = NumPoints;
+			int32 NumRightPositions = NumPoints;
+
+
+#if 0 && WITH_MASS_DEBUG // Detailed debug disabled
+			if (DebugIsSelected(Entity))
+			{
+				float Radius = 2.f;
+				for (int32 Index = 0; Index < NumPoints; Index++)
+				{
+					if (Index < NumPoints - 1)
+					{
+						DebugDrawLine(ObstacleDebugContext, Points[Index], Points[Index + 1], FColor::Blue, /*Thickness=*/6.f);
+					}
+					DebugDrawSphere(ObstacleDebugContext, Points[Index], Radius, FColor::Blue);
+					DebugDrawSphere(ObstacleDebugContext, LeftPositions[Index], Radius, FColor::Green);
+					DebugDrawSphere(ObstacleDebugContext, RightPositions[Index], Radius, FColor::Red);
+					Radius += 4.f;
+				}
+			}
+#endif //WITH_MASS_DEBUG
+
+			// Remove edges crossing when there are 3 edges.
+			if (NumPoints == 4)
+			{
+				FVector Intersection = FVector::ZeroVector;
+				if (FMath::SegmentIntersection2D(LeftPositions[0], LeftPositions[1], LeftPositions[2], LeftPositions[3], Intersection))
+				{
+					LeftPositions[1] = Intersection;
+					LeftPositions[2] = LeftPositions[3];
+					NumLeftPositions--;
+				}
+
+				Intersection = FVector::ZeroVector;
+				if (FMath::SegmentIntersection2D(RightPositions[0], RightPositions[1], RightPositions[2], RightPositions[3], Intersection))
+				{
+					RightPositions[1] = Intersection;
+					RightPositions[2] = RightPositions[3];
+					NumRightPositions--;
+				}
+			}
+
+			// Add edges
+			for (int32 Index = 0; Index < NumLeftPositions - 1; Index++)
+			{
+				Edges.AvoidanceEdges.Add(FNavigationAvoidanceEdge(LeftPositions[Index + 1], LeftPositions[Index])); // Left side: reverse start and end to keep the normal inside.
+			}
+
+			for (int32 Index = 0; Index < NumRightPositions - 1; Index++)
+			{
+				Edges.AvoidanceEdges.Add(FNavigationAvoidanceEdge(RightPositions[Index], RightPositions[Index + 1]));
+			}
+		}
+	});
+}
+
+#undef UNSAFE_FOR_MT
diff --git a/Engine/Plugins/AI/MassAI/Source/MassAIMovement/Private/MassAvoidanceTrait.cpp b/Engine/Plugins/AI/MassAI/Source/MassAIMovement/Private/MassAvoidanceTrait.cpp
new file mode 100644
--- /dev/null
+++ b/Engine/Plugins/AI/MassAI/Source/MassAIMovement/Private/MassAvoidanceTrait.cpp
@@ -0,0 +1,47 @@
+// Copyright Epic Games, Inc. All Rights Reserved.
+
+#include "MassAvoidanceTrait.h"
+#include "MassEntityTemplateRegistry.h"
+#include "Engine/World.h"
+#include "MassAvoidanceProcessors.h"
+#include "MassCommonFragments.h"
+#include "MassAIMovementFragments.h"
+#include "MassSimulationLOD.h"
+#include "MassZoneGraphMovementFragments.h"
+
+void UMassAvoidanceTrait::BuildTemplate(FMassEntityTemplateBuildContext& BuildContext, UWorld& World) const
+{
+	// UMassAvoidanceProcessor
+	BuildContext.AddFragmentWithDefaultInitializer<FDataFragment_AgentRadius>(); // From UMassProcessor_AgentMovement
+	BuildContext.AddFragmentWithDefaultInitializer<FMassNavigationEdgesFragment>();
+	BuildContext.AddFragmentWithDefaultInitializer<FMassMoveTargetFragment>();
+	BuildContext.AddFragmentWithDefaultInitializer<FDataFragment_Transform>(); // From UMassProcessor_Movement
+	BuildContext.AddFragmentWithDefaultInitializer<FMassVelocityFragment>(); // From UMassProcessor_Movement
+	BuildContext.AddFragmentWithDefaultInitializer<FMassSteeringFragment>();
+	
+	// LOD: UMassSimulationLODProcessors
+	BuildContext.AddFragmentWithDefaultInitializer<FDataFragment_Transform>();
+	BuildContext.AddFragmentWithDefaultInitializer<FMassLODInfoFragment>();
+	BuildContext.AddFragmentWithDefaultInitializer<FMassSimulationLODFragment>();
+	BuildContext.AddChunkFragment<FMassSimulationVariableTickChunkFragment>();
+
+	if (bUseZoneGraphMovement)
+	{
+		// UMassLaneCacheBoundaryProcessor
+		BuildContext.AddFragmentWithDefaultInitializer<FMassZoneGraphCachedLaneFragment>();
+		BuildContext.AddFragmentWithDefaultInitializer<FMassMoveTargetFragment>();
+		BuildContext.AddFragmentWithDefaultInitializer<FMassZoneGraphLaneLocationFragment>();
+		BuildContext.AddFragmentWithDefaultInitializer<FMassNavigationEdgesFragment>();
+		BuildContext.AddFragmentWithDefaultInitializer<FMassLaneCacheBoundaryFragment>();
+	}
+	else
+	{
+		// Put non-ZG movement specific fragments here
+		// UMassLaneBoundaryProcessor
+		BuildContext.AddFragmentWithDefaultInitializer<FDataFragment_Transform>();
+		BuildContext.AddFragmentWithDefaultInitializer<FMassNavigationEdgesFragment>();
+		BuildContext.AddFragmentWithDefaultInitializer<FMassLastUpdatePositionFragment>();
+		BuildContext.AddFragmentWithDefaultInitializer<FMassZoneGraphLaneLocationFragment>();
+		BuildContext.AddFragmentWithDefaultInitializer<FMassAvoidanceBoundaryLastLaneHandleFragment>();
+	}
+}
diff --git a/Engine/Plugins/AI/MassAI/Source/MassAIMovement/Private/MassMovementDelegates.cpp b/Engine/Plugins/AI/MassAI/Source/MassAIMovement/Private/MassMovementDelegates.cpp
new file mode 100644
--- /dev/null
+++ b/Engine/Plugins/AI/MassAI/Source/MassAIMovement/Private/MassMovementDelegates.cpp
@@ -0,0 +1,11 @@
+// Copyright Epic Games, Inc. All Rights Reserved.
+
+#include "MassMovementDelegates.h"
+#include "CoreMinimal.h"
+
+namespace UE::MassMovement::Delegates
+{
+#if WITH_EDITOR
+	FOnMassMovementNamesChanged MASSAIMOVEMENT_API OnMassMovementNamesChanged;
+#endif
+} // UEUE::MassMovement::Delegates
