AggregateTicking and Pipe merged into Engine/Plugins/Runtime/MassEntity

diff --git a/Engine/Plugins/AI/MassAI/Source/MassAIMovement/Private/MassMovementModule.cpp b/Engine/Plugins/AI/MassAI/Source/MassAIMovement/Private/MassMovementModule.cpp
new file mode 100644
--- /dev/null
+++ b/Engine/Plugins/AI/MassAI/Source/MassAIMovement/Private/MassMovementModule.cpp
@@ -0,0 +1,42 @@
+// Copyright Epic Games, Inc. All Rights Reserved.
+
+#include "CoreMinimal.h"
+#include "IMassMovementModule.h"
+#include "UObject/CoreRedirects.h"
+
+
+class FMassMovementModule : public IMassMovementModule
+{
+	/** IModuleInterface implementation */
+	virtual void StartupModule() override;
+	virtual void ShutdownModule() override;
+};
+
+IMPLEMENT_MODULE(FMassMovementModule, MassMovement)
+
+
+
+void FMassMovementModule::StartupModule()
+{
+	// This code will execute after your module is loaded into memory (but after global variables are initialized, of course.)
+
+	TArray<FCoreRedirect> Redirects;
+	Redirects.Emplace(ECoreRedirectFlags::Type_Struct, TEXT("VelocityComponent"), TEXT("MassVelocityFragment"));
+	Redirects.Emplace(ECoreRedirectFlags::Type_Struct, TEXT("AvoidanceComponent"), TEXT("MassAvoidanceFragment"));
+	Redirects.Emplace(ECoreRedirectFlags::Type_Struct, TEXT("DataFragment_GridCellLocation"), TEXT("MassGridCellLocationFragment"));
+	Redirects.Emplace(ECoreRedirectFlags::Type_Struct, TEXT("DataFragment_PathFollow"), TEXT("MassPathFollowFragment"));
+	Redirects.Emplace(ECoreRedirectFlags::Type_Struct, TEXT("DataFragment_Wander"), TEXT("MassWanderFragment"));
+	Redirects.Emplace(ECoreRedirectFlags::Type_Struct, TEXT("DataFragment_MoveTo"), TEXT("MassMoveToFragment"));
+	
+	FCoreRedirects::AddRedirectList(Redirects, TEXT("MassMovement"));
+}
+
+
+void FMassMovementModule::ShutdownModule()
+{
+	// This function may be called during shutdown to clean up your module.  For modules that support dynamic reloading,
+	// we call this function before unloading the module.
+}
+
+
+
diff --git a/Engine/Plugins/AI/MassAI/Source/MassAIMovement/Private/MassMovementSettings.cpp b/Engine/Plugins/AI/MassAI/Source/MassAIMovement/Private/MassMovementSettings.cpp
new file mode 100644
--- /dev/null
+++ b/Engine/Plugins/AI/MassAI/Source/MassAIMovement/Private/MassMovementSettings.cpp
@@ -0,0 +1,109 @@
+// Copyright Epic Games, Inc. All Rights Reserved.
+
+#include "MassMovementSettings.h"
+#include "MassMovementDelegates.h"
+#include "MassAIMovementTypes.h"
+
+//----------------------------------------------------------------------//
+// UMassMovementSettings
+//----------------------------------------------------------------------//
+
+void UMassMovementSettings::UpdateConfigs()
+{
+	// Update cached values on configs.
+	for (FMassMovementConfig& Config : MovementConfigs)
+	{
+		Config.Update();
+	}
+}
+
+void UMassMovementSettings::PostInitProperties()
+{
+	Super::PostInitProperties();
+	UpdateConfigs();
+}
+
+#if WITH_EDITOR
+void UMassMovementSettings::PostEditChangeChainProperty(FPropertyChangedChainEvent& PropertyChangedEvent)
+{
+	Super::PostEditChangeChainProperty(PropertyChangedEvent);
+
+	FProperty* Property = PropertyChangedEvent.Property;
+	FProperty* MemberProperty = nullptr;
+	if (PropertyChangedEvent.PropertyChain.GetActiveMemberNode())
+	{
+		MemberProperty = PropertyChangedEvent.PropertyChain.GetActiveMemberNode()->GetValue();
+	}
+
+	if (MemberProperty && Property)
+	{
+		if (MemberProperty->GetFName() == GET_MEMBER_NAME_CHECKED(UMassMovementSettings, MovementStyles))
+		{
+			const int32 ArrayIndex = PropertyChangedEvent.GetArrayIndex(MemberProperty->GetFName().ToString());
+
+			// Ensure unique ID on duplicated items.
+			if (PropertyChangedEvent.ChangeType == EPropertyChangeType::ArrayAdd)
+			{
+				if (MovementStyles.IsValidIndex(ArrayIndex))
+				{
+					MovementStyles[ArrayIndex].ID = FGuid::NewGuid();
+					MovementStyles[ArrayIndex].Name = FName(TEXT("Movement Style"));
+				}
+			}
+			else if (PropertyChangedEvent.ChangeType == EPropertyChangeType::Duplicate)
+			{
+				if (MovementStyles.IsValidIndex(ArrayIndex))
+				{
+					MovementStyles[ArrayIndex].ID = FGuid::NewGuid();
+					MovementStyles[ArrayIndex].Name = FName(MovementStyles[ArrayIndex].Name.ToString() + TEXT(" Duplicate"));
+				}
+			}
+
+			UpdateConfigs();
+			UE::MassMovement::Delegates::OnMassMovementNamesChanged.Broadcast();
+		}
+
+		if (MemberProperty->GetFName() == GET_MEMBER_NAME_CHECKED(UMassMovementSettings, MovementConfigs))
+		{
+			const int32 ArrayIndex = PropertyChangedEvent.GetArrayIndex(MemberProperty->GetFName().ToString());
+
+			// Ensure unique ID on duplicated items.
+			if (PropertyChangedEvent.ChangeType == EPropertyChangeType::ArrayAdd)
+			{
+				if (MovementConfigs.IsValidIndex(ArrayIndex))
+				{
+					MovementConfigs[ArrayIndex].ID = FGuid::NewGuid();
+					MovementConfigs[ArrayIndex].Name = FName(TEXT("Movement Config"));
+				}
+			}
+			else if (PropertyChangedEvent.ChangeType == EPropertyChangeType::Duplicate)
+			{
+				if (MovementConfigs.IsValidIndex(ArrayIndex))
+				{
+					MovementConfigs[ArrayIndex].ID = FGuid::NewGuid();
+					MovementConfigs[ArrayIndex].Name = FName(MovementConfigs[ArrayIndex].Name.ToString() + TEXT(" Duplicate"));
+				}
+			}
+
+			UpdateConfigs();
+			UE::MassMovement::Delegates::OnMassMovementNamesChanged.Broadcast();
+		}
+	}
+}
+#endif // WITH_EDITOR
+
+const FMassMovementStyle* UMassMovementSettings::GetMovementStyleByID(const FGuid ID) const
+{
+	return MovementStyles.FindByPredicate([ID](const FMassMovementStyle& Style) { return Style.ID == ID; });
+}
+	
+const FMassMovementConfig* UMassMovementSettings::GetMovementConfigByID(const FGuid ID) const
+{
+	return MovementConfigs.FindByPredicate([ID](const FMassMovementConfig& Config) { return Config.ID == ID; });
+}
+
+FMassMovementConfigHandle UMassMovementSettings::GetMovementConfigHandleByID(const FGuid ID) const
+{
+	const int32 Index = MovementConfigs.IndexOfByPredicate([ID](const FMassMovementConfig& Config) { return Config.ID == ID; });
+	return Index != INDEX_NONE ? FMassMovementConfigHandle((uint8)Index) : FMassMovementConfigHandle();
+}
diff --git a/Engine/Plugins/AI/MassAI/Source/MassAIMovement/Private/MassMovementSubsystem.cpp b/Engine/Plugins/AI/MassAI/Source/MassAIMovement/Private/MassMovementSubsystem.cpp
new file mode 100644
--- /dev/null
+++ b/Engine/Plugins/AI/MassAI/Source/MassAIMovement/Private/MassMovementSubsystem.cpp
@@ -0,0 +1,24 @@
+// Copyright Epic Games, Inc. All Rights Reserved.
+
+#include "MassMovementSubsystem.h"
+#include "Engine/World.h"
+#include "MassSimulationSubsystem.h"
+#include "ZoneGraphSubsystem.h"
+
+//----------------------------------------------------------------------//
+// UMassMovementSubsystem
+//----------------------------------------------------------------------//
+UMassMovementSubsystem::UMassMovementSubsystem() : AvoidanceObstacleGrid(250.f) // 2.5m grid
+{
+}
+
+void UMassMovementSubsystem::Initialize(FSubsystemCollectionBase& Collection)
+{
+	Super::Initialize(Collection);
+
+	Collection.InitializeDependency<UMassSimulationSubsystem>();
+	ZoneGraphSubsystem = Collection.InitializeDependency<UZoneGraphSubsystem>();
+
+	EntitySubsystem = UPipeEntitySubsystem::GetCurrent(GetWorld());
+}
+
diff --git a/Engine/Plugins/AI/MassAI/Source/MassAIMovement/Private/MassZoneGraphLaneObstacleAnnotations.cpp b/Engine/Plugins/AI/MassAI/Source/MassAIMovement/Private/MassZoneGraphLaneObstacleAnnotations.cpp
new file mode 100644
--- /dev/null
+++ b/Engine/Plugins/AI/MassAI/Source/MassAIMovement/Private/MassZoneGraphLaneObstacleAnnotations.cpp
@@ -0,0 +1,191 @@
+﻿// Copyright Epic Games, Inc. All Rights Reserved.
+
+#include "MassZoneGraphLaneObstacleAnnotations.h"
+
+#include "MassCommonTypes.h"
+#include "MassMovementTypes.h"
+#include "NavigationProcessor.h"
+#include "ZoneGraphAnnotationSubsystem.h"
+#include "ZoneGraphSubsystem.h"
+#include "ZoneGraphRenderingUtilities.h"
+#include "VisualLogger/VisualLogger.h"
+
+void UZoneGraphLaneObstacleAnnotations::PostSubsystemsInitialized()
+{
+	Super::PostSubsystemsInitialized();
+
+	ZoneGraphSubsystem = UWorld::GetSubsystem<UZoneGraphSubsystem>(GetWorld());
+	checkf(ZoneGraphSubsystem, TEXT("Expecting ZoneGraphSubsystem to be present."));
+}
+
+FZoneGraphTagMask UZoneGraphLaneObstacleAnnotations::GetAnnotationTags() const
+{
+	return FZoneGraphTagMask(LaneObstacleTag);
+}
+
+void UZoneGraphLaneObstacleAnnotations::HandleEvents(TConstArrayView<const UScriptStruct*> AllEventStructs, const FInstancedStructStream& Events)
+{
+	Events.ForEach([this](FStructView View)
+	{
+		if (const FZoneGraphLaneObstacleChangeEvent* LaneObstacleEvent = View.GetMutablePtr<FZoneGraphLaneObstacleChangeEvent>())
+		{
+			LaneObstacleChangeEvents.Add(*LaneObstacleEvent);
+		}
+	});
+}
+
+void UZoneGraphLaneObstacleAnnotations::TickAnnotation(const float DeltaTime, FZoneGraphAnnotationTagContainer& AnnotationTagContainer)
+{
+	// Process events
+	if (!LaneObstacleTag.IsValid())
+	{
+		return;
+	}
+
+	for (const FZoneGraphLaneObstacleChangeEvent& Event : LaneObstacleChangeEvents)
+	{
+		const FMassLaneObstacle& LaneObstacle = Event.LaneObstacle;
+
+		// Handle tag and update lane obstacle data.
+		if (Event.EventAction == EMassLaneObstacleEventAction::Add)
+		{
+			// Add tag.
+			const FZoneGraphLaneHandle LaneHandle = LaneObstacle.LaneSection.LaneHandle;
+			TArrayView<FZoneGraphTagMask> LaneTags = AnnotationTagContainer.GetMutableAnnotationTagsForData(LaneHandle.DataHandle);
+			LaneTags[LaneHandle.Index].Add(LaneObstacleTag);
+
+			// Add to obstacle container.
+			const FZoneGraphDataHandle ZoneGraphDataHandle = LaneHandle.DataHandle;
+			FMassLaneObstacleContainer& LaneObstaclesContainer = RegisteredLaneData[ZoneGraphDataHandle.Index].LaneObstacles;
+			LaneObstaclesContainer.Add(LaneObstacle);
+		}
+		else if (Event.EventAction == EMassLaneObstacleEventAction::Remove)
+		{
+			for (FMassRegisteredMovementLaneData& LaneData : RegisteredLaneData)
+			{
+				FMassLaneObstacleContainer& LaneObstaclesContainer = LaneData.LaneObstacles;
+				const FMassLaneObstacle* Obstacle = LaneObstaclesContainer.Find(LaneObstacle.GetID());
+				if (Obstacle)
+				{
+					// Remove from obstacle container.
+					const bool WasLast = LaneObstaclesContainer.Remove(*Obstacle);
+					if (WasLast)
+					{
+						// Remove tag.
+						const FZoneGraphLaneHandle LaneHandle = Obstacle->LaneSection.LaneHandle;
+						TArrayView<FZoneGraphTagMask> LaneTags = AnnotationTagContainer.GetMutableAnnotationTagsForData(LaneHandle.DataHandle);
+						LaneTags[LaneHandle.Index].Remove(LaneObstacleTag);
+					}
+				}
+			}
+		}
+
+#if !UE_BUILD_SHIPPING && !UE_BUILD_TEST
+		MarkRenderStateDirty();
+#endif
+
+	}
+	LaneObstacleChangeEvents.Reset();
+
+#if WITH_MASS_DEBUG && 0
+	for (const FMassRegisteredMovementLaneData& LaneData : RegisteredLaneData)
+	{
+		const FMassLaneObstacleContainer& LaneObstaclesContainer = LaneData.LaneObstacles;
+		for (const auto& Pair : LaneObstaclesContainer.LaneObstaclesMap)
+		{
+			for(const FMassLaneObstacle& Obstacle : Pair.Value)
+			{
+				FZoneGraphLaneLocation StartLaneLocation;
+				FZoneGraphLaneLocation EndLaneLocation;
+				FZoneGraphLaneSection Section = Obstacle.LaneSection;
+				if (ZoneGraphSubsystem->CalculateLocationAlongLane(Section.LaneHandle, Section.StartDistanceAlongLane, StartLaneLocation) &&
+					ZoneGraphSubsystem->CalculateLocationAlongLane(Section.LaneHandle, Section.EndDistanceAlongLane, EndLaneLocation))
+				{
+					UE_VLOG_LOCATION(this, LogMassDynamicObstacle, Log, StartLaneLocation.Position, 20.f, FColor::Red, TEXT(""));
+					UE_VLOG_LOCATION(this, LogMassDynamicObstacle, Log, EndLaneLocation.Position, 20.f, FColor::Red, TEXT(""));
+				}
+			}
+		}
+	}
+#endif // WITH_MASS_DEBUG 
+}
+
+void UZoneGraphLaneObstacleAnnotations::PostZoneGraphDataAdded(const AZoneGraphData& ZoneGraphData)
+{
+	const UWorld* World = GetWorld();
+
+	// Only consider valid graph from our world
+	if (ZoneGraphData.GetWorld() != World)
+	{
+		return;
+	}
+
+	const FString WorldName = World->GetName();
+	const FZoneGraphStorage& Storage = ZoneGraphData.GetStorage();
+	const int32 ZoneGraphDataIndex = Storage.DataHandle.Index;
+
+	UE_VLOG_UELOG(this, LogMassDynamicObstacle, Verbose, TEXT("%s adding obstacle lane data for zone graph %d/%d"), *WorldName, Storage.DataHandle.Index, Storage.DataHandle.Generation);
+
+	if (ZoneGraphDataIndex >= RegisteredLaneData.Num())
+	{
+		RegisteredLaneData.SetNum(ZoneGraphDataIndex + 1);
+	}
+
+	FMassRegisteredMovementLaneData& LaneData = RegisteredLaneData[ZoneGraphDataIndex];
+	LaneData.DataHandle = Storage.DataHandle;
+}
+
+void UZoneGraphLaneObstacleAnnotations::PreZoneGraphDataRemoved(const AZoneGraphData& ZoneGraphData)
+{
+	// Only consider valid graph from our world
+	if (ZoneGraphData.GetWorld() != GetWorld())
+	{
+		return;
+	}
+
+	const FZoneGraphStorage& Storage = ZoneGraphData.GetStorage();
+	const int32 Index = Storage.DataHandle.Index;
+	if (!RegisteredLaneData.IsValidIndex(Index))
+	{
+		return;
+	}
+
+	// Removing lane data.
+	FMassRegisteredMovementLaneData& LaneData = RegisteredLaneData[Index];
+	LaneData.Reset();
+
+#if !UE_BUILD_SHIPPING && !UE_BUILD_TEST
+	MarkRenderStateDirty();
+#endif
+}
+
+#if !UE_BUILD_SHIPPING && !UE_BUILD_TEST
+void UZoneGraphLaneObstacleAnnotations::DebugDraw(FZoneGraphAnnotationSceneProxy* DebugProxy)
+{
+	UZoneGraphSubsystem* ZoneGraph = UWorld::GetSubsystem<UZoneGraphSubsystem>(GetWorld());
+	if (!ZoneGraph)
+	{
+		return;
+	}
+	
+	for (const FMassRegisteredMovementLaneData& LaneData : RegisteredLaneData)
+	{
+		const FZoneGraphStorage* ZoneStorage = ZoneGraph->GetZoneGraphStorage(LaneData.DataHandle);
+		if (!ZoneStorage)
+		{
+			continue;
+		}
+
+		const FMassLaneObstacleContainer& LaneObstaclesContainer = LaneData.LaneObstacles;
+		for (const auto& Pair : LaneObstaclesContainer.LaneObstaclesMap)
+		{
+			for(const FMassLaneObstacle& Obstacle : Pair.Value)
+			{
+				const FZoneGraphLaneSection& Section = Obstacle.LaneSection;
+				const FVector ZOffset(0.f, 0.f, 1.f);
+				UE::ZoneGraph::RenderingUtilities::AppendLaneSection(DebugProxy, *ZoneStorage, Section, FColor::Red, /*LineThickness*/5.0f, ZOffset);
+			}
+		}
+	}
+}
+#endif // !UE_BUILD_SHIPPING && !UE_BUILD_TEST
\ No newline at end of file
diff --git a/Engine/Plugins/AI/MassAI/Source/MassAIMovement/Private/MassZoneGraphMovementFragments.cpp b/Engine/Plugins/AI/MassAI/Source/MassAIMovement/Private/MassZoneGraphMovementFragments.cpp
new file mode 100644
--- /dev/null
+++ b/Engine/Plugins/AI/MassAI/Source/MassAIMovement/Private/MassZoneGraphMovementFragments.cpp
@@ -0,0 +1,475 @@
+﻿// Copyright Epic Games, Inc. All Rights Reserved.
+
+#include "MassZoneGraphMovementFragments.h"
+#include "ZoneGraphSubsystem.h"
+#include "ZoneGraphTypes.h"
+#include "ZoneGraphQuery.h"
+
+
+namespace UE::MassMovement::ZoneGraphPath
+{
+	struct FCachedLaneSegmentIterator
+	{
+		FCachedLaneSegmentIterator(const FMassZoneGraphCachedLaneFragment& InCachedLane, const float DistanceAlongPath, const bool bInMoveReverse)
+			: CachedLane(InCachedLane)
+			, SegmentInc(bInMoveReverse ? -1 : 1)
+			, bMoveReverse(bInMoveReverse)
+		{
+			check(CachedLane.NumPoints >= 2);
+			CurrentSegment = CachedLane.FindSegmentIndexAtDistance(DistanceAlongPath);
+			LastSegment = bInMoveReverse ? 0 : ((int32)CachedLane.NumPoints - 2);
+		}
+
+		bool HasReachedDistance(const float Distance) const
+		{
+			if (CurrentSegment == LastSegment)
+			{
+				return true;
+			}
+
+			if (bMoveReverse)
+			{
+				const float SegStartDistance = CachedLane.LanePointProgressions[CurrentSegment].Get();
+				if (Distance > SegStartDistance)
+				{
+					return true;
+				}
+			}
+			else
+			{
+				const float SegEndDistance = CachedLane.LanePointProgressions[CurrentSegment + 1].Get();
+				if (Distance < SegEndDistance)
+				{
+					return true;
+				}
+			}
+
+			return false;
+		}
+
+		void Next()
+		{
+			if (CurrentSegment != LastSegment)
+			{
+				CurrentSegment += SegmentInc;
+			}
+		}
+
+		const FMassZoneGraphCachedLaneFragment& CachedLane;
+		int32 CurrentSegment = 0;
+		int32 LastSegment = 0;
+		int32 SegmentInc = 0;
+		bool bMoveReverse = false;
+	};
+
+} // UE::MassMovement::ZoneGraphPath
+
+void FMassZoneGraphCachedLaneFragment::CacheLaneData(const FZoneGraphStorage& ZoneGraphStorage, const FZoneGraphLaneHandle CurrentLaneHandle,
+													 const float CurrentDistanceAlongLane, const float TargetDistanceAlongLane, const float InflateDistance)
+{
+	const FZoneLaneData& Lane = ZoneGraphStorage.Lanes[CurrentLaneHandle.Index];
+
+	const float StartDistance = FMath::Min(CurrentDistanceAlongLane, TargetDistanceAlongLane);
+	const float EndDistance = FMath::Max(CurrentDistanceAlongLane, TargetDistanceAlongLane);
+	const float CurrentLaneLength = ZoneGraphStorage.LanePointProgressions[Lane.PointsEnd - 1];
+
+	// If cached data contains the request part of the lane, early out.
+	const float InflatedStartDistance = FMath::Max(0.0f, StartDistance - InflateDistance);
+	const float InflatedEndDistance = FMath::Min(EndDistance + InflateDistance, CurrentLaneLength);
+	if (LaneHandle == CurrentLaneHandle
+		&& NumPoints > 0
+		&& InflatedStartDistance >= LanePointProgressions[0].Get()
+		&& InflatedEndDistance <= LanePointProgressions[NumPoints - 1].Get())
+	{
+		return;
+	}
+
+	Reset();
+	CacheID++;
+
+	LaneHandle = CurrentLaneHandle;
+	LaneWidth = FMassInt16Real(Lane.Width);
+	LaneLength = CurrentLaneLength;
+
+	const int32 LaneNumPoints = Lane.PointsEnd - Lane.PointsBegin;
+	if (LaneNumPoints <= (int32)MaxPoints)
+	{
+		// If we can fit all the lane's points, just do a copy.
+		NumPoints = (uint8)LaneNumPoints;
+		for (int32 Index = 0; Index < (int32)NumPoints; Index++)
+		{
+			LanePoints[Index] = ZoneGraphStorage.LanePoints[Lane.PointsBegin + Index];
+			LaneTangentVectors[Index] = FMassSnorm8Vector2D(FVector2D(ZoneGraphStorage.LaneTangentVectors[Lane.PointsBegin + Index]));
+			LanePointProgressions[Index] = FMassInt16Real10(ZoneGraphStorage.LanePointProgressions[Lane.PointsBegin + Index]);
+		}
+	}
+	else
+	{
+		// Find the segment of the lane that is important and copy that.
+		int32 StartSegmentIndex = 0;
+		int32 EndSegmentIndex = 0;
+		UE::ZoneGraph::Query::CalculateLaneSegmentIndexAtDistance(ZoneGraphStorage, CurrentLaneHandle, StartDistance, StartSegmentIndex);
+		UE::ZoneGraph::Query::CalculateLaneSegmentIndexAtDistance(ZoneGraphStorage, CurrentLaneHandle, EndDistance, EndSegmentIndex);
+
+		// Expand if close to start of a segment start.
+		if ((StartSegmentIndex - 1) >= Lane.PointsBegin && (StartDistance - InflateDistance) < ZoneGraphStorage.LanePointProgressions[StartSegmentIndex])
+		{
+			StartSegmentIndex--;
+		}
+		// Expand if close to end segment end.
+		if ((EndSegmentIndex + 1) < (Lane.PointsEnd - 2) && (EndDistance + InflateDistance) > ZoneGraphStorage.LanePointProgressions[EndSegmentIndex + 1])
+		{
+			EndSegmentIndex++;
+		}
+	
+		NumPoints = (uint8)FMath::Min((EndSegmentIndex - StartSegmentIndex) + 2, (int32)MaxPoints);
+
+		for (int32 Index = 0; Index < (int32)NumPoints; Index++)
+		{
+			check((StartSegmentIndex + Index) >= Lane.PointsBegin && (StartSegmentIndex + Index) < Lane.PointsEnd);
+			LanePoints[Index] = ZoneGraphStorage.LanePoints[StartSegmentIndex + Index];
+			LaneTangentVectors[Index] = FMassSnorm8Vector2D(FVector2D(ZoneGraphStorage.LaneTangentVectors[StartSegmentIndex + Index]));
+			LanePointProgressions[Index] = FMassInt16Real10(ZoneGraphStorage.LanePointProgressions[StartSegmentIndex + Index]);
+		}
+	}
+
+	// Calculate extra space around the lane on adjacent lanes.
+	TArray<FZoneGraphLinkedLane> LinkedLanes;
+	UE::ZoneGraph::Query::GetLinkedLanes(ZoneGraphStorage, CurrentLaneHandle, EZoneLaneLinkType::Adjacent, EZoneLaneLinkFlags::Left|EZoneLaneLinkFlags::Right, EZoneLaneLinkFlags::None, LinkedLanes);
+
+	float AdjacentLeftWidth = 0.0f;
+	float AdjacentRightWidth = 0.0f;
+	for (const FZoneGraphLinkedLane& LinkedLane : LinkedLanes)
+	{
+		if (LinkedLane.HasFlags(EZoneLaneLinkFlags::Left))
+		{
+			const FZoneLaneData& AdjacentLane = ZoneGraphStorage.Lanes[LinkedLane.DestLane.Index];
+			AdjacentLeftWidth += AdjacentLane.Width;
+		}
+		else if (LinkedLane.HasFlags(EZoneLaneLinkFlags::Right))
+		{
+			const FZoneLaneData& AdjacentLane = ZoneGraphStorage.Lanes[LinkedLane.DestLane.Index];
+			AdjacentRightWidth += AdjacentLane.Width;
+		}
+	}
+	LaneLeftSpace = FMassInt16Real(AdjacentLeftWidth);
+	LaneRightSpace = FMassInt16Real(AdjacentRightWidth);
+}
+
+bool FMassZoneGraphShortPathFragment::RequestPath(const FMassZoneGraphCachedLaneFragment& CachedLane, const FZoneGraphShortPathRequest& Request, const float CurrentDistanceAlongLane, const float AgentRadius)
+{
+	Reset();
+
+	if (CachedLane.NumPoints < 2)
+	{
+		return false;
+	}
+
+	static constexpr float Epsilon = 0.1f;
+	check(CurrentDistanceAlongLane >= -Epsilon && CurrentDistanceAlongLane <= (CachedLane.LaneLength + Epsilon));
+	
+	// Set common lane parameters
+#if WITH_MASS_DEBUG
+	DebugLaneHandle = CachedLane.LaneHandle;
+#endif
+
+	bMoveReverse = Request.bMoveReverse;
+	EndOfPathIntent = Request.EndOfPathIntent;
+	bPartialResult = false;
+
+	const float DeflatedLaneHalfWidth = FMath::Max(0.0f, CachedLane.LaneWidth.Get() - AgentRadius) * 0.5f;
+	const float DeflatedLaneLeft = DeflatedLaneHalfWidth + CachedLane.LaneLeftSpace.Get();
+	const float DeflatedLaneRight = DeflatedLaneHalfWidth + CachedLane.LaneRightSpace.Get();
+
+	const float TargetDistanceAlongLane = FMath::Clamp(Request.TargetDistance, 0.0f, CachedLane.LaneLength);
+	const float MinDistanceAlongLane = FMath::Min(CurrentDistanceAlongLane, TargetDistanceAlongLane);
+	const float MaxDistanceAlongLane = FMath::Max(CurrentDistanceAlongLane, TargetDistanceAlongLane);
+	
+	const float TangentSign = Request.bMoveReverse ? -1.0f : 1.0f;
+
+	// Slop factors used when testing if a point is conservatively inside the lane.
+	constexpr float OffLaneCapSlop = 10.0f;
+	constexpr float OffLaneEdgeSlop = 1.0f;
+
+	// Calculate how the start point relates to the corresponding location on lane.
+	FVector StartLanePosition;
+	FVector StartLaneTangent;
+	CachedLane.GetPointAndTangentAtDistance(CurrentDistanceAlongLane, StartLanePosition, StartLaneTangent);
+	float StartDistanceAlongPath = CurrentDistanceAlongLane;
+
+	FVector StartPosition = Request.StartPosition;
+	// Calculate start point's relation to the start point location on lane.
+	const FVector StartDelta = StartPosition - StartLanePosition;
+	const FVector StartLeftDir = FVector::CrossProduct(StartLaneTangent, FVector::UpVector);
+	float StartLaneOffset = FVector::DotProduct(StartLeftDir, StartDelta);
+	float StartLaneForwardOffset = FVector::DotProduct(StartLaneTangent, StartDelta) * TangentSign;
+	// The point is off-lane if behind the start, or beyond the boundary.
+	const bool bStartOffLane = StartLaneForwardOffset < -OffLaneCapSlop
+								|| StartLaneOffset < -(DeflatedLaneRight + OffLaneEdgeSlop)
+								|| StartLaneOffset > (DeflatedLaneLeft + OffLaneEdgeSlop);
+	StartLaneOffset = FMath::Clamp(StartLaneOffset, -DeflatedLaneRight, DeflatedLaneLeft);
+
+	if (bStartOffLane)
+	{
+		// The start point was off-lane, move the start location along the lane a bit further to have smoother connection.
+		const float StartForwardOffset = FMath::Clamp(Request.AnticipationDistance + StartLaneForwardOffset, 0.0f, Request.AnticipationDistance);
+		StartDistanceAlongPath += StartForwardOffset * TangentSign; // Not clamping this distance intentionally so that the halfway point and clamping later works correctly.
+	}
+
+	// Calculate how the end point relates to the corresponding location on lane.
+	const bool bHasEndOfPathPoint = Request.bIsEndOfPathPositionSet;
+	float EndDistanceAlongPath = TargetDistanceAlongLane;
+	FVector EndLanePosition = FVector::ZeroVector;
+	FVector EndLaneTangent = FVector::ZeroVector;
+	bool bEndOffLane = false;
+	float EndLaneOffset = StartLaneOffset;
+
+	if (bHasEndOfPathPoint)
+	{
+		// Calculate end point's relation to the end point location on lane.
+		CachedLane.GetPointAndTangentAtDistance(TargetDistanceAlongLane, EndLanePosition, EndLaneTangent);
+		const FVector EndPosition = Request.EndOfPathPosition;
+		const FVector EndDelta = EndPosition - EndLanePosition;
+		const FVector LeftDir = FVector::CrossProduct(EndLaneTangent, FVector::UpVector);
+		EndLaneOffset = FVector::DotProduct(LeftDir, EndDelta);
+		const float EndLaneForwardOffset = FVector::DotProduct(EndLaneTangent, EndDelta) * TangentSign;
+		// The point is off-lane if further than the, or beyond the boundary.
+		bEndOffLane = EndLaneForwardOffset > OffLaneCapSlop
+						|| EndLaneOffset < -(DeflatedLaneRight + OffLaneEdgeSlop)
+						|| EndLaneOffset > (DeflatedLaneLeft + OffLaneEdgeSlop);
+		EndLaneOffset = FMath::Clamp(EndLaneOffset, -DeflatedLaneRight, DeflatedLaneLeft);
+
+		// Move the end location along the lane a bit back to have smoother connection.
+		const float EndForwardOffset = FMath::Clamp(Request.AnticipationDistance - EndLaneForwardOffset, 0.0f, Request.AnticipationDistance);
+		EndDistanceAlongPath -= EndForwardOffset * TangentSign; // Not clamping this distance intentionally so that the halfway point and clamping later works correctly.
+	}
+
+	// Clamp the path move distances to current lane. We use halfway point to split the anticipation in case it gets truncated.
+	const float HalfwayDistanceAlongLane = FMath::Clamp((StartDistanceAlongPath + EndDistanceAlongPath) * 0.5f, MinDistanceAlongLane, MaxDistanceAlongLane);
+
+	if (Request.bMoveReverse)
+	{
+		StartDistanceAlongPath = FMath::Clamp(StartDistanceAlongPath, HalfwayDistanceAlongLane, MaxDistanceAlongLane);
+		EndDistanceAlongPath = FMath::Clamp(EndDistanceAlongPath, MinDistanceAlongLane, HalfwayDistanceAlongLane);
+	}
+	else
+	{
+		StartDistanceAlongPath = FMath::Clamp(StartDistanceAlongPath, MinDistanceAlongLane, HalfwayDistanceAlongLane);
+		EndDistanceAlongPath = FMath::Clamp(EndDistanceAlongPath, HalfwayDistanceAlongLane, MaxDistanceAlongLane);
+	}
+
+	// Check if the mid path got clamped away. This can happen if start of end or both are off-mesh, or just a short path.
+	const float MidPathMoveDistance = FMath::Abs(EndDistanceAlongPath - StartDistanceAlongPath); 
+	const bool bHasMidPath = MidPathMoveDistance > KINDA_SMALL_NUMBER;
+
+	// If end position is not set to a specific location, use proposed offset
+	if (!bHasEndOfPathPoint)
+	{
+		// Slope defines how much the offset can change over the course of the path.
+		constexpr float MaxLaneOffsetSlope = 1.0f / 10.0f;
+		const float MaxOffset = MidPathMoveDistance * MaxLaneOffsetSlope;
+		const float LaneOffset = FMath::Clamp(Request.EndOfPathOffset.Get(), -MaxOffset, MaxOffset);
+		EndLaneOffset = FMath::Clamp(EndLaneOffset + LaneOffset, -DeflatedLaneRight, DeflatedLaneLeft);
+	}
+
+	// Always add off-lane start point.
+	if (bStartOffLane)
+	{
+		FMassZoneGraphPathPoint& StartPoint = Points[NumPoints++];
+		StartPoint.DistanceAlongLane = FMassInt16Real10(CurrentDistanceAlongLane);
+		StartPoint.Position = Request.StartPosition;
+		StartPoint.Tangent = FMassSnorm8Vector2D(StartLaneTangent * TangentSign);
+		StartPoint.bOffLane = true;
+		StartPoint.bIsLaneExtrema = false;
+
+		// Update start point to be inside the lane.
+		CachedLane.GetPointAndTangentAtDistance(StartDistanceAlongPath, StartLanePosition, StartLaneTangent);
+		const FVector LeftDir = FVector::CrossProduct(StartLaneTangent, FVector::UpVector);
+		StartPosition = StartLanePosition + LeftDir * StartLaneOffset;
+
+		// Adjust the start point to point towards the first on-lane point.
+		const FVector DirToClampedPoint = StartPosition - StartPoint.Position;
+		StartPoint.Tangent = FMassSnorm8Vector2D(DirToClampedPoint.GetSafeNormal());
+	}
+
+	// The second point is added if there was no off-lane start point, or we have mid path.
+	// This ensures that there's always at least one start point, and that no excess points are added if both start & end are off-lane close to each other.
+	if (!bStartOffLane || bHasMidPath)
+	{
+		// Add first on-lane point.
+		FMassZoneGraphPathPoint& Point = Points[NumPoints++];
+		Point.DistanceAlongLane = FMassInt16Real10(StartDistanceAlongPath);
+		Point.Position = StartPosition;
+		Point.Tangent = FMassSnorm8Vector2D(StartLaneTangent * TangentSign);
+		Point.bOffLane = false;
+		Point.bIsLaneExtrema = false;
+	}
+
+	// Add in between points.
+	const float InvDistanceRange = 1.0f / (EndDistanceAlongPath - StartDistanceAlongPath); // Used for lane offset interpolation. 
+	float PrevDistanceAlongLane = StartDistanceAlongPath;
+
+	UE::MassMovement::ZoneGraphPath::FCachedLaneSegmentIterator SegmentIterator(CachedLane, StartDistanceAlongPath, Request.bMoveReverse);
+	while (!SegmentIterator.HasReachedDistance(EndDistanceAlongPath))
+	{
+		// The segment endpoint is start when moving backwards (i.e. the segment index), and end when moving forwards.
+		const int32 CurrentSegmentEndPointIndex = SegmentIterator.CurrentSegment + (SegmentIterator.bMoveReverse ? 0 : 1);
+		const float DistanceAlongLane = CachedLane.LanePointProgressions[CurrentSegmentEndPointIndex].Get();
+
+		if (FMath::IsNearlyEqual(PrevDistanceAlongLane, DistanceAlongLane) == false)
+		{
+			if (NumPoints < MaxPoints)
+			{
+				const FVector& LanePosition = CachedLane.LanePoints[CurrentSegmentEndPointIndex];
+				const FVector LaneTangent = CachedLane.LaneTangentVectors[CurrentSegmentEndPointIndex].GetVector();
+
+				const float LaneOffsetT = (DistanceAlongLane - StartDistanceAlongPath) * InvDistanceRange;
+				const float LaneOffset = FMath::Lerp(StartLaneOffset, EndLaneOffset, LaneOffsetT);
+				const FVector LeftDir = FVector::CrossProduct(LaneTangent, FVector::UpVector);
+
+				FMassZoneGraphPathPoint& Point = Points[NumPoints++];
+				Point.DistanceAlongLane = FMassInt16Real10(DistanceAlongLane);
+				Point.Position = LanePosition + LeftDir * LaneOffset;
+				Point.Tangent = FMassSnorm8Vector2D(LaneTangent * TangentSign);
+				Point.bOffLane = false;
+				Point.bIsLaneExtrema = false;
+
+				PrevDistanceAlongLane = DistanceAlongLane;
+			}
+			else
+			{
+				bPartialResult = true;
+				break;
+			}
+		}
+		
+		SegmentIterator.Next();
+	}
+
+	// The second last point is added if there is no end point, or we have mid path.
+	// This ensures that there's always at least one end point, and that no excess points are added if both start & end are off-lane close to each other.
+	if (!bHasEndOfPathPoint || bHasMidPath)
+	{
+		if (NumPoints < MaxPoints)
+		{
+			// Interpolate last point on mid path.
+			FVector LanePosition;
+			FVector LaneTangent;
+			CachedLane.InterpolatePointAndTangentOnSegment(SegmentIterator.CurrentSegment, EndDistanceAlongPath, LanePosition, LaneTangent);
+
+			const float LaneOffset = EndLaneOffset;
+			const FVector LeftDir = FVector::CrossProduct(LaneTangent, FVector::UpVector);
+
+			FMassZoneGraphPathPoint& Point = Points[NumPoints++];
+			Point.DistanceAlongLane = FMassInt16Real10(EndDistanceAlongPath);
+			Point.Position = LanePosition + LeftDir * LaneOffset;
+			Point.Tangent = FMassSnorm8Vector2D(LaneTangent * TangentSign);
+			Point.bOffLane = false;
+			Point.bIsLaneExtrema = !Request.bIsEndOfPathPositionSet && CachedLane.IsDistanceAtLaneExtrema(EndDistanceAlongPath);
+		}
+		else
+		{
+			bPartialResult = true;
+		}
+	}
+
+	// Add end of path point if set.
+	if (bHasEndOfPathPoint)
+	{
+		if (NumPoints < MaxPoints)
+		{
+			const FVector EndPosition = Request.EndOfPathPosition;
+			const FVector EndDirection = Request.bIsEndOfPathDirectionSet ? Request.EndOfPathDirection.Get() : (EndLaneTangent * TangentSign);
+			
+			FMassZoneGraphPathPoint& Point = Points[NumPoints++];
+			Point.DistanceAlongLane = FMassInt16Real10(TargetDistanceAlongLane);
+			Point.Position = EndPosition;
+			Point.Tangent = FMassSnorm8Vector2D(EndDirection);
+			Point.bOffLane = bEndOffLane;
+			Point.bIsLaneExtrema = false;
+		}
+		else
+		{
+			bPartialResult = true;
+		}
+	}
+
+	checkf(NumPoints >= 2, TEXT("Path should have at least 2 points at this stage, has %d."), NumPoints);
+
+	// Calculate movement distance at each point.
+	float PathDistance = 0.0f;
+	Points[0].Distance.Set(PathDistance);
+	for (uint8 PointIndex = 1; PointIndex < NumPoints; PointIndex++)
+	{
+		FMassZoneGraphPathPoint& PrevPoint = Points[PointIndex - 1];
+		FMassZoneGraphPathPoint& Point = Points[PointIndex];
+		const FVector PrevPosition = PrevPoint.Position;
+		const FVector Position = Point.Position;
+		const float DeltaDistance = FVector::Dist(PrevPosition, Position);
+		PathDistance += DeltaDistance;
+		Point.Distance.Set(PathDistance);
+	}
+	
+	// If the last point on path reaches end of the lane, set the next handle to the next lane. It will be update when path finishes.
+	if (!bPartialResult && Request.NextLaneHandle.IsValid())
+	{
+		const FMassZoneGraphPathPoint& LastPoint = Points[NumPoints - 1];
+	
+		if (Request.NextExitLinkType == EZoneLaneLinkType::Adjacent || LastPoint.bIsLaneExtrema)
+		{
+			NextLaneHandle = Request.NextLaneHandle;
+			NextExitLinkType = Request.NextExitLinkType;
+		}
+	}
+
+	return true;
+}
+
+bool FMassZoneGraphShortPathFragment::RequestStand(const FMassZoneGraphCachedLaneFragment& CachedLane, const float CurrentDistanceAlongLane, const FVector& CurrentPosition)
+{
+	Reset();
+
+	if (CachedLane.NumPoints < 2)
+	{
+		return false;
+	}
+
+	static constexpr float Epsilon = 0.1f;
+	check(CurrentDistanceAlongLane >= -Epsilon && CurrentDistanceAlongLane <= (CachedLane.LaneLength + Epsilon));
+	
+	// Get current location
+	FVector CurrentLanePosition;
+	FVector CurrentLaneTangent;
+	CachedLane.GetPointAndTangentAtDistance(CurrentDistanceAlongLane, CurrentLanePosition, CurrentLaneTangent);
+	
+	// Set common lane parameters
+#if WITH_MASS_DEBUG
+	DebugLaneHandle = CachedLane.LaneHandle;
+#endif
+	bMoveReverse = false;
+	EndOfPathIntent = EMassMovementAction::Stand;
+	bPartialResult = false;
+
+	// Add start point, if the start is outside the lane, add another point to get back to lane.
+	const FVector StartMoveOffset = CurrentPosition - CurrentLanePosition;
+
+	FMassZoneGraphPathPoint& StartPoint = Points[NumPoints];
+	StartPoint.DistanceAlongLane = FMassInt16Real10(CurrentDistanceAlongLane);
+	StartPoint.Position = CurrentPosition;
+	StartPoint.Tangent = FMassSnorm8Vector2D(CurrentLaneTangent);
+	StartPoint.bOffLane = false;
+	StartPoint.bIsLaneExtrema = false;
+	StartPoint.Distance = FMassInt16Real(0.0f);
+	NumPoints++;
+
+	FMassZoneGraphPathPoint& EndPoint = Points[NumPoints];
+	EndPoint.DistanceAlongLane = FMassInt16Real10(CurrentDistanceAlongLane);
+	EndPoint.Position = CurrentPosition;
+	EndPoint.Tangent = FMassSnorm8Vector2D(CurrentLaneTangent);
+	EndPoint.bOffLane = false;
+	EndPoint.bIsLaneExtrema = false;
+	EndPoint.Distance = FMassInt16Real(0.0f);
+	NumPoints++;
+
+	return true;
+}
diff --git a/Engine/Plugins/AI/MassAI/Source/MassAIMovement/Private/MassZoneGraphMovementProcessors.cpp b/Engine/Plugins/AI/MassAI/Source/MassAIMovement/Private/MassZoneGraphMovementProcessors.cpp
new file mode 100644
--- /dev/null
+++ b/Engine/Plugins/AI/MassAI/Source/MassAIMovement/Private/MassZoneGraphMovementProcessors.cpp
@@ -0,0 +1,749 @@
+// Copyright Epic Games, Inc. All Rights Reserved.
+
+#include "MassZoneGraphMovementProcessors.h"
+#include "MassZoneGraphMovementFragments.h"
+#include "ZoneGraphSubsystem.h"
+#include "MassAIMovementFragments.h"
+#include "MassMovementSettings.h"
+#include "MassSignalSubsystem.h"
+#include "MassCommonFragments.h"
+#include "ZoneGraphQuery.h"
+#include "VisualLogger/VisualLogger.h"
+#include "MassSimulationLOD.h"
+#include "Engine/World.h"
+
+#define UNSAFE_FOR_MT 1
+
+#if WITH_MASS_DEBUG
+
+namespace UE::MassMovement::Debug
+{
+	FColor MixColors(const FColor ColorA, const FColor ColorB)
+	{
+		const int32 R = ((int32)ColorA.R + (int32)ColorB.R) / 2;
+		const int32 G = ((int32)ColorA.G + (int32)ColorB.G) / 2;
+		const int32 B = ((int32)ColorA.B + (int32)ColorB.B) / 2;
+		const int32 A = ((int32)ColorA.A + (int32)ColorB.A) / 2;
+		return FColor((uint8)R, (uint8)G, (uint8)B, (uint8)A);
+	}
+}
+
+#endif // WITH_MASS_DEBUG
+
+namespace UE::MassMovement
+{
+
+	/*
+	 * Calculates speed scale based on agent's forward direction and desired steering direction.
+	 */
+	static float CalcDirectionalSpeedScale(const FVector ForwardDirection, const FVector SteerDirection)
+	{
+		// @todo: make these configurable
+		constexpr float ForwardSpeedScale = 1.0f;
+		constexpr float BackwardSpeedScale = 0.25f;
+		constexpr float SideSpeedScale = 0.5f;
+
+		const FVector LeftDirection = FVector::CrossProduct(ForwardDirection, FVector::UpVector);
+		const float DirX = FVector::DotProduct(LeftDirection, SteerDirection);
+		const float DirY = FVector::DotProduct(ForwardDirection, SteerDirection);
+
+		// Calculate intersection between a direction vector and ellipse, where A & B are the size of the ellipse.
+		// The direction vector is starting from the center of the ellipse.
+		constexpr float SideA = SideSpeedScale;
+		const float SideB = DirY > 0.0f ? ForwardSpeedScale : BackwardSpeedScale;
+		const float Disc = FMath::Square(SideA) * FMath::Square(DirY) + FMath::Square(SideB) * FMath::Square(DirX);
+		const float Speed = (Disc > SMALL_NUMBER) ? (SideA * SideB / FMath::Sqrt(Disc)) : 0.0f;;
+
+		return Speed;
+	}
+
+	/** Speed envelope when approaching a point. NormalizedDistance in range [0..1] */
+	static float ArrivalSpeedEnvelope(const float NormalizedDistance)
+	{
+		return FMath::Sqrt(NormalizedDistance);
+	}
+
+} // UE::MassMovement
+
+
+//----------------------------------------------------------------------//
+//  UMassZoneGraphLocationInitializer
+//----------------------------------------------------------------------//
+UMassZoneGraphLocationInitializer::UMassZoneGraphLocationInitializer()
+{
+	FragmentType = FMassZoneGraphLaneLocationFragment::StaticStruct();
+	ExecutionFlags = (int32)(EProcessorExecutionFlags::All);
+}
+
+void UMassZoneGraphLocationInitializer::ConfigureQueries()
+{
+	EntityQuery.AddRequirement<FMassZoneGraphLaneLocationFragment>(ELWComponentAccess::ReadWrite);
+	EntityQuery.AddRequirement<FDataFragment_Transform>(ELWComponentAccess::ReadOnly);
+	EntityQuery.AddRequirement<FMassMovementConfigFragment>(ELWComponentAccess::ReadOnly);
+	EntityQuery.AddRequirement<FMassMoveTargetFragment>(ELWComponentAccess::ReadWrite); // Make optional?
+}
+
+void UMassZoneGraphLocationInitializer::Initialize(UObject& Owner)
+{
+	Super::Initialize(Owner);
+	ZoneGraphSubsystem = UWorld::GetSubsystem<UZoneGraphSubsystem>(Owner.GetWorld());
+	SignalSubsystem = UWorld::GetSubsystem<UMassSignalSubsystem>(Owner.GetWorld());
+}
+
+void UMassZoneGraphLocationInitializer::Execute(UEntitySubsystem& EntitySubsystem, FLWComponentSystemExecutionContext& Context)
+{
+	const UMassMovementSettings* Settings = GetDefault<UMassMovementSettings>();
+	if (!ZoneGraphSubsystem || !Settings || !SignalSubsystem)
+	{
+		return;
+	}
+
+	EntityQuery.ForEachEntityChunk(EntitySubsystem, Context, [this, Settings](FLWComponentSystemExecutionContext& Context)
+	{
+		const int32 NumEntities = Context.GetEntitiesNum();
+		const TArrayView<FMassZoneGraphLaneLocationFragment> LaneLocationList = Context.GetMutableComponentView<FMassZoneGraphLaneLocationFragment>();
+		const TArrayView<FMassMoveTargetFragment> MoveTargetList = Context.GetMutableComponentView<FMassMoveTargetFragment>();
+		const TConstArrayView<FMassMovementConfigFragment> MovementConfigList = Context.GetComponentView<FMassMovementConfigFragment>();
+		const TConstArrayView<FDataFragment_Transform> TransformList = Context.GetComponentView<FDataFragment_Transform>();
+
+		// Get the default movement config.
+		FMassMovementConfigHandle CurrentConfigHandle;
+		const FMassMovementConfig* CurrentMovementConfig = nullptr;
+
+		for (int32 EntityIndex = 0; EntityIndex < NumEntities; ++EntityIndex)
+		{
+			const FDataFragment_Transform& Transform = TransformList[EntityIndex];
+			const FVector& AgentLocation = Transform.GetTransform().GetLocation();
+			const FMassMovementConfigFragment& MovementConfig = MovementConfigList[EntityIndex];
+			FMassMoveTargetFragment& MoveTarget = MoveTargetList[EntityIndex];
+			FMassZoneGraphLaneLocationFragment& LaneLocation = LaneLocationList[EntityIndex];
+
+			if (MovementConfig.ConfigHandle != CurrentConfigHandle)
+			{
+				CurrentMovementConfig = Settings->GetMovementConfigByHandle(MovementConfig.ConfigHandle);
+				CurrentConfigHandle = MovementConfig.ConfigHandle;
+			}
+			if (!CurrentMovementConfig)
+			{
+#if WITH_MASS_DEBUG && UNSAFE_FOR_MT
+				const FLWEntity Entity = Context.GetEntity(EntityIndex);
+				UE_VLOG(this, LogMassNavigation, Log, TEXT("Entity [%s] Invalid movement config."), *Entity.DebugGetDescription());
+#endif
+				continue;
+			}
+
+			const FVector QuerySize(CurrentMovementConfig->QueryRadius);
+			const FBox QueryBounds(AgentLocation - QuerySize, AgentLocation + QuerySize);
+			
+			FZoneGraphLaneLocation NearestLane;
+			float NearestLaneDistSqr = 0;
+			
+			if (ZoneGraphSubsystem->FindNearestLane(QueryBounds, CurrentMovementConfig->LaneFilter, NearestLane, NearestLaneDistSqr))
+			{
+				const FZoneGraphStorage* ZoneGraphStorage = ZoneGraphSubsystem->GetZoneGraphStorage(NearestLane.LaneHandle.DataHandle);
+				check(ZoneGraphStorage); // Assume valid storage since we just got result.
+
+				LaneLocation.LaneHandle = NearestLane.LaneHandle;
+				LaneLocation.DistanceAlongLane = NearestLane.DistanceAlongLane;
+				UE::ZoneGraph::Query::GetLaneLength(*ZoneGraphStorage, LaneLocation.LaneHandle, LaneLocation.LaneLength);
+				
+				MoveTarget.Center = AgentLocation;
+				MoveTarget.Forward = NearestLane.Tangent;
+				MoveTarget.DistanceToGoal = 0.0f;
+				MoveTarget.SlackRadius = 0.0f;
+			}
+			else
+			{
+				LaneLocation.LaneHandle.Reset();
+				LaneLocation.DistanceAlongLane = 0.0f;
+				LaneLocation.LaneLength = 0.0f;
+
+				MoveTarget.Center = AgentLocation;
+				MoveTarget.Forward = FVector::ForwardVector;
+				MoveTarget.DistanceToGoal = 0.0f;
+				MoveTarget.SlackRadius = 0.0f;
+			}
+		}
+	});
+}
+
+
+//----------------------------------------------------------------------//
+//  UMassZoneGraphPathFollowProcessor
+//----------------------------------------------------------------------//
+UMassZoneGraphPathFollowProcessor::UMassZoneGraphPathFollowProcessor()
+{
+	ExecutionFlags = (int32)EProcessorExecutionFlags::All;
+	ExecutionOrder.ExecuteInGroup = UE::Mass::ProcessorGroupNames::Tasks;
+	ExecutionOrder.ExecuteBefore.Add(UE::Mass::ProcessorGroupNames::Avoidance);
+}
+
+void UMassZoneGraphPathFollowProcessor::Initialize(UObject& Owner)
+{
+	Super::Initialize(Owner);
+	ZoneGraphSubsystem = UWorld::GetSubsystem<UZoneGraphSubsystem>(Owner.GetWorld());
+	SignalSubsystem = UWorld::GetSubsystem<UMassSignalSubsystem>(Owner.GetWorld());
+}
+
+void UMassZoneGraphPathFollowProcessor::ConfigureQueries()
+{
+	EntityQuery_Conditional.AddRequirement<FMassZoneGraphShortPathFragment>(ELWComponentAccess::ReadWrite);
+	EntityQuery_Conditional.AddRequirement<FMassZoneGraphLaneLocationFragment>(ELWComponentAccess::ReadWrite);
+	EntityQuery_Conditional.AddRequirement<FMassMoveTargetFragment>(ELWComponentAccess::ReadWrite);
+	EntityQuery_Conditional.AddRequirement<FMassSimulationLODFragment>(ELWComponentAccess::ReadOnly);
+	EntityQuery_Conditional.AddChunkRequirement<FMassSimulationVariableTickChunkFragment>(ELWComponentAccess::ReadOnly);
+	EntityQuery_Conditional.SetChunkFilter(&FMassSimulationVariableTickChunkFragment::ShouldTickChunkThisFrame);
+}
+
+void UMassZoneGraphPathFollowProcessor::Execute(UEntitySubsystem& EntitySubsystem, FLWComponentSystemExecutionContext& Context)
+{
+	if (!SignalSubsystem || !ZoneGraphSubsystem)
+	{
+		return;
+	}
+	
+	TArray<FLWEntity> EntitiesToSignalPathDone;
+	TArray<FLWEntity> EntitiesToSignalLaneChanged;
+
+	EntityQuery_Conditional.ForEachEntityChunk(EntitySubsystem, Context, [this, &EntitiesToSignalPathDone, &EntitiesToSignalLaneChanged](FLWComponentSystemExecutionContext& Context)
+		{
+		const int32 NumEntities = Context.GetEntitiesNum();
+		const TArrayView<FMassZoneGraphShortPathFragment> ShortPathList = Context.GetMutableComponentView<FMassZoneGraphShortPathFragment>();
+		const TArrayView<FMassZoneGraphLaneLocationFragment> LaneLocationList = Context.GetMutableComponentView<FMassZoneGraphLaneLocationFragment>();
+		const TArrayView<FMassMoveTargetFragment> MoveTargetList = Context.GetMutableComponentView<FMassMoveTargetFragment>();
+		const TConstArrayView<FMassSimulationLODFragment> SimLODList = Context.GetComponentView<FMassSimulationLODFragment>();
+
+		for (int32 EntityIndex = 0; EntityIndex < NumEntities; ++EntityIndex)
+		{
+			FMassZoneGraphShortPathFragment& ShortPath = ShortPathList[EntityIndex];
+			FMassZoneGraphLaneLocationFragment& LaneLocation = LaneLocationList[EntityIndex];
+			FMassMoveTargetFragment& MoveTarget = MoveTargetList[EntityIndex];
+			const FMassSimulationLODFragment& SimLOD = SimLODList[EntityIndex];
+			const FLWEntity Entity = Context.GetEntity(EntityIndex);
+			const float DeltaTime = SimLOD.DeltaTime;
+
+			bool bDisplayDebug = false;
+#if WITH_MASS_DEBUG && UNSAFE_FOR_MT // this will result in bDisplayDebug == false and disabling of all the vlogs below
+			bDisplayDebug = UE::MassDebug::IsDebuggingEntity(Entity);
+			if (bDisplayDebug)
+			{
+				UE_VLOG(this, LogMassNavigation, Log, TEXT("Entity [%s] Updating path following"), *Entity.DebugGetDescription());
+			}
+#endif // WITH_MASS_DEBUG
+
+			// Must have at least two points to interpolate.
+			if (MoveTarget.GetCurrentAction() == EMassMovementAction::Move && ShortPath.NumPoints >= 2)
+			{
+				const bool bWasDone = ShortPath.IsDone();
+
+				// Note: this should be in sync with the logic in apply velocity.
+				const bool bHasSteering = SimLOD.LOD != EMassLOD::Off;
+
+				if (!bHasSteering || !MoveTarget.bSteeringFallingBehind)
+				{
+					// Update progress
+					ShortPath.ProgressDistance += MoveTarget.DesiredSpeed.Get() * DeltaTime;
+				}
+
+				// @todo MassMovement: Ideally we would carry over any left over distance to the next path, especially when dealing with larger timesteps.
+				// @todo MassMovement: Feedback current movement progress back to ShortPath.DesiredSpeed.
+
+				if (!bWasDone)
+				{
+					const uint8 LastPointIndex = ShortPath.NumPoints - 1;
+#if WITH_MASS_DEBUG
+					ensureMsgf(LaneLocation.LaneHandle == ShortPath.DebugLaneHandle, TEXT("Short path lane should match current lane location."));
+#endif // WITH_MASS_DEBUG
+
+					if (ShortPath.ProgressDistance <= 0.0f)
+					{
+						// Requested time before the start of the path
+						LaneLocation.DistanceAlongLane = ShortPath.Points[0].DistanceAlongLane.Get();
+						
+						MoveTarget.Center = ShortPath.Points[0].Position;
+						MoveTarget.Forward = ShortPath.Points[0].Tangent.GetVector();
+						MoveTarget.DistanceToGoal = ShortPath.Points[LastPointIndex].Distance.Get();
+						MoveTarget.bOffBoundaries = ShortPath.Points[0].bOffLane;
+
+						UE_CVLOG(bDisplayDebug,this, LogMassNavigation, Verbose, TEXT("Entity [%s] before start of lane %s at distance %.1f. Distance to goal: %.1f."),
+							*Entity.DebugGetDescription(),
+							*LaneLocation.LaneHandle.ToString(),
+							LaneLocation.DistanceAlongLane,
+							MoveTarget.DistanceToGoal);
+					}
+					else if (ShortPath.ProgressDistance <= ShortPath.Points[LastPointIndex].Distance.Get())
+					{
+						// Requested time along the path, interpolate.
+						uint8 PointIndex = 0;
+						while (PointIndex < (ShortPath.NumPoints - 2))
+						{
+							const FMassZoneGraphPathPoint& NextPoint = ShortPath.Points[PointIndex + 1];
+							if (ShortPath.ProgressDistance <= NextPoint.Distance.Get())
+							{
+								break;
+							}
+							PointIndex++;
+						}
+
+						const FMassZoneGraphPathPoint& CurrPoint = ShortPath.Points[PointIndex];
+						const FMassZoneGraphPathPoint& NextPoint = ShortPath.Points[PointIndex + 1];
+						const float T = (ShortPath.ProgressDistance - CurrPoint.Distance.Get()) / (NextPoint.Distance.Get() - CurrPoint.Distance.Get());
+						
+						LaneLocation.DistanceAlongLane = FMath::Min(FMath::Lerp(CurrPoint.DistanceAlongLane.Get(), NextPoint.DistanceAlongLane.Get(), T), LaneLocation.LaneLength);
+
+						MoveTarget.Center = FMath::Lerp(CurrPoint.Position, NextPoint.Position, T);
+						MoveTarget.Forward = FMath::Lerp(CurrPoint.Tangent.GetVector(), NextPoint.Tangent.GetVector(), T).GetSafeNormal();
+						MoveTarget.DistanceToGoal = ShortPath.Points[LastPointIndex].Distance.Get() - FMath::Lerp(CurrPoint.Distance.Get(), NextPoint.Distance.Get(), T);
+						MoveTarget.bOffBoundaries = CurrPoint.bOffLane || NextPoint.bOffLane;
+
+						UE_CVLOG(bDisplayDebug, this, LogMassNavigation, Verbose, TEXT("Entity [%s] along lane %s at distance %.1f. Distance to goal: %.1f."),
+							*Entity.DebugGetDescription(),
+							*LaneLocation.LaneHandle.ToString(),
+							LaneLocation.DistanceAlongLane,
+							MoveTarget.DistanceToGoal);
+					}
+					else
+					{
+						// Requested time after the end of the path, clamp to lane length in case quantization overshoots.
+						LaneLocation.DistanceAlongLane = FMath::Min(ShortPath.Points[LastPointIndex].DistanceAlongLane.Get(), LaneLocation.LaneLength);
+
+						MoveTarget.Center = ShortPath.Points[LastPointIndex].Position;
+						MoveTarget.Forward = ShortPath.Points[LastPointIndex].Tangent.GetVector();
+						MoveTarget.DistanceToGoal = 0.0f;
+						MoveTarget.bOffBoundaries = ShortPath.Points[LastPointIndex].bOffLane;
+
+						UE_CVLOG(bDisplayDebug, this, LogMassNavigation, Log, TEXT("Entity [%s] Finished path follow on lane %s at distance %f."),
+							*Entity.DebugGetDescription(), *LaneLocation.LaneHandle.ToString(), LaneLocation.DistanceAlongLane);
+
+						if (bDisplayDebug)
+						{
+							UE_VLOG(this, LogMassNavigation, Log, TEXT("Entity [%s] End of path."), *Entity.DebugGetDescription());
+						}
+
+						// Check to see if need advance to next lane.
+						if (ShortPath.NextLaneHandle.IsValid())
+						{
+							const FZoneGraphStorage* ZoneGraphStorage = ZoneGraphSubsystem->GetZoneGraphStorage(LaneLocation.LaneHandle.DataHandle);
+							if (ZoneGraphStorage != nullptr)
+							{
+								if (ShortPath.NextExitLinkType == EZoneLaneLinkType::Outgoing)
+								{
+									float NewLaneLength = 0.f;
+									UE::ZoneGraph::Query::GetLaneLength(*ZoneGraphStorage, ShortPath.NextLaneHandle, NewLaneLength);
+
+									UE_CVLOG(bDisplayDebug, this, LogMassNavigation, Log, TEXT("Entity [%s] Switching to OUTGOING lane %s -> %s, new distance %f."),
+										*Entity.DebugGetDescription(), *LaneLocation.LaneHandle.ToString(), *ShortPath.NextLaneHandle.ToString(), 0.f);
+
+									// update lane location
+									LaneLocation.LaneHandle = ShortPath.NextLaneHandle;
+									LaneLocation.LaneLength = NewLaneLength;
+									LaneLocation.DistanceAlongLane = 0.0f;
+								}
+								else if (ShortPath.NextExitLinkType == EZoneLaneLinkType::Incoming)
+								{
+									float NewLaneLength = 0.f;
+									UE::ZoneGraph::Query::GetLaneLength(*ZoneGraphStorage, ShortPath.NextLaneHandle, NewLaneLength);
+
+									UE_CVLOG(bDisplayDebug, this, LogMassNavigation, Log, TEXT("Entity [%s] Switching to INCOMING lane %s -> %s, new distance %f."),
+										*Entity.DebugGetDescription(), *LaneLocation.LaneHandle.ToString(), *ShortPath.NextLaneHandle.ToString(), NewLaneLength);
+
+									// update lane location
+									LaneLocation.LaneHandle = ShortPath.NextLaneHandle;
+									LaneLocation.LaneLength = NewLaneLength;
+									LaneLocation.DistanceAlongLane = NewLaneLength;
+								}
+								else if (ShortPath.NextExitLinkType == EZoneLaneLinkType::Adjacent)
+								{
+									FZoneGraphLaneLocation NewLocation;
+									float DistanceSqr;
+									if (UE::ZoneGraph::Query::FindNearestLocationOnLane(*ZoneGraphStorage, ShortPath.NextLaneHandle, MoveTarget.Center, MAX_flt, NewLocation, DistanceSqr))
+									{
+										float NewLaneLength = 0.f;
+										UE::ZoneGraph::Query::GetLaneLength(*ZoneGraphStorage, ShortPath.NextLaneHandle, NewLaneLength);
+
+										UE_CVLOG(bDisplayDebug, this, LogMassNavigation, Log, TEXT("Entity [%s] Switching to ADJACENT lane %s -> %s, new distance %f."),
+											*Entity.DebugGetDescription(), *LaneLocation.LaneHandle.ToString(), *ShortPath.NextLaneHandle.ToString(), NewLocation.DistanceAlongLane);
+
+										// update lane location
+										LaneLocation.LaneHandle = ShortPath.NextLaneHandle;
+										LaneLocation.LaneLength = NewLaneLength;
+										LaneLocation.DistanceAlongLane = NewLocation.DistanceAlongLane;
+
+										MoveTarget.Forward = NewLocation.Tangent;
+									}
+									else
+									{
+										UE_CVLOG(bDisplayDebug, this, LogMassNavigation, Error, TEXT("Entity [%s] Failed to switch to ADJACENT lane %s -> %s."),
+											*Entity.DebugGetDescription(), *LaneLocation.LaneHandle.ToString(), *ShortPath.NextLaneHandle.ToString());
+									}
+								}
+								else
+								{
+									ensureMsgf(false, TEXT("Unhandle NextExitLinkType type %s"), *UEnum::GetValueAsString(ShortPath.NextExitLinkType));
+								}
+
+								// Signal lane changed.
+								EntitiesToSignalLaneChanged.Add(Entity);
+							}
+							else
+							{
+								UE_CVLOG(bDisplayDebug, this, LogMassNavigation, Error, TEXT("Entity [%s] Could not find ZoneGraph storage for lane %s."),
+									*Entity.DebugGetDescription(), *LaneLocation.LaneHandle.ToString());
+							}
+						}
+						else
+						{
+							UE_CVLOG(bDisplayDebug, this, LogMassNavigation, Log, TEXT("Entity [%s] Next lane not defined."), *Entity.DebugGetDescription());
+						}
+
+						ShortPath.bDone = true;
+					}
+				}
+
+				const bool bIsDone = ShortPath.IsDone();
+
+				// Signal path done.
+				if (!bWasDone && bIsDone)
+				{
+					EntitiesToSignalPathDone.Add(Entity);
+				}
+
+#if WITH_MASS_DEBUG && UNSAFE_FOR_MT
+				if (bDisplayDebug)
+				{
+					const FColor EntityColor = UE::MassDebug::GetEntityDebugColor(Entity);
+
+					const FVector ZOffset(0,0,25);
+					const FColor LightEntityColor = UE::MassMovement::Debug::MixColors(EntityColor, FColor::White);
+					
+					for (uint8 PointIndex = 0; PointIndex < ShortPath.NumPoints - 1; PointIndex++)
+					{
+						const FMassZoneGraphPathPoint& CurrPoint = ShortPath.Points[PointIndex];
+						const FMassZoneGraphPathPoint& NextPoint = ShortPath.Points[PointIndex + 1];
+
+						// Path
+						UE_VLOG_SEGMENT_THICK(this, LogMassNavigation, Display, CurrPoint.Position + ZOffset, NextPoint.Position + ZOffset, EntityColor, /*Thickness*/3.0f, TEXT(""));
+					}
+					
+					for (uint8 PointIndex = 0; PointIndex < ShortPath.NumPoints; PointIndex++)
+					{
+						const FMassZoneGraphPathPoint& CurrPoint = ShortPath.Points[PointIndex];
+						const FVector CurrBase = CurrPoint.Position + ZOffset;
+						// Lane tangents
+						UE_VLOG_SEGMENT_THICK(this, LogMassNavigation, Display, CurrBase, CurrBase + CurrPoint.Tangent.GetVector() * 100.0f, LightEntityColor, /*Thickness*/1.0f, TEXT(""));
+					}
+
+					if (ShortPath.NumPoints > 0 && ShortPath.NextLaneHandle.IsValid())
+					{
+						const FMassZoneGraphPathPoint& LastPoint = ShortPath.Points[ShortPath.NumPoints - 1];
+						const FVector CurrBase = LastPoint.Position + ZOffset;
+						UE_VLOG_SEGMENT_THICK(this, LogMassNavigation, Display, CurrBase, CurrBase + FVector(0,0,100), FColor::Red, /*Thickness*/3.0f, TEXT("Next: %s"), *ShortPath.NextLaneHandle.ToString());
+					}
+				}
+#endif // WITH_MASS_DEBUG
+			}
+		}
+	});
+
+	if (EntitiesToSignalPathDone.Num())
+	{
+		SignalSubsystem->SignalEntities(UE::Mass::Signals::FollowPointPathDone, EntitiesToSignalPathDone);
+	}
+	if (EntitiesToSignalLaneChanged.Num())
+	{
+		SignalSubsystem->SignalEntities(UE::Mass::Signals::CurrentLaneChanged, EntitiesToSignalLaneChanged);
+	}
+}
+
+
+//----------------------------------------------------------------------//
+//  UMassZoneGraphSteeringProcessor
+//----------------------------------------------------------------------//
+UMassZoneGraphSteeringProcessor::UMassZoneGraphSteeringProcessor()
+{
+	ExecutionFlags = int32(EProcessorExecutionFlags::All);
+	ExecutionOrder.ExecuteAfter.Add(UE::Mass::ProcessorGroupNames::Tasks);
+	ExecutionOrder.ExecuteBefore.Add(UE::Mass::ProcessorGroupNames::Avoidance);
+}
+
+void UMassZoneGraphSteeringProcessor::Initialize(UObject& Owner)
+{
+	Super::Initialize(Owner);
+	SignalSubsystem = UWorld::GetSubsystem<UMassSignalSubsystem>(Owner.GetWorld());
+}
+
+void UMassZoneGraphSteeringProcessor::ConfigureQueries()
+{
+	EntityQuery.AddRequirement<FMassMoveTargetFragment>(ELWComponentAccess::ReadWrite);
+	EntityQuery.AddRequirement<FMassMovementConfigFragment>(ELWComponentAccess::ReadOnly);
+	EntityQuery.AddRequirement<FDataFragment_Transform>(ELWComponentAccess::ReadWrite);
+	EntityQuery.AddRequirement<FMassSteeringFragment>(ELWComponentAccess::ReadWrite);
+	EntityQuery.AddRequirement<FMassSteeringGhostFragment>(ELWComponentAccess::ReadWrite);
+	EntityQuery.AddRequirement<FMassVelocityFragment>(ELWComponentAccess::ReadOnly);
+
+	// No need for Off LOD to do steering, applying move target directly
+	EntityQuery.AddTagRequirement<FMassOffLODTag>(ELWComponentPresence::None);
+}
+
+void UMassZoneGraphSteeringProcessor::Execute(UEntitySubsystem& EntitySubsystem, FLWComponentSystemExecutionContext& Context)
+{
+	const UMassMovementSettings* Settings = GetDefault<UMassMovementSettings>();
+	check(Settings);
+
+	if (!SignalSubsystem)
+	{
+		return;
+	}
+	
+	EntityQuery.ForEachEntityChunk(EntitySubsystem, Context, [this, Settings](FLWComponentSystemExecutionContext& Context)
+	{
+		const int32 NumEntities = Context.GetEntitiesNum();
+		const TArrayView<FMassMoveTargetFragment> MoveTargetList = Context.GetMutableComponentView<FMassMoveTargetFragment>();
+		const TConstArrayView<FMassMovementConfigFragment> MovementConfigList = Context.GetComponentView<FMassMovementConfigFragment>();
+		const TConstArrayView<FMassVelocityFragment> VelocityList = Context.GetComponentView<FMassVelocityFragment>();
+		const TArrayView<FDataFragment_Transform> TransformList = Context.GetMutableComponentView<FDataFragment_Transform>();
+		const TArrayView<FMassSteeringFragment> SteeringList = Context.GetMutableComponentView<FMassSteeringFragment>();
+		const TArrayView<FMassSteeringGhostFragment> GhostList = Context.GetMutableComponentView<FMassSteeringGhostFragment>();
+
+		// @todo: make configurable
+		constexpr float StandDeadZoneRadius = 5.0f;
+		constexpr float SteeringReactionTime = 0.2f;
+		constexpr float SteerK = 1.f / SteeringReactionTime;
+		const float DeltaTime = Context.GetDeltaTimeSeconds();
+
+		// Get the default movement config.
+		FMassMovementConfigHandle CurrentConfigHandle;
+		const FMassMovementConfig* CurrentMovementConfig = nullptr;
+
+		for (int32 EntityIndex = 0; EntityIndex < NumEntities; ++EntityIndex)
+		{
+			FDataFragment_Transform& TransformFragment = TransformList[EntityIndex];
+			FMassSteeringFragment& Steering = SteeringList[EntityIndex];
+			FMassSteeringGhostFragment& Ghost = GhostList[EntityIndex];
+			FMassMoveTargetFragment& MoveTarget = MoveTargetList[EntityIndex];
+			const FMassMovementConfigFragment& MovementConfig = MovementConfigList[EntityIndex];
+			const FMassVelocityFragment& Velocity = VelocityList[EntityIndex];
+			const FLWEntity Entity = Context.GetEntity(EntityIndex);
+
+			if (MovementConfig.ConfigHandle != CurrentConfigHandle)
+			{
+				CurrentMovementConfig = Settings->GetMovementConfigByHandle(MovementConfig.ConfigHandle);
+				CurrentConfigHandle = MovementConfig.ConfigHandle;
+			}
+			if (!CurrentMovementConfig)
+			{
+#if WITH_MASS_DEBUG && UNSAFE_FOR_MT
+				UE_VLOG(this, LogMassNavigation, Log, TEXT("Entity [%s] Invalid movement config."), *Entity.DebugGetDescription());
+#endif
+				continue;
+			}
+
+			FTransform& Transform = TransformFragment.GetMutableTransform();;
+
+			// Calculate velocity for steering.
+			const FVector CurrentLocation = Transform.GetLocation();
+			const FVector CurrentForward = Transform.GetRotation().GetForwardVector();
+
+			const float LookAheadDistance = FMath::Max(KINDA_SMALL_NUMBER, CurrentMovementConfig->Steering.LookAheadDistance);
+
+			if (MoveTarget.GetCurrentAction() == EMassMovementAction::Move)
+			{
+				// Tune down avoidance and speed when arriving at goal.
+				float ArrivalFade = 1.0f;
+				if (MoveTarget.IntentAtGoal == EMassMovementAction::Stand)
+				{
+					ArrivalFade = FMath::Clamp(MoveTarget.DistanceToGoal / LookAheadDistance, 0.0f, 1.0f);
+				}
+				const float SteeringPredictionDistance = LookAheadDistance * ArrivalFade;
+
+				// Steer towards and along the move target.
+				const FVector TargetSide = FVector::CrossProduct(MoveTarget.Forward, FVector::UpVector);
+				const FVector Delta = CurrentLocation - MoveTarget.Center;
+
+				const float ForwardOffset = FVector::DotProduct(MoveTarget.Forward, Delta);
+
+				// Calculate steering direction. When far away from the line defined by TargetPosition and TargetTangent,
+				// the steering direction is towards the line, the close we get, the more it aligns with the line.
+				const float SidewaysOffset = FVector::DotProduct(TargetSide, Delta);
+				const float SteerForward = FMath::Sqrt(FMath::Max(0.0f, FMath::Square(SteeringPredictionDistance) - FMath::Square(SidewaysOffset)));
+
+				// The Max() here makes the steering directions behind the TargetPosition to steer towards it directly.
+				FVector SteerTarget = MoveTarget.Center + MoveTarget.Forward * FMath::Clamp(ForwardOffset + SteerForward, 0.0f, SteeringPredictionDistance);
+
+				FVector SteerDirection = SteerTarget - CurrentLocation;
+				SteerDirection.Z = 0.0f;
+				const float DistanceToSteerTarget = SteerDirection.Length();
+				if (DistanceToSteerTarget > KINDA_SMALL_NUMBER)
+				{
+					SteerDirection *= 1.0f / DistanceToSteerTarget;
+				}
+				
+				const float DirSpeedScale = UE::MassMovement::CalcDirectionalSpeedScale(CurrentForward, SteerDirection);
+				float DesiredSpeed = MoveTarget.DesiredSpeed.Get() * DirSpeedScale;
+
+				// Control speed based relation to the forward axis of the move target.
+				float CatchupDesiredSpeed = DesiredSpeed;
+				if (ForwardOffset < 0.0f)
+				{
+					// Falling behind, catch up
+					const float T = FMath::Min(-ForwardOffset / LookAheadDistance, 1.0f);
+					CatchupDesiredSpeed = FMath::Lerp(DesiredSpeed, CurrentMovementConfig->MaximumSpeed, T);
+				}
+				else if (ForwardOffset > 0.0f)
+				{
+					// Ahead, slow down.
+					const float T = FMath::Min(ForwardOffset / LookAheadDistance, 1.0f);
+					CatchupDesiredSpeed = FMath::Lerp(DesiredSpeed, DesiredSpeed * 0.0f, 1.0f - FMath::Square(1.0f - T));
+				}
+
+				// Control speed based on distance to move target. This allows to catch up even if speed above reaches zero.
+				const float DeviantSpeed = FMath::Min(FMath::Abs(SidewaysOffset) / LookAheadDistance, 1.0f) * DesiredSpeed;
+
+				DesiredSpeed = FMath::Max(CatchupDesiredSpeed, DeviantSpeed);
+
+				// Slow down towards the end of path.
+				if (MoveTarget.IntentAtGoal == EMassMovementAction::Stand)
+				{
+					const float NormalizedDistanceToSteerTarget = FMath::Clamp(DistanceToSteerTarget / LookAheadDistance, 0.0f, 1.0f);
+					DesiredSpeed *= UE::MassMovement::ArrivalSpeedEnvelope(FMath::Max(ArrivalFade, NormalizedDistanceToSteerTarget));
+				}
+
+				// @todo: This current completely overrides steering, we probably should have one processor that resets the steering at the beginning of the frame.
+				Steering.DesiredVelocity = SteerDirection * DesiredSpeed;
+				Steering.SteeringForce = SteerK * (Steering.DesiredVelocity - Velocity.Value); // Goal force
+
+				MoveTarget.bSteeringFallingBehind = ForwardOffset < -LookAheadDistance * 0.8f;
+			}
+			else if (MoveTarget.GetCurrentAction() == EMassMovementAction::Stand)
+			{
+				constexpr float TargetMoveThresholdBase = 15.0f; // How much the target should deviate from the ghost location before update.
+				constexpr float TargetSpeedHysteresisScale = 0.85f; // How much the max speed can drop before we stop tracking it.
+				constexpr float TargetSelectionCooldown = 2.0f;	// Time between updates, varied randomly.
+
+				// "Randomize" target move threshold so that different agents react a bit differently.
+				const float PerEntityScale = (float)(Entity.Index & 8) / 7.0f;
+				const float TargetMoveThreshold = TargetMoveThresholdBase * (0.9f + PerEntityScale * 0.2f);
+				
+				if (Ghost.LastSeenActionID != MoveTarget.GetCurrentActionID())
+				{
+					// Reset when action changes. @todo: should reset only when move->stand?
+					Ghost.Location = MoveTarget.Center;
+					Ghost.Velocity = FVector::ZeroVector;
+					Ghost.LastSeenActionID = MoveTarget.GetCurrentActionID();
+
+					Ghost.SteerTarget = MoveTarget.Center;
+					Ghost.TargetMaxSpeed = 0.0f;
+					Ghost.bUpdatingTarget = false;
+					Ghost.TargetCooldown = FMath::RandRange(TargetSelectionCooldown * 0.25f, TargetSelectionCooldown);
+					Ghost.bEnteredFromMoveAction = MoveTarget.GetPreviousAction() == EMassMovementAction::Move;
+				}
+
+				Ghost.TargetCooldown = FMath::Max(0.0f, Ghost.TargetCooldown - DeltaTime);
+
+				if (!Ghost.bUpdatingTarget)
+				{
+					// Update the move target if enough time has passed and the target has moved. 
+					if (Ghost.TargetCooldown <= 0.0f && FVector::DistSquared(Ghost.SteerTarget, Ghost.Location) > FMath::Square(TargetMoveThreshold))
+					{
+						Ghost.SteerTarget = Ghost.Location;
+						Ghost.TargetMaxSpeed = 0.0f;
+						Ghost.bUpdatingTarget = true;
+						Ghost.bEnteredFromMoveAction = false;
+					}
+				}
+				else
+				{
+					// Updating target
+					Ghost.SteerTarget = Ghost.Location;
+					const float GhostSpeed = Ghost.Velocity.Length();
+					if (GhostSpeed > (Ghost.TargetMaxSpeed * TargetSpeedHysteresisScale))
+					{
+						Ghost.TargetMaxSpeed = FMath::Max(Ghost.TargetMaxSpeed, GhostSpeed);
+					}
+					else
+					{
+						// Speed is dropping, we have found the peak change, stop updating the target and start cooldown.
+						Ghost.TargetCooldown = FMath::RandRange(TargetSelectionCooldown * 0.5f, TargetSelectionCooldown);
+						Ghost.bUpdatingTarget = false;
+					}
+				}
+				
+				// Move directly towards the move target when standing.
+				FVector SteerDirection = FVector::ZeroVector;
+				float DesiredSpeed = 0.0f;
+
+				FVector Delta = Ghost.SteerTarget - CurrentLocation;
+				Delta.Z = 0.0f;
+				const float Distance = Delta.Size();
+				if (Distance > StandDeadZoneRadius)
+				{
+					SteerDirection = Delta / Distance;
+					if (Ghost.bEnteredFromMoveAction)
+					{
+						// If the current steering target is from approaching a move target, use the same speed logic as movement to ensure smooth transition.
+						const float SpeedFade = FMath::Clamp((Distance - StandDeadZoneRadius) / FMath::Max(KINDA_SMALL_NUMBER, LookAheadDistance - StandDeadZoneRadius), 0.0f, 1.0f);
+						DesiredSpeed = MoveTarget.DesiredSpeed.Get() * UE::MassMovement::CalcDirectionalSpeedScale(CurrentForward, SteerDirection) * UE::MassMovement::ArrivalSpeedEnvelope(SpeedFade);
+					}
+					else
+					{
+						// More aggressive movement when doing adjustments.
+						constexpr float AdjustmentLookAheadDistance = 100.0f; // @todo: make configurable.
+						// Intentionally not taking dead zone into account here, so that the speed does not drop to zero.
+						const float SpeedFade = FMath::Clamp(Distance / FMath::Max(KINDA_SMALL_NUMBER, LookAheadDistance - AdjustmentLookAheadDistance), 0.0f, 1.0f);
+						DesiredSpeed = MoveTarget.DesiredSpeed.Get() * UE::MassMovement::ArrivalSpeedEnvelope(SpeedFade);
+					}
+				}
+				
+				// @todo: This current completely overrides steering, we probably should have one processor that resets the steering at the beginning of the frame.
+				Steering.DesiredVelocity = SteerDirection * DesiredSpeed;
+				Steering.SteeringForce = SteerK * (Steering.DesiredVelocity - Velocity.Value); // Goal force
+				
+				MoveTarget.bSteeringFallingBehind = false;
+			}
+			else if (MoveTarget.GetCurrentAction() == EMassMovementAction::Animate)
+			{
+				// No steering when animating.
+				Steering.Reset();
+				MoveTarget.bSteeringFallingBehind = false;
+			}
+
+#if WITH_MASS_DEBUG && UNSAFE_FOR_MT
+			FColor EntityColor = FColor::White;
+			const bool bDisplayDebug = UE::MassDebug::IsDebuggingEntity(Entity, &EntityColor);
+			if (bDisplayDebug)
+			{
+				const FVector ZOffset(0,0,25);
+
+				const FColor DarkEntityColor = UE::MassMovement::Debug::MixColors(EntityColor, FColor::Black);
+				const FColor LightEntityColor = UE::MassMovement::Debug::MixColors(EntityColor, FColor::White);
+				
+				const FVector MoveTargetCenter = MoveTarget.Center + ZOffset;
+
+				// MoveTarget slack boundary
+				UE_VLOG_CIRCLE_THICK(this, LogMassNavigation, Log, MoveTargetCenter, FVector::UpVector, CurrentMovementConfig->Steering.LookAheadDistance, EntityColor, 2,
+					TEXT("%s MoveTgt %s"), *Entity.DebugGetDescription(), *UEnum::GetDisplayValueAsText(MoveTarget.IntentAtGoal).ToString());
+
+				// MoveTarget orientation
+				UE_VLOG_SEGMENT_THICK(this, LogMassNavigation, Log, MoveTargetCenter, MoveTargetCenter + MoveTarget.Forward * CurrentMovementConfig->Steering.LookAheadDistance, EntityColor, 2.0f, TEXT(""));
+
+				// MoveTarget - current location relation.
+				if (FVector::Dist2D(CurrentLocation, MoveTarget.Center) > CurrentMovementConfig->Steering.LookAheadDistance * 1.5f)
+				{
+					UE_VLOG_SEGMENT_THICK(this, LogMassNavigation, Log, MoveTargetCenter, CurrentLocation + ZOffset, FColor::Red, 1.0f, TEXT("LOST"));
+				}
+				else
+				{
+					UE_VLOG_SEGMENT_THICK(this, LogMassNavigation, Log, MoveTargetCenter, CurrentLocation + ZOffset, DarkEntityColor, 1.0f, TEXT(""));
+				}
+
+				// Steering
+				UE_VLOG_SEGMENT_THICK(this, LogMassNavigation, Log, CurrentLocation + ZOffset, CurrentLocation + Steering.DesiredVelocity + ZOffset, LightEntityColor, 2.0f,
+					TEXT("%s Steer %.1f"), *Entity.DebugGetDescription(), Steering.DesiredVelocity.Length());
+			}
+#endif // WITH_MASS_DEBUG
+			
+		}
+	});
+}
+
+#undef UNSAFE_FOR_MT
diff --git a/Engine/Plugins/AI/MassAI/Source/MassAIMovement/Private/MassZoneGraphMovementTrait.cpp b/Engine/Plugins/AI/MassAI/Source/MassAIMovement/Private/MassZoneGraphMovementTrait.cpp
new file mode 100644
--- /dev/null
+++ b/Engine/Plugins/AI/MassAI/Source/MassAIMovement/Private/MassZoneGraphMovementTrait.cpp
@@ -0,0 +1,42 @@
+﻿// Copyright Epic Games, Inc. All Rights Reserved.
+
+#include "MassZoneGraphMovementTrait.h"
+#include "MassEntityTemplateRegistry.h"
+#include "Engine/World.h"
+#include "MassZoneGraphMovementFragments.h"
+#include "MassCommonFragments.h"
+#include "MassMovementSubsystem.h"
+#include "MassMovementSettings.h"
+#include "MassAIMovementFragments.h"
+#include "VisualLogger/VisualLogger.h"
+
+void UMassZoneGraphMovementTrait::BuildTemplate(FMassEntityTemplateBuildContext& BuildContext, UWorld& World) const
+{
+	UMassMovementSubsystem* MovementSubsystem = World.GetSubsystem<UMassMovementSubsystem>();
+	if (!MovementSubsystem)
+	{
+		UE_VLOG(&World, LogMassNavigation, Error, TEXT("Failed to get MassMovement Subsystem."));
+		return;
+	}
+
+	BuildContext.AddFragmentWithDefaultInitializer<FDataFragment_AgentRadius>();
+	BuildContext.AddFragmentWithDefaultInitializer<FDataFragment_Transform>();
+	BuildContext.AddFragmentWithDefaultInitializer<FMassVelocityFragment>();
+
+	BuildContext.AddFragment<FMassMoveTargetFragment>();
+	
+	BuildContext.AddFragmentWithDefaultInitializer<FMassZoneGraphLaneLocationFragment>();
+	
+	const UMassMovementSettings* Settings = GetDefault<UMassMovementSettings>();
+	check(Settings);
+
+	const FMassMovementConfigHandle ConfigHandle = Settings->GetMovementConfigHandleByID(Config.ID);
+	FMassMovementConfigFragment& MovementConfigFragment = BuildContext.AddFragment_GetRef<FMassMovementConfigFragment>();
+	MovementConfigFragment.ConfigHandle = ConfigHandle;
+
+	BuildContext.AddFragment<FMassZoneGraphPathRequestFragment>();
+	BuildContext.AddFragment<FMassZoneGraphShortPathFragment>();
+	BuildContext.AddFragment<FMassZoneGraphCachedLaneFragment>();
+	BuildContext.AddFragment<FMassSteeringFragment>();
+	BuildContext.AddFragment<FMassSteeringGhostFragment>();
+}
diff --git a/Engine/Plugins/AI/MassAI/Source/MassAIMovement/Private/MassZoneGraphMovementUtils.cpp b/Engine/Plugins/AI/MassAI/Source/MassAIMovement/Private/MassZoneGraphMovementUtils.cpp
new file mode 100644
--- /dev/null
+++ b/Engine/Plugins/AI/MassAI/Source/MassAIMovement/Private/MassZoneGraphMovementUtils.cpp
@@ -0,0 +1,126 @@
+﻿// Copyright Epic Games, Inc. All Rights Reserved.
+
+#include "MassZoneGraphMovementUtils.h"
+
+#include "MassAIMovementFragments.h"
+#include "MassZoneGraphMovementFragments.h"
+#include "ZoneGraphSubsystem.h"
+#include "VisualLogger/VisualLogger.h"
+
+namespace UE::MassMovement
+{
+	static constexpr float InflateDistance = 200.0f; // @todo: make a setting.
+
+	bool ActivateActionMove(const UWorld& World,
+							const UObject* Requester,
+							const FLWEntity Entity,
+							const UZoneGraphSubsystem& ZoneGraphSubsystem,
+							const FMassZoneGraphLaneLocationFragment& LaneLocation,
+							const FZoneGraphShortPathRequest& PathRequest,
+							const float AgentRadius,
+							const float DesiredSpeed,
+							FMassMoveTargetFragment& MoveTarget,
+							FMassZoneGraphShortPathFragment& ShortPath,
+							FMassZoneGraphCachedLaneFragment& CachedLane)
+	{
+		ShortPath.Reset();
+		CachedLane.Reset();
+		MoveTarget.DistanceToGoal = 0.0f;
+		MoveTarget.DesiredSpeed.Set(0.0f);
+
+		if (!ensureMsgf(MoveTarget.GetCurrentAction() == EMassMovementAction::Move, TEXT("Expecting action 'Move': Invalid action %s"), MoveTarget.GetCurrentAction()))
+		{
+			return false;
+		}
+
+		const FZoneGraphStorage* ZoneGraphStorage = ZoneGraphSubsystem.GetZoneGraphStorage(LaneLocation.LaneHandle.DataHandle);
+		if (ZoneGraphStorage == nullptr)
+		{
+			UE_VLOG(Requester, LogMassNavigation, Error, TEXT("Entity [%s] move request failed: missing ZoneGraph Storage for current lane %s."),
+				*Entity.DebugGetDescription(),
+				*LaneLocation.LaneHandle.ToString());
+			return false;
+		}
+
+		MoveTarget.IntentAtGoal = EMassMovementAction::Stand;
+		MoveTarget.DesiredSpeed.Set(DesiredSpeed);
+
+		CachedLane.CacheLaneData(*ZoneGraphStorage, LaneLocation.LaneHandle, LaneLocation.DistanceAlongLane, PathRequest.TargetDistance, InflateDistance);
+		if (ShortPath.RequestPath(CachedLane, PathRequest, LaneLocation.DistanceAlongLane, AgentRadius))
+		{
+			MoveTarget.IntentAtGoal = ShortPath.EndOfPathIntent;
+			MoveTarget.DistanceToGoal = (ShortPath.NumPoints > 0) ? ShortPath.Points[ShortPath.NumPoints - 1].DistanceAlongLane.Get() : 0.0f;
+
+			UE_CVLOG(UE::MassDebug::IsDebuggingEntity(Entity),
+				Requester,
+				LogMassNavigation,
+				Log,
+				TEXT("Move %s, on lane %s, from %.1fcm to %.1fcm, next lane %s."),
+				PathRequest.bMoveReverse ? TEXT("reverse") : TEXT("forward"),
+				*LaneLocation.LaneHandle.ToString(),
+				LaneLocation.DistanceAlongLane,
+				PathRequest.TargetDistance,
+				*PathRequest.NextLaneHandle.ToString());
+		}
+		else
+		{
+			UE_VLOG(Requester, LogMassNavigation, Error, TEXT("Entity [%s] move request failed: unable to request path on lane %s."),
+				*Entity.DebugGetDescription(),
+				*LaneLocation.LaneHandle.ToString());
+			return false;
+		}
+
+		UE_VLOG(Requester, LogMassNavigation, Log, TEXT("Entity [%s] successfully requested %s"), *Entity.DebugGetDescription(), *MoveTarget.ToString());
+		return true;
+	}
+
+	bool ActivateActionStand(const UWorld& World,
+							 const UObject* Requester,
+							 const FLWEntity Entity,
+							 const UZoneGraphSubsystem& ZoneGraphSubsystem,
+							 const FMassZoneGraphLaneLocationFragment& LaneLocation,
+							 const float DesiredSpeed,
+							 FMassMoveTargetFragment& MoveTarget,
+							 FMassZoneGraphShortPathFragment& ShortPath,
+							 FMassZoneGraphCachedLaneFragment& CachedLane)
+	{
+		ShortPath.Reset();
+		CachedLane.Reset();
+		MoveTarget.DistanceToGoal = 0.0f;
+		MoveTarget.DesiredSpeed.Set(0.0f);
+
+		if (!ensureMsgf(MoveTarget.GetCurrentAction() == EMassMovementAction::Stand, TEXT("Expecting action 'Stand': Invalid action %s"), MoveTarget.GetCurrentAction()))
+		{
+			return false;
+		}
+
+		const FZoneGraphStorage* ZoneGraphStorage = ZoneGraphSubsystem.GetZoneGraphStorage(LaneLocation.LaneHandle.DataHandle);
+
+		MoveTarget.IntentAtGoal = EMassMovementAction::Stand;
+		MoveTarget.DesiredSpeed.Set(DesiredSpeed);
+
+		CachedLane.CacheLaneData(*ZoneGraphStorage, LaneLocation.LaneHandle, LaneLocation.DistanceAlongLane, LaneLocation.DistanceAlongLane, InflateDistance);
+
+		UE_VLOG(Requester, LogMassNavigation, Log, TEXT("Entity [%s] successfully requested %s"), *Entity.DebugGetDescription(), *MoveTarget.ToString());
+		return true;
+	}
+
+	bool ActivateActionAnimate(const UWorld& World,
+							   const UObject* Requester,
+							   const FLWEntity Entity,
+							   FMassMoveTargetFragment& MoveTarget)
+	{
+		MoveTarget.DistanceToGoal = 0.0f;
+		MoveTarget.DesiredSpeed.Set(0.0f);
+
+		if (!ensureMsgf(MoveTarget.GetCurrentAction() == EMassMovementAction::Animate, TEXT("Expecting action 'Animate': Invalid action %s"), MoveTarget.GetCurrentAction()))
+		{
+			return false;
+		}
+
+		MoveTarget.IntentAtGoal = EMassMovementAction::Stand;
+
+		UE_VLOG(Requester, LogMassNavigation, Log, TEXT("Entity [%s] successfully requested %s"), *Entity.DebugGetDescription(), *MoveTarget.ToString());
+		return true;
+	}
+}
diff --git a/Engine/Plugins/AI/MassAI/Source/MassAIMovement/Private/MovementProcessor.cpp b/Engine/Plugins/AI/MassAI/Source/MassAIMovement/Private/MovementProcessor.cpp
new file mode 100644
--- /dev/null
+++ b/Engine/Plugins/AI/MassAI/Source/MassAIMovement/Private/MovementProcessor.cpp
@@ -0,0 +1,225 @@
+// Copyright Epic Games, Inc. All Rights Reserved.
+
+#include "MovementProcessor.h"
+#include "AIHelpers.h"
+#include "NavigationSystem.h"
+#include "Components/HierarchicalInstancedStaticMeshComponent.h"
+#include "DrawDebugHelpers.h"
+#include "MassCommonFragments.h"
+#include "MassAIMovementFragments.h"
+#include "MassMovementTypes.h"
+#include "VisualLogger/VisualLogger.h"
+
+namespace UE::MassMovement
+{
+	struct FAvoidanceAgent
+	{
+		FVector Location;
+		FVector Velocity;
+		float Radius;
+
+		FAvoidanceAgent(const FVector& InLocation, const FVector& InVelocity, const float InRadius = 1.f)
+			: Location(InLocation), Velocity(InVelocity), Radius(InRadius)
+		{
+		}
+	};
+
+	float TimeToCollision(const FAvoidanceAgent& Agent, const FAvoidanceAgent& Obstacle)
+	{
+		const float RadiusSum = Agent.Radius + Obstacle.Radius;
+		const FVector VecToObstacle = Obstacle.Location - Agent.Location;
+		const float C = FVector::DotProduct(VecToObstacle, VecToObstacle) - RadiusSum * RadiusSum;
+
+		if (C < 0.f) //agents are colliding
+		{
+			return 0.f;
+		}
+		const FVector VelocityDelta = Agent.Velocity - Obstacle.Velocity;
+		const float A = FVector::DotProduct(VelocityDelta, VelocityDelta);
+		const float B = FVector::DotProduct(VecToObstacle, VelocityDelta);
+		const float Discriminator = B * B - A * C;
+		if (Discriminator <= 0)
+		{
+			return FLT_MAX;
+		}
+		const float Tau = (B - FMath::Sqrt(Discriminator)) / A;
+		return (Tau < 0) ? FLT_MAX : Tau;
+	}
+
+	namespace Tweakables
+	{
+		float AvoidDistanceCutOff = 500.f;
+		float AvoidAgentRadius = 50.f;
+		float AvoidMaxForce = 100.f;
+		float AvoidTimeHorizon = 5.f;
+	} // Tweakables
+
+	FAutoConsoleVariableRef Vars[] = {
+		FAutoConsoleVariableRef(TEXT("ai.mass.AvoidDistance"), Tweakables::AvoidDistanceCutOff, TEXT(""), ECVF_Default), FAutoConsoleVariableRef(TEXT("ai.mass.AvoidRadius"), Tweakables::AvoidAgentRadius, TEXT(""), ECVF_Default), FAutoConsoleVariableRef(TEXT("ai.mass.AvoidMaxForce"), Tweakables::AvoidMaxForce, TEXT(""), ECVF_Default), FAutoConsoleVariableRef(TEXT("ai.mass.AvoidTimeHorizon"), Tweakables::AvoidTimeHorizon, TEXT(""), ECVF_Default)
+	};
+
+ 	static bool DebugIsSelected(const FLWEntity& Entity)
+ 	{
+ #if WITH_MASS_DEBUG
+ 		return UE::MassDebug::IsDebuggingEntity(Entity);
+ #else
+ 	return false;
+ #endif // WITH_MASS_DEBUG
+ 	}
+
+ 	static void DebugDrawLine(const UObject* LogOwner, const UWorld* World, const FVector& Start, const FVector& End, const FColor& Color, const bool bPersistent = false)
+ 	{
+ #if WITH_MASS_DEBUG
+ 		UE_VLOG_SEGMENT(LogOwner, LogNavigation, Log, Start, End, Color, TEXT(""));
+
+ 		if (World)
+ 		{
+ 			DrawDebugLine(World, Start, End, Color, bPersistent);
+ 		}
+ #endif // WITH_MASS_DEBUG
+ 	}
+
+ 	static void DebugDrawArrow(const UObject* LogOwner, const UWorld* World, const FVector& Start, const FVector& End, const FColor& Color)
+ 	{
+ #if WITH_MASS_DEBUG
+ 		UE_VLOG_ARROW(LogOwner, LogNavigation, Log, Start, End, Color, TEXT(""));
+
+ 		if (World)
+ 		{
+ 			DrawDebugDirectionalArrow(World, Start, End, /*arrow size = */20.f, Color);
+ 		}
+ #endif // WITH_MASS_DEBUG
+ 	}
+
+ 	static void DebugDrawCylinder(const UObject* LogOwner, const UWorld* World, const FVector& Start, const FVector& End, const float Radius, const FColor& Color)
+ 	{
+ #if WITH_MASS_DEBUG
+ 		UE_VLOG_CYLINDER(LogOwner, LogNavigation, Log, Start, End, Radius, Color, TEXT(""));
+
+ 		if (World)
+ 		{
+ 			DrawDebugCylinder(World, Start, End, Radius, /*segments = */16, Color);
+ 		}
+ #endif // WITH_MASS_DEBUG
+	}
+} // namespace UE::Movement
+
+UMassProcessor_Movement::UMassProcessor_Movement()
+{
+	bAutoRegisterWithProcessingPhases = false;
+}
+
+//----------------------------------------------------------------------//
+// UMassProcessor_Movement 
+//----------------------------------------------------------------------//
+void UMassProcessor_Movement::ConfigureQueries()
+{
+	EntityQuery.AddRequirement<FDataFragment_Transform>(ELWComponentAccess::ReadWrite);
+	EntityQuery.AddRequirement<FMassVelocityFragment>(ELWComponentAccess::ReadWrite);
+}
+
+void UMassProcessor_Movement::Execute(UEntitySubsystem& EntitySubsystem, FLWComponentSystemExecutionContext& Context)
+{
+	QUICK_SCOPE_CYCLE_COUNTER(MovementProcessor_Run);
+
+	const float TimeDelta = Context.GetDeltaTimeSeconds();
+
+	EntityQuery.ForEachEntityChunk(EntitySubsystem, Context, [this, TimeDelta](FLWComponentSystemExecutionContext& Context)
+		{
+			const int32 NumEntities = Context.GetEntitiesNum();
+			const TArrayView<FDataFragment_Transform> LocationList = Context.GetMutableComponentView<FDataFragment_Transform>();
+			const TArrayView<FMassVelocityFragment> VelocityList = Context.GetMutableComponentView<FMassVelocityFragment>();
+			for (int32 i = 0; i < NumEntities; ++i)
+			{
+				LocationList[i].GetMutableTransform().AddToTranslation(VelocityList[i].Value * TimeDelta);
+
+				if (const TOptional<float> Yaw = UE::AI::GetYawFromVector(VelocityList[i].Value))
+				{
+					FQuat Rotation(FVector::UpVector, Yaw.GetValue());
+					LocationList[i].GetMutableTransform().SetRotation(Rotation);
+				}
+			}
+		});
+}
+
+//----------------------------------------------------------------------//
+//  UMassProcessor_AgentMovement
+//----------------------------------------------------------------------//
+UMassProcessor_AgentMovement::UMassProcessor_AgentMovement()
+{
+	bAutoRegisterWithProcessingPhases = false;
+}
+
+void UMassProcessor_AgentMovement::ConfigureQueries()
+{
+	Super::ConfigureQueries();
+	EntityQuery.AddRequirement<FDataFragment_AgentRadius>(ELWComponentAccess::ReadOnly);
+}
+
+//----------------------------------------------------------------------//
+//  UMassFragmentDestructor_AvoidanceObstacleRemover
+//----------------------------------------------------------------------//
+UMassAvoidanceObstacleRemoverFragmentDestructor::UMassAvoidanceObstacleRemoverFragmentDestructor()
+{
+	FragmentType = FMassAvoidanceObstacleGridCellLocationFragment::StaticStruct();
+	ExecutionFlags = (int32)(EProcessorExecutionFlags::All);
+}
+
+void UMassAvoidanceObstacleRemoverFragmentDestructor::ConfigureQueries()
+{
+	EntityQuery.AddRequirement<FMassAvoidanceObstacleGridCellLocationFragment>(ELWComponentAccess::ReadWrite);
+}
+
+void UMassAvoidanceObstacleRemoverFragmentDestructor::Initialize(UObject& Owner)
+{
+	Super::Initialize(Owner);
+	WeakMovementSubsystem = UWorld::GetSubsystem<UMassMovementSubsystem>(Owner.GetWorld());
+}
+
+void UMassAvoidanceObstacleRemoverFragmentDestructor::Execute(UEntitySubsystem& EntitySubsystem, FLWComponentSystemExecutionContext& Context)
+{
+	UMassMovementSubsystem* MovementSubsystem = WeakMovementSubsystem.Get();
+	if (!MovementSubsystem)
+	{
+		return;
+	}
+
+	EntityQuery.ForEachEntityChunk(EntitySubsystem, Context, [this, MovementSubsystem](FLWComponentSystemExecutionContext& Context)
+		{
+			const int32 NumEntities = Context.GetEntitiesNum();
+			const TArrayView<FMassAvoidanceObstacleGridCellLocationFragment> AvoidanceObstacleCellLocationList = Context.GetMutableComponentView<FMassAvoidanceObstacleGridCellLocationFragment>();
+
+			for (int32 i = 0; i < NumEntities; ++i)
+			{
+				FMassAvoidanceObstacleItem ObstacleItem;
+				ObstacleItem.Entity = Context.GetEntity(i);
+				MovementSubsystem->GetGridMutable().Remove(ObstacleItem, AvoidanceObstacleCellLocationList[i].CellLoc);
+			}
+		});
+}
+
+
+//----------------------------------------------------------------------//
+//  UVelocityFragmentInitializer
+//----------------------------------------------------------------------//
+UVelocityFragmentInitializer::UVelocityFragmentInitializer()
+{
+	FragmentType = FMassVelocityFragment::StaticStruct();
+}
+
+void UVelocityFragmentInitializer::ConfigureQueries()
+{
+	EntityQuery.AddRequirement<FMassVelocityFragment>(ELWComponentAccess::ReadWrite);
+}
+
+void UVelocityFragmentInitializer::Execute(UEntitySubsystem& EntitySubsystem, FLWComponentSystemExecutionContext& Context)
+{
+	EntityQuery.ForEachEntityChunk(EntitySubsystem, Context, [this](FLWComponentSystemExecutionContext& Context)
+		{
+			const TArrayView<FMassVelocityFragment> VelocitiesList = Context.GetMutableComponentView<FMassVelocityFragment>();
+			for (FMassVelocityFragment& VelocityFragment : VelocitiesList)
+			{
+				VelocityFragment.Value = FVector(FMath::FRandRange(-100.f, 100.f), FMath::FRandRange(-100.f, 100.f), 0);
+			}
+		});
+}
diff --git a/Engine/Plugins/AI/MassAI/Source/MassAIMovement/Private/NavigationProcessor.cpp b/Engine/Plugins/AI/MassAI/Source/MassAIMovement/Private/NavigationProcessor.cpp
new file mode 100644
--- /dev/null
+++ b/Engine/Plugins/AI/MassAI/Source/MassAIMovement/Private/NavigationProcessor.cpp
@@ -0,0 +1,292 @@
+// Copyright Epic Games, Inc. All Rights Reserved.
+
+#include "NavigationProcessor.h"
+
+#include "MassCommonUtils.h"
+#include "LWCCommandBuffer.h"
+#include "MassCommonFragments.h"
+#include "MassAIMovementFragments.h"
+#include "MassMovementTypes.h"
+#include "MassSignals/Public/MassSignalSubsystem.h"
+#include "Math/UnrealMathUtility.h"
+#include "VisualLogger/VisualLogger.h"
+#include "GameFramework/GameStateBase.h"
+#include "Misc/Timespan.h"
+#include "MassLODTypes.h"
+#include "MassSimulationLOD.h"
+#include "MassAvoidanceSettings.h"
+
+#define UNSAFE_FOR_MT 0
+#define MOVEMENT_DEBUGDRAW 0	// Set to 1 to see heading debugdraw
+
+namespace UE::MassMovement
+{
+
+	static float Damp(const float X, const float Goal, const float HalfLife, const float DeltaTime)
+	{
+		return FMath::Lerp(X, Goal, 1.f - FMath::Pow(2.f, -DeltaTime / (HalfLife + KINDA_SMALL_NUMBER)));
+	}
+
+	// Calculates yaw angle from direction vector.
+	static float GetYawFromDirection(const FVector Direction)
+	{
+		return FMath::Atan2(Direction.Y, Direction.X);
+	}
+
+	// Wraps and angle to range -PI..PI. Angle in radians.
+	static float WrapAngle(const float Angle)
+	{
+		float WrappedAngle = FMath::Fmod(Angle, PI*2.0f);
+
+		if (WrappedAngle > PI)
+		{
+			WrappedAngle -= PI*2.0f;
+		}
+
+		if (WrappedAngle < -PI)
+		{
+			WrappedAngle += PI*2.0f;
+		}
+
+		return WrappedAngle;
+	}
+
+	// Linearly interpolates between two angles (in Radians).
+	static float LerpAngle(const float AngleA, const float AngleB, const float T)
+	{
+		const float DeltaAngle = WrapAngle(AngleB - AngleA);
+		return AngleA + DeltaAngle * T;
+	}
+
+	// Exponential smooth from current angle to target angle. Angles in radians.
+	static float ExponentialSmoothingAngle(const float Angle, const float TargetAngle, const float DeltaTime, const float SmoothingTime)
+	{
+		// Note: based on FMath::ExponentialSmoothingApprox().
+		if (SmoothingTime < KINDA_SMALL_NUMBER)
+		{
+			return TargetAngle;
+		}
+		const float A = DeltaTime / SmoothingTime;
+		const float Exp = FMath::InvExpApprox(A);
+		return TargetAngle + WrapAngle(Angle - TargetAngle) * Exp;
+	}
+
+} // UE::MassMovement
+
+//----------------------------------------------------------------------//
+//  UMassApplyVelocityMoveTargetProcessor
+//----------------------------------------------------------------------//
+UMassApplyVelocityMoveTargetProcessor::UMassApplyVelocityMoveTargetProcessor()
+{
+	ExecutionFlags = (int32)EProcessorExecutionFlags::All;
+	ExecutionOrder.ExecuteInGroup = UE::Mass::ProcessorGroupNames::Avoidance;
+	ExecutionOrder.ExecuteAfter.Add(TEXT("MassAvoidanceProcessor"));
+}
+
+void UMassApplyVelocityMoveTargetProcessor::ConfigureQueries()
+{
+	HighResEntityQuery.AddRequirement<FMassMoveTargetFragment>(ELWComponentAccess::ReadOnly);
+	HighResEntityQuery.AddRequirement<FMassVelocityFragment>(ELWComponentAccess::ReadWrite);
+	HighResEntityQuery.AddRequirement<FDataFragment_Transform>(ELWComponentAccess::ReadWrite);
+	HighResEntityQuery.AddRequirement<FMassSteeringFragment>(ELWComponentAccess::ReadWrite);
+	HighResEntityQuery.AddTagRequirement<FMassOffLODTag>(ELWComponentPresence::None);
+
+	LowResEntityQuery_Conditional.AddRequirement<FDataFragment_Transform>(ELWComponentAccess::ReadWrite);
+	LowResEntityQuery_Conditional.AddRequirement<FMassMoveTargetFragment>(ELWComponentAccess::ReadOnly);
+	LowResEntityQuery_Conditional.AddTagRequirement<FMassOffLODTag>(ELWComponentPresence::All);
+	LowResEntityQuery_Conditional.AddChunkRequirement<FMassSimulationVariableTickChunkFragment>(ELWComponentAccess::ReadOnly);
+	LowResEntityQuery_Conditional.SetChunkFilter(&FMassSimulationVariableTickChunkFragment::ShouldTickChunkThisFrame);
+}
+
+void UMassApplyVelocityMoveTargetProcessor::Execute(UEntitySubsystem& EntitySubsystem,
+													FLWComponentSystemExecutionContext& Context)
+{
+	// Clamp max delta time to avoid force explosion on large time steps (i.e. during initialization).
+	const float TimeDelta = FMath::Min(0.1f, Context.GetDeltaTimeSeconds());
+
+	const UMassAvoidanceSettings* Settings = UMassAvoidanceSettings::Get();
+	check(Settings);
+
+	{
+		QUICK_SCOPE_CYCLE_COUNTER(HighRes);
+
+		HighResEntityQuery.ForEachEntityChunk(EntitySubsystem, Context, [this, Settings, TimeDelta](FLWComponentSystemExecutionContext& Context)
+		{
+			const int32 NumEntities = Context.GetEntitiesNum();
+			constexpr float ZDamperHalfLife = 0.2f;
+
+			const float OrientationEndOfPathHeadingAnticipation = FMath::Max(KINDA_SMALL_NUMBER, Settings->OrientationEndOfPathHeadingAnticipation);
+			const float OrientationSmoothingTime = Settings->OrientationSmoothingTime;
+			constexpr float OrientationBlendWhileMoving = 0.4f;
+			constexpr float OrientationBlendWhileStanding = 0.95f;
+
+			const TConstArrayView<FMassMoveTargetFragment> MoveTargetList = Context.GetComponentView<FMassMoveTargetFragment>();
+			const TArrayView<FDataFragment_Transform> LocationList = Context.GetMutableComponentView<FDataFragment_Transform>();
+			const TArrayView<FMassSteeringFragment> SteeringList = Context.GetMutableComponentView<FMassSteeringFragment>();
+			const TArrayView<FMassVelocityFragment> VelocityList = Context.GetMutableComponentView<FMassVelocityFragment>();
+
+			for (int32 EntityIndex = 0; EntityIndex < NumEntities; ++EntityIndex)
+			{
+				const FMassMoveTargetFragment& MoveTarget = MoveTargetList[EntityIndex];
+				FMassSteeringFragment& Steering = SteeringList[EntityIndex];
+				FMassVelocityFragment& Velocity = VelocityList[EntityIndex];
+
+				// Do not touch transform at all when animating
+				if (MoveTarget.GetCurrentAction() == EMassMovementAction::Animate)
+				{
+					Velocity.Value = FVector::ZeroVector;
+					Steering.Reset();
+					continue;
+				}
+
+				// Update velocity from steering forces.
+				Velocity.Value += Steering.SteeringForce * TimeDelta;
+				Velocity.Value.Z = 0.0f;
+
+				constexpr float LowSpeedThreshold = 3.0f;
+
+				if (MoveTargetList[EntityIndex].GetCurrentAction() == EMassMovementAction::Stand)
+				{
+					// Clamp small velocities in stand to zero to avoid tiny drifting.
+					if (Velocity.Value.SquaredLength() < FMath::Square(LowSpeedThreshold))
+					{
+						Velocity.Value = FVector::ZeroVector;
+					}
+				}
+
+				FVector DeltaLoc = Velocity.Value * TimeDelta;
+
+				// Apply DeltaLoc on X,Y and set the Z of the current lane location.
+				FTransform& CurrentTransform = LocationList[EntityIndex].GetMutableTransform();
+
+				const float Z = CurrentTransform.GetLocation().Z;
+				const float TargetZ = MoveTarget.Center.Z;
+				const float NewZ = UE::MassMovement::Damp(Z, TargetZ, ZDamperHalfLife, TimeDelta);
+				DeltaLoc.Z = NewZ - Z;
+
+				CurrentTransform.AddToTranslation(DeltaLoc);
+
+				// Direction
+				const FVector CurrentForward = CurrentTransform.GetRotation().GetForwardVector();
+				const float CurrentHeading = UE::MassMovement::GetYawFromDirection(CurrentForward);
+
+				float Blend = 0.0f;
+				if (MoveTarget.GetCurrentAction() == EMassMovementAction::Move)
+				{
+					if (MoveTarget.IntentAtGoal == EMassMovementAction::Stand && MoveTarget.DistanceToGoal < OrientationEndOfPathHeadingAnticipation)
+					{
+						// Fade towards the movement target direction at the end of the path.
+						const float Strength = 1.0f - FMath::Square(FMath::Clamp(MoveTarget.DistanceToGoal / OrientationEndOfPathHeadingAnticipation, 0.0f, 1.0));
+						Blend = FMath::Lerp(Blend, OrientationBlendWhileStanding, Strength);
+					}
+					else
+					{
+						Blend = OrientationBlendWhileMoving;
+					}
+				}
+				else // Stand
+				{
+					Blend = OrientationBlendWhileStanding;
+				}
+
+				const float VelocityHeading = UE::MassMovement::GetYawFromDirection(Velocity.Value);
+				const float MovementHeading = UE::MassMovement::GetYawFromDirection(MoveTarget.Forward);
+				const float DesiredHeading = UE::MassMovement::LerpAngle(VelocityHeading, MovementHeading, Blend);
+				const float NewHeading = UE::MassMovement::ExponentialSmoothingAngle(CurrentHeading, DesiredHeading, TimeDelta, OrientationSmoothingTime);
+
+				FQuat Rotation(FVector::UpVector, NewHeading);
+				CurrentTransform.SetRotation(Rotation);
+			}
+		});
+	}
+
+	{
+		QUICK_SCOPE_CYCLE_COUNTER(LowRes);
+
+		LowResEntityQuery_Conditional.ForEachEntityChunk(EntitySubsystem, Context, [this](FLWComponentSystemExecutionContext& Context)
+			{
+				const int32 NumEntities = Context.GetEntitiesNum();
+
+				const TArrayView<FDataFragment_Transform> LocationList = Context.GetMutableComponentView<FDataFragment_Transform>();
+				const TConstArrayView<FMassMoveTargetFragment> MoveTargetList = Context.GetComponentView<FMassMoveTargetFragment>();
+
+				for (int32 EntityIndex = 0; EntityIndex < NumEntities; ++EntityIndex)
+				{
+					FTransform& CurrentTransform = LocationList[EntityIndex].GetMutableTransform();
+					const FMassMoveTargetFragment& MoveTarget = MoveTargetList[EntityIndex];
+
+					// Snap position to move target directly
+					CurrentTransform.SetLocation(MoveTarget.Center);
+					CurrentTransform.SetRotation(FQuat::FindBetweenNormals(FVector::ForwardVector, MoveTarget.Forward));
+				}
+			});
+	}
+}
+
+//----------------------------------------------------------------------//
+//  UMassDynamicObstacleProcessor
+//----------------------------------------------------------------------//
+UMassDynamicObstacleProcessor::UMassDynamicObstacleProcessor()
+{
+	bAutoRegisterWithProcessingPhases = false;
+	ExecutionOrder.ExecuteAfter.Add(UE::Mass::ProcessorGroupNames::SyncWorldToMass);
+}
+
+void UMassDynamicObstacleProcessor::ConfigureQueries()
+{
+	EntityQuery_Conditional.AddRequirement<FDataFragment_Transform>(ELWComponentAccess::ReadOnly);
+	EntityQuery_Conditional.AddRequirement<FDataFragment_AgentRadius>(ELWComponentAccess::ReadOnly);
+	EntityQuery_Conditional.AddRequirement<FMassDynamicObstacleFragment>(ELWComponentAccess::ReadWrite);
+	EntityQuery_Conditional.AddChunkRequirement<FMassSimulationVariableTickChunkFragment>(ELWComponentAccess::ReadOnly);
+	EntityQuery_Conditional.SetChunkFilter(&FMassSimulationVariableTickChunkFragment::ShouldTickChunkThisFrame);
+}
+
+void UMassDynamicObstacleProcessor::Execute(UEntitySubsystem& EntitySubsystem, FLWComponentSystemExecutionContext& Context)
+{
+	EntityQuery_Conditional.ForEachEntityChunk(EntitySubsystem, Context, [this](FLWComponentSystemExecutionContext& Context)
+		{
+			const int32 NumEntities = Context.GetEntitiesNum();
+
+			const TConstArrayView<FDataFragment_Transform> LocationList = Context.GetComponentView<FDataFragment_Transform>();
+			const TConstArrayView<FDataFragment_AgentRadius> RadiusList = Context.GetComponentView<FDataFragment_AgentRadius>();
+			const TArrayView<FMassDynamicObstacleFragment> ObstacleDataList = Context.GetMutableComponentView<FMassDynamicObstacleFragment>();
+
+			const FDateTime Now = FDateTime::UtcNow();
+
+			for (int32 EntityIndex = 0; EntityIndex < NumEntities; ++EntityIndex)
+			{
+				// @todo: limit update frequency, this does not need to occur every frame
+				const FVector Position = LocationList[EntityIndex].GetTransform().GetLocation();
+				const float Radius = RadiusList[EntityIndex].Radius;
+				FMassDynamicObstacleFragment& Obstacle = ObstacleDataList[EntityIndex];
+
+				UE_VLOG_LOCATION(this, LogMassDynamicObstacle, Display, Position, Radius, Obstacle.bHasStopped ? FColor::Red : FColor::Green, TEXT(""));
+
+				if ((Position - Obstacle.LastPosition).SquaredLength() < FMath::Square(DistanceBuffer))
+				{
+					const FTimespan TimeElapsed = Now - Obstacle.LastMovedTimeStamp;
+					if (TimeElapsed.GetSeconds() > DelayBeforeStopNotification && !Obstacle.bHasStopped)
+					{
+						// The obstacle hasn't moved for a while.
+						Obstacle.bHasStopped = true;
+						OnStop(Obstacle, Radius);
+					}
+				}
+				else
+				{
+					// Update position and time stamp
+					Obstacle.LastPosition = Position;
+					Obstacle.LastMovedTimeStamp = Now;
+
+					// If the obstacle had stopped, signal the move.
+					if (Obstacle.bHasStopped)
+					{
+						Obstacle.bHasStopped = false;
+						OnMove(Obstacle);
+					}
+				}
+			}
+		});
+}
+
+#undef UNSAFE_FOR_MT
