AggregateTicking and Pipe merged into Engine/Plugins/Runtime/MassEntity

diff --git a/Engine/Plugins/AI/MassAI/Source/MassAIMovement/Private/SnapToNavigationProcessor.cpp b/Engine/Plugins/AI/MassAI/Source/MassAIMovement/Private/SnapToNavigationProcessor.cpp
new file mode 100644
--- /dev/null
+++ b/Engine/Plugins/AI/MassAI/Source/MassAIMovement/Private/SnapToNavigationProcessor.cpp
@@ -0,0 +1,69 @@
+// Copyright Epic Games, Inc. All Rights Reserved.
+
+#include "SnapToNavigationProcessor.h"
+#include "MassCommonFragments.h"
+#include "MassCommonTypes.h"
+#include "AI/Navigation/NavigationTypes.h"
+#include "MassAIMovementTypes.h"
+#include "NavigationData.h"
+#include "NavigationSystem.h"
+
+//----------------------------------------------------------------------//
+// USnapToNavigationProcessor 
+//----------------------------------------------------------------------//
+USnapToNavigationProcessor::USnapToNavigationProcessor()
+{
+	ExecutionOrder.ExecuteAfter.Add(UE::Mass::ProcessorGroupNames::Avoidance);
+	ExecutionOrder.ExecuteBefore.Add(UE::Mass::ProcessorGroupNames::UpdateWorldFromMass);
+}
+
+void USnapToNavigationProcessor::ConfigureQueries()
+{
+	EntityQuery.AddRequirement<FDataFragment_Transform>(ELWComponentAccess::ReadWrite);
+	EntityQuery.AddRequirement<FDataFragment_NavLocation>(ELWComponentAccess::ReadWrite);
+}
+
+void USnapToNavigationProcessor::Execute(UEntitySubsystem& EntitySubsystem, FLWComponentSystemExecutionContext& Context)
+{
+	ANavigationData* NavData = WeakNavData.Get();
+	if (NavData == nullptr)
+	{
+		return;
+	}
+
+	QUICK_SCOPE_CYCLE_COUNTER(UMovementProcessor_Run);
+
+	EntityQuery.ForEachEntityChunk(EntitySubsystem, Context, [this, NavData](FLWComponentSystemExecutionContext& Context)
+		{
+			const int32 NumEntities = Context.GetEntitiesNum();
+			const TArrayView<FDataFragment_Transform> LocationList = Context.GetMutableComponentView<FDataFragment_Transform>();
+			const TArrayView<FDataFragment_NavLocation> NavLocationList = Context.GetMutableComponentView<FDataFragment_NavLocation>();
+
+			TArray<FNavigationProjectionWork> Workload;
+			for (int32 Index = 0; Index < NumEntities; ++Index)
+			{
+				Workload.Add(FNavigationProjectionWork(FNavLocation(LocationList[Index].GetTransform().GetLocation(), NavLocationList[Index].NodeRef)));
+			}
+
+			NavData->BatchProjectPoints(Workload, NavData->GetDefaultQueryExtent());
+
+			for (int32 Index = 0; Index < NumEntities; ++Index)
+			{
+				LocationList[Index].GetMutableTransform().SetLocation(Workload[Index].OutLocation.Location);
+				NavLocationList[Index].NodeRef = Workload[Index].OutLocation.NodeRef;
+			}
+		});
+}
+
+void USnapToNavigationProcessor::Initialize(UObject& InOwner)
+{
+	Super::Initialize(InOwner);
+
+	ANavigationData* NavData = Cast<ANavigationData>(&InOwner);
+	if (NavData == nullptr)
+	{
+		UNavigationSystemV1* NavSys = FNavigationSystem::GetCurrent<UNavigationSystemV1>(InOwner.GetWorld());
+		NavData = NavSys ? NavSys->GetDefaultNavDataInstance() : nullptr;
+	}
+	WeakNavData = NavData;
+}
diff --git a/Engine/Plugins/AI/MassAI/Source/MassAIMovement/Public/AssignRandomNavLocationProcessor.h b/Engine/Plugins/AI/MassAI/Source/MassAIMovement/Public/AssignRandomNavLocationProcessor.h
new file mode 100644
--- /dev/null
+++ b/Engine/Plugins/AI/MassAI/Source/MassAIMovement/Public/AssignRandomNavLocationProcessor.h
@@ -0,0 +1,34 @@
+// Copyright Epic Games, Inc. All Rights Reserved.
+
+#pragma once
+
+#include "MassProcessor.h"
+#include "AI/Navigation/NavigationTypes.h"
+#include "AssignRandomNavLocationProcessor.generated.h"
+
+
+class UNavigationSystemV1;
+
+UCLASS()
+class MASSAIMOVEMENT_API UAssignRandomNavLocationProcessor : public UPipeProcessor
+{
+	GENERATED_BODY()
+
+public:
+	UAssignRandomNavLocationProcessor();
+
+protected:
+	virtual void ConfigureQueries() override;
+	virtual void Execute(UEntitySubsystem& EntitySubsystem, FLWComponentSystemExecutionContext& Context) override;
+	virtual void Initialize(UObject& InOwner) override;
+
+	UPROPERTY()
+	UNavigationSystemV1* NavigationSystem = nullptr;
+
+	UPROPERTY(EditAnywhere, Category="Processor")
+	float Radius = 5000.f;
+
+	FVector Origin = FNavigationSystem::InvalidLocation;
+
+	FLWComponentQuery EntityQuery;
+};
\ No newline at end of file
diff --git a/Engine/Plugins/AI/MassAI/Source/MassAIMovement/Public/MassAIMovementFragments.h b/Engine/Plugins/AI/MassAI/Source/MassAIMovement/Public/MassAIMovementFragments.h
new file mode 100644
--- /dev/null
+++ b/Engine/Plugins/AI/MassAI/Source/MassAIMovement/Public/MassAIMovementFragments.h
@@ -0,0 +1,247 @@
+// Copyright Epic Games, Inc. All Rights Reserved.
+
+#pragma once
+
+#include "CoreMinimal.h"
+#include "MassCommonTypes.h"
+#include "MassAIMovementTypes.h"
+#include "ZoneGraphAnnotationTypes.h"
+#include "ZoneGraphTypes.h"
+#include "MassMovementFragments.h"
+#include "MassAIMovementFragments.generated.h"
+
+class UWorld;
+
+struct FMassAvoidanceColliderInstance
+{
+	float Radius = 0.f;
+	float Offset = 0.f;
+};
+
+enum class EMassColliderType : uint8
+{
+	Circle,
+	Pill,
+};
+
+struct FMassCircleCollider
+{
+	FMassCircleCollider() = default;
+	FMassCircleCollider(const float Radius) : Radius(Radius) {}
+	float Radius = 0.f;
+};
+
+struct FMassPillCollider
+{
+	FMassPillCollider() = default;
+	FMassPillCollider(const float Radius, const float HalfLength) : Radius(Radius), HalfLength(HalfLength) {}
+	float Radius = 0.f;
+	float HalfLength = 0.f;
+};
+
+/** Fragment holding data for avoidance coliders */
+USTRUCT()
+struct MASSAIMOVEMENT_API FMassAvoidanceColliderFragment : public FLWComponentData
+{
+	GENERATED_BODY()
+
+	FMassAvoidanceColliderFragment()
+	{
+		Type = EMassColliderType::Circle;
+		Data[0] = 0.f;
+		Data[1] = 0.f;
+	}
+
+	FMassAvoidanceColliderFragment(const FMassCircleCollider& Circle)
+	{
+		Type = EMassColliderType::Circle;
+		Data[0] = Circle.Radius;
+		Data[1] = 0.f;
+	}
+
+	FMassAvoidanceColliderFragment(const FMassPillCollider& Pill)
+	{
+		Type = EMassColliderType::Pill;
+		Data[0] = Pill.Radius;
+		Data[1] = Pill.HalfLength;
+	}
+	
+	FMassCircleCollider GetCircleCollider() const
+	{
+		check(Type == EMassColliderType::Circle);
+		return FMassCircleCollider(Data[0]);
+	}
+
+	FMassPillCollider GetPillCollider() const
+	{
+		check(Type == EMassColliderType::Pill);
+		return FMassPillCollider(Data[0], Data[1]);
+	}
+
+	float Data[2];
+	EMassColliderType Type;
+};
+
+/** Cell location for dynamic obstacles */
+USTRUCT()
+struct MASSAIMOVEMENT_API FMassAvoidanceObstacleGridCellLocationFragment : public FLWComponentData
+{
+	GENERATED_BODY()
+	FAvoidanceObstacleHashGrid2D::FCellLocation CellLoc;
+};
+
+/** Move target. */
+USTRUCT()
+struct MASSAIMOVEMENT_API FMassMoveTargetFragment : public FLWComponentData
+{
+	GENERATED_BODY()
+
+	FMassMoveTargetFragment() : bNetDirty(false), bOffBoundaries(false), bSteeringFallingBehind(false) {}
+
+	/** To setup current action from the authoritative world */
+	void CreateNewAction(const EMassMovementAction InAction, const UWorld& InWorld);
+
+	/** To setup current action from replicated data */
+	void CreateReplicatedAction(const EMassMovementAction InAction, const uint16 InActionID, const float InWorldStartTime, const float InServerStartTime);
+
+	void MarkNetDirty() { bNetDirty = true; }
+	bool GetNetDirty() const { return bNetDirty; }
+	void ResetNetDirty() { bNetDirty = false; }
+
+public:
+	FString ToString() const;
+
+	EMassMovementAction GetPreviousAction() const { return PreviousAction; }
+	EMassMovementAction GetCurrentAction() const { return CurrentAction; }
+	float GetCurrentActionStartTime() const { return CurrentActionWorldStartTime; }
+	float GetCurrentActionServerStartTime() const { return CurrentActionServerStartTime; }
+	uint16 GetCurrentActionID() const { return CurrentActionID; }
+
+	/** Center of the move target. */
+	FVector Center = FVector::ZeroVector;
+
+	/** Forward direction of the movement target.  */
+	FVector Forward = FVector::ZeroVector;
+
+	/** Distance remaining to the movement goal. */
+	float DistanceToGoal = 0.0f;
+
+	/** Allowed deviation around the movement target. */
+	float SlackRadius = 0.0f;
+
+private:
+	/** World time in seconds when the action started. */
+	float CurrentActionWorldStartTime = 0.0f;
+
+	/** Server time in seconds when the action started. */
+	float CurrentActionServerStartTime = 0.0f;
+
+	/** Number incremented each time new action (i.e move, stand, animation) is started. */
+	uint16 CurrentActionID = 0;
+
+public:
+	/** Requested movement speed. */
+	FMassInt16Real DesiredSpeed = FMassInt16Real(0.0f);
+
+	/** Intended movement action at the target. */
+	EMassMovementAction IntentAtGoal = EMassMovementAction::Move;
+
+private:
+	/** Current movement action. */
+	EMassMovementAction CurrentAction = EMassMovementAction::Move;
+
+	/** Previous movement action. */
+	EMassMovementAction PreviousAction = EMassMovementAction::Move;
+
+	uint8 bNetDirty : 1;
+public:
+	/** True if the movement target is assumed to be outside navigation boundaries. */
+	uint8 bOffBoundaries : 1;
+
+	/** True if the movement target is assumed to be outside navigation boundaries. */
+	uint8 bSteeringFallingBehind : 1;
+};
+
+/** Steering fragment. */
+USTRUCT()
+struct MASSAIMOVEMENT_API FMassSteeringFragment : public FLWComponentData
+{
+	GENERATED_BODY()
+
+	void Reset()
+	{
+		DesiredVelocity = FVector::ZeroVector;
+		SteeringForce = FVector::ZeroVector;
+	}
+
+	/** Cached desired velocity from steering. Note: not used for moving the entity. */
+	FVector DesiredVelocity = FVector::ZeroVector;
+
+	/** Combined steering force from all steering, used to move the entity. */
+	FVector SteeringForce = FVector::ZeroVector;
+};
+
+/** Steering ghost fragment. */
+USTRUCT()
+struct MASSAIMOVEMENT_API FMassSteeringGhostFragment : public FLWComponentData
+{
+	GENERATED_BODY()
+
+	bool IsValid(const uint16 CurrentActionID) const
+	{
+		return LastSeenActionID == CurrentActionID;
+	}
+
+	/** The action ID the ghost was initialized for */
+	uint16 LastSeenActionID = 0;
+
+	/** Location of the ghost */
+	FVector Location = FVector::ZeroVector;
+	
+	/** Velocity of the ghost */
+	FVector Velocity = FVector::ZeroVector;
+
+	/** Selected steer target based on ghost, updates periodically. */
+	FVector SteerTarget = FVector::ZeroVector;
+
+	/** Used during target update to see when the target movement stops */
+	float TargetMaxSpeed = 0.0f;
+
+	/** Cooldown between target updates */
+	float TargetCooldown = 0.0f;
+
+	/** True if the target is being updated */
+	bool bUpdatingTarget = false;
+
+	/** True if we just entered from move action */
+	bool bEnteredFromMoveAction = false;
+};
+
+USTRUCT()
+struct MASSAIMOVEMENT_API FMassDynamicObstacleFragment : public FLWComponentData
+{
+	GENERATED_BODY()
+
+	/** Lanes currently blocked by the dynamic obstacle. */
+	static constexpr int MaxLaneObstacles = 8;
+	TArray<FMassLaneObstacleID, TFixedAllocator<MaxLaneObstacles>> LaneObstacleIDs;
+
+	/** Time stamp when that obstacle stopped moving. */
+	FDateTime LastMovedTimeStamp = 0.f;
+
+	/** Position of the dynamic obstacle when it last moved. */
+	FVector LastPosition = FVector::ZeroVector;
+
+	/** Has this dynamic obstacle stopped moving. */
+	bool bHasStopped = true;
+};
+
+/** Stores handle to rich movement config */
+USTRUCT()
+struct MASSAIMOVEMENT_API FMassMovementConfigFragment : public FLWComponentData
+{
+	GENERATED_BODY()
+
+	/** Index to FMassMovementConfig in UMassMovementSettings. */
+	FMassMovementConfigHandle ConfigHandle;
+};
diff --git a/Engine/Plugins/AI/MassAI/Source/MassAIMovement/Public/MassAIMovementTypes.h b/Engine/Plugins/AI/MassAI/Source/MassAIMovement/Public/MassAIMovementTypes.h
new file mode 100644
--- /dev/null
+++ b/Engine/Plugins/AI/MassAI/Source/MassAIMovement/Public/MassAIMovementTypes.h
@@ -0,0 +1,228 @@
+// Copyright Epic Games, Inc. All Rights Reserved.
+
+#pragma once
+
+#include "LWComponentTypes.h"
+
+#include "ZoneGraphTypes.h"
+#include "MassMovementTypes.h"
+#include "MassAIMovementTypes.generated.h"
+
+MASSAIMOVEMENT_API DECLARE_LOG_CATEGORY_EXTERN(LogMassNavigation, Warning, All);
+MASSAIMOVEMENT_API DECLARE_LOG_CATEGORY_EXTERN(LogMassDynamicObstacle, Warning, All);
+
+//@ todo remove optimization hack once we find a better way to filter out signals from LOD/listener on signals, as for now we only need this signal for look at in high and med LOD
+#define HACK_DISABLE_PATH_CHANGED_ON_LOWER_LOD 1
+
+namespace UE::Mass::Signals
+{
+	const FName FollowPointPathStart = FName(TEXT("FollowPointPathStart"));
+	const FName FollowPointPathDone = FName(TEXT("FollowPointPathDone"));
+	const FName CurrentLaneChanged = FName(TEXT("CurrentLaneChanged"));
+}
+
+UENUM()
+enum class EMassMovementAction : uint8
+{
+	Stand,		// Stop and stand.
+	Move,		// Move or keep on moving.
+	Animate,	// Animation has control over the transform
+};
+
+/** Actions for lane dynamic obstacles. */
+UENUM()
+enum class EMassLaneObstacleEventAction : uint8
+{
+	Add,	// Add lane obstacle
+	Remove	// Remove lane obstacle
+};
+
+/** Reference to movement style in MassMovementSettings. */
+USTRUCT()
+struct MASSAIMOVEMENT_API FMassMovementStyleRef
+{
+	GENERATED_BODY()
+
+	UPROPERTY(EditAnywhere, Category = "Movement")
+	FName Name;
+
+	UPROPERTY(EditAnywhere, Category = "Movement", meta = (IgnoreForMemberInitializationTest, EditCondition = "false", EditConditionHides))
+	FGuid ID;
+};
+
+/** Reference to movement config in MassMovementSettings. */
+USTRUCT()
+struct MASSAIMOVEMENT_API FMassMovementConfigRef
+{
+	GENERATED_BODY()
+
+	UPROPERTY(EditAnywhere, Category = "Movement")
+	FName Name;
+
+	UPROPERTY(EditAnywhere, Category = "Movement", meta = (IgnoreForMemberInitializationTest, EditCondition = "false", EditConditionHides))
+	FGuid ID;
+};
+
+
+/** Describes movement style name. */
+USTRUCT()
+struct MASSAIMOVEMENT_API FMassMovementStyle
+{
+	GENERATED_BODY()
+
+	UPROPERTY(EditAnywhere, Category = "Movement")
+	FName Name;
+
+	UPROPERTY(EditAnywhere, Category = "Movement", meta = (IgnoreForMemberInitializationTest, EditCondition = "false", EditConditionHides))
+	FGuid ID;
+};
+
+/**
+ * Movement style consists of multiple speeds which are assigned to agents based on agents unique ID.
+ * Same speed is assigned consistently for the same ID.
+ */
+USTRUCT()
+struct MASSAIMOVEMENT_API FMassMovementStyleSpeed
+{
+	GENERATED_BODY()
+
+	/** Desired speed */
+	UPROPERTY(EditAnywhere, Category = "Movement Style", meta = (UIMin = 0.0, ClampMin = 0.0))
+	float Speed = 140.0f;
+
+	/** How much default desired speed is varied randomly. */
+	UPROPERTY(EditAnywhere, Category = "Movement Style", meta = (ClampMin = 0.0, ClampMax = 1.0))
+	float Variance = 0.1f;
+
+	/** Probability to assign this speed. */
+	UPROPERTY(EditAnywhere, Category = "Movement Style", meta = (UIMin = 1, ClampMin = 1, UIMax = 100, ClampMax = 100))
+	int32 Probability = 1;
+
+	/** Running sum of the probabilities so far, used to faster lookup. Update via FMassMovementConfig::Update(). */
+	int32 ProbabilityThreshold = 0;
+};
+
+/** Movement style config. A movement style abstracts movement properties for specific style. Behaviors can refer to specific styles when handling speeds. */
+USTRUCT()
+struct MASSAIMOVEMENT_API FMassMovementStyleConfig
+{
+	GENERATED_BODY()
+
+	/** Style of the movement */
+	UPROPERTY(EditAnywhere, Category = "Movement Style")
+	FMassMovementStyleRef Style;
+
+	/** Array of desired speeds (cm/s) assigned to agents based on probability. */
+	UPROPERTY(EditAnywhere, Category = "Movement Style")
+	TArray<FMassMovementStyleSpeed> DesiredSpeeds;
+
+	/** Sum of all probabilities in style speeds array. Update via FMassMovementConfig::Update() */
+	int32 MaxProbability = 0;
+};
+
+/** Steering related movement configs. */
+USTRUCT()
+struct MASSAIMOVEMENT_API FMassMovementSteeringConfig
+{
+	GENERATED_BODY()
+
+	/** Steering reaction time in seconds. */
+	UPROPERTY(config, EditAnywhere, Category = "Steering", meta = (ClampMin = "0.05"))
+	float ReactionTime = 0.3f;
+
+	/** 200..600 Smaller steering maximum acceleration makes the steering more \"calm\" but less opportunistic, may not find solution, or gets stuck. */
+	UPROPERTY(config, EditAnywhere, Category = "Steering", meta = (UIMin = 0.0, ClampMin = 0.0))
+	float MaxAcceleration = 250.f;
+
+	/** How much ahead to predict when steering. */
+	UPROPERTY(EditAnywhere, Category="Steering", meta = (UIMin = 0.0, ClampMin = 0.0))
+	float LookAheadDistance = 100.0f;
+};
+	
+USTRUCT()
+struct MASSAIMOVEMENT_API FMassMovementConfig
+{
+	GENERATED_BODY()
+
+	/** Updates internal values for faster desired speed generation. */
+	void Update();
+
+	/** Generates desired speed based on style and unique id. The id is used deterministically assign a specific speed range. */
+	float GenerateDesiredSpeed(const FMassMovementStyleRef& Style, const int32 UniqueId) const
+	{
+		float DesiredSpeed = DefaultDesiredSpeed;
+		float DesiredSpeedVariance = DefaultDesiredSpeedVariance;
+		
+		const FMassMovementStyleConfig* StyleConfig = MovementStyles.FindByPredicate([&Style](const FMassMovementStyleConfig& Config)
+			{
+				return Config.Style.ID == Style.ID;
+			});
+		
+		if (StyleConfig != nullptr && StyleConfig->MaxProbability > 0)
+		{
+			const int32 Prob = UniqueId % StyleConfig->MaxProbability;
+			for (const FMassMovementStyleSpeed& Speed : StyleConfig->DesiredSpeeds)
+			{
+				if (Prob < Speed.ProbabilityThreshold)
+				{
+					DesiredSpeed = Speed.Speed;
+					DesiredSpeedVariance = Speed.Variance;
+					break;
+				}
+			}
+		}
+		
+		return DesiredSpeed * FMath::RandRange(1.0f - DesiredSpeedVariance, 1.0f + DesiredSpeedVariance);
+	}
+
+	/** Name of the movement configuration. */
+	UPROPERTY(EditAnywhere, Category = "Movement")
+	FName Name;
+
+	/** GUID representing this configuration. */
+	UPROPERTY(EditAnywhere, Category = "Movement", meta = (IgnoreForMemberInitializationTest, EditCondition = "false", EditConditionHides))
+	FGuid ID;
+
+	/** Maximum speed (cm/s). */
+	UPROPERTY(EditAnywhere, Category = "Movement", meta = (ClampMin = "0"))
+	float MaximumSpeed = 200.f;
+
+	/** Default desired speed (cm/s). */
+	UPROPERTY(EditAnywhere, Category = "Movement", meta = (ClampMin = "0"))
+	float DefaultDesiredSpeed = 140.f;
+
+	/** How much default desired speed is varied randomly. */
+	UPROPERTY(EditAnywhere, Category = "Movement", meta = (ClampMin = "0", ClampMax = "1"))
+	float DefaultDesiredSpeedVariance = 0.1f;
+
+	/** Steering config. */
+	UPROPERTY(EditAnywhere, Category="Movement")
+	FMassMovementSteeringConfig Steering;
+	
+	/** Filter describing which lanes can be used when spawned. */
+	UPROPERTY(EditAnywhere, Category="Movement")
+	FZoneGraphTagFilter LaneFilter;
+
+	/** Query radius when trying to find nearest lane when spawned. */
+	UPROPERTY(EditAnywhere, Category="Movement", meta = (UIMin = 0.0, ClampMin = 0.0))
+	float QueryRadius = 500.0f;
+
+	/** List of supported movement styles for this configuration. */
+	UPROPERTY(EditAnywhere, Category = "Movement")
+	TArray<FMassMovementStyleConfig> MovementStyles;
+};
+
+/** Handle to FMassMovementConfig in MassMovementSettings. */
+struct MASSAIMOVEMENT_API FMassMovementConfigHandle
+{
+	FMassMovementConfigHandle() = default;
+	explicit FMassMovementConfigHandle(const uint8 InIndex) : Index(InIndex) {}
+
+	static constexpr uint8 InvalidValue = 0xff; 
+	
+	bool IsValid() const { return Index != InvalidValue; }
+	bool operator==(const FMassMovementConfigHandle& Other) const { return Index == Other.Index; }
+	bool operator!=(const FMassMovementConfigHandle& Other) const { return Index != Other.Index; }
+	
+	uint8 Index = InvalidValue;
+};
diff --git a/Engine/Plugins/AI/MassAI/Source/MassAIMovement/Public/MassAvoidanceObstacleTrait.h b/Engine/Plugins/AI/MassAI/Source/MassAIMovement/Public/MassAvoidanceObstacleTrait.h
new file mode 100644
--- /dev/null
+++ b/Engine/Plugins/AI/MassAI/Source/MassAIMovement/Public/MassAvoidanceObstacleTrait.h
@@ -0,0 +1,18 @@
+// Copyright Epic Games, Inc. All Rights Reserved.
+
+#pragma once 
+
+#include "MassEntityTraitBase.h"
+#include "MassAvoidanceObstacleTrait.generated.h"
+
+UCLASS(BlueprintType, EditInlineNew, CollapseCategories, meta = (DisplayName = "Dynamic Obstacle"))
+class MASSAIMOVEMENT_API UMassAvoidanceObstacleTrait : public UMassEntityTraitBase
+{
+	GENERATED_BODY()
+
+protected:
+	virtual void BuildTemplate(FMassEntityTemplateBuildContext& BuildContext, UWorld& World) const override;
+
+	UPROPERTY(EditAnywhere, Category = Mass)
+	bool bExtendToEdgeObstacle = false;
+};
\ No newline at end of file
diff --git a/Engine/Plugins/AI/MassAI/Source/MassAIMovement/Public/MassAvoidanceProcessors.h b/Engine/Plugins/AI/MassAI/Source/MassAIMovement/Public/MassAvoidanceProcessors.h
new file mode 100644
--- /dev/null
+++ b/Engine/Plugins/AI/MassAI/Source/MassAIMovement/Public/MassAvoidanceProcessors.h
@@ -0,0 +1,212 @@
+// Copyright Epic Games, Inc. All Rights Reserved.
+
+#pragma once
+
+#include "MassProcessor.h"
+#include "MassCommonFragments.h"
+#include "MassSimulationLOD.h"
+#include "MovementProcessor.h"
+#include "NavMesh/RecastNavMesh.h"
+#include "ZoneGraphTypes.h"
+#include "MassAvoidanceProcessors.generated.h"
+
+class UZoneGraphSubsystem;
+
+MASSAIMOVEMENT_API DECLARE_LOG_CATEGORY_EXTERN(LogAvoidance, Warning, All);
+MASSAIMOVEMENT_API DECLARE_LOG_CATEGORY_EXTERN(LogAvoidanceVelocities, Warning, All);
+MASSAIMOVEMENT_API DECLARE_LOG_CATEGORY_EXTERN(LogAvoidanceAgents, Warning, All);
+MASSAIMOVEMENT_API DECLARE_LOG_CATEGORY_EXTERN(LogAvoidanceObstacles, Warning, All);
+
+USTRUCT()
+struct MASSAIMOVEMENT_API FMassEdgeDetectionParamsFragment : public FLWComponentData
+{
+	GENERATED_BODY()
+	float EdgeDetectionRange = 500.f;
+};
+
+/** Edge with normal */
+struct FNavigationAvoidanceEdge
+{
+	FNavigationAvoidanceEdge(const FVector InStart, const FVector InEnd)
+	{
+		Start = InStart;
+		End = InEnd;
+		LeftDir = FVector::CrossProduct((End - Start).GetSafeNormal(), FVector::UpVector);
+	}
+	FVector Start = FVector::ZeroVector;
+	FVector End = FVector::ZeroVector;
+	FVector LeftDir = FVector::ZeroVector;
+};
+
+USTRUCT()
+struct MASSAIMOVEMENT_API FMassNavigationEdgesFragment : public FLWComponentData
+{
+	GENERATED_BODY()
+
+	static const int MaxEdgesCount = 8;
+	TArray<FNavigationAvoidanceEdge, TFixedAllocator<MaxEdgesCount>> AvoidanceEdges;
+};
+
+/** Component tag to tell the avoidance to extend the size of this obstacle when too close to edges. */
+USTRUCT()
+struct MASSAIMOVEMENT_API FMassAvoidanceExtendToEdgeObstacleTag : public FComponentTag
+{
+	GENERATED_BODY()
+};
+
+/** Experimental: move using cumulative forces to avoid close agents */
+UCLASS()
+class MASSAIMOVEMENT_API UMassAvoidanceProcessor : public UPipeProcessor
+{
+	GENERATED_BODY()
+
+public:
+	UMassAvoidanceProcessor();
+
+protected:
+	virtual void ConfigureQueries() override;
+	virtual void Initialize(UObject& Owner) override;
+	virtual void Execute(UEntitySubsystem& EntitySubsystem, FLWComponentSystemExecutionContext& Context) override;
+
+private:
+	TWeakObjectPtr<UWorld> WeakWorld;
+	TWeakObjectPtr<UMassMovementSubsystem> WeakMovementSubsystem;
+	FLWComponentQuery EntityQuery;
+};
+
+/** Avoidance while standing. */
+UCLASS()
+class MASSAIMOVEMENT_API UMassStandingAvoidanceProcessor : public UPipeProcessor
+{
+	GENERATED_BODY()
+
+public:
+	UMassStandingAvoidanceProcessor();
+
+protected:
+	virtual void ConfigureQueries() override;
+	virtual void Initialize(UObject& Owner) override;
+	virtual void Execute(UEntitySubsystem& EntitySubsystem, FLWComponentSystemExecutionContext& Context) override;
+
+private:
+	TWeakObjectPtr<UWorld> WeakWorld;
+	TWeakObjectPtr<UMassMovementSubsystem> WeakMovementSubsystem;
+	FLWComponentQuery EntityQuery;
+};
+
+/** Component Tag to tell if the entity is in the avoidance obstacle grid */
+USTRUCT()
+struct MASSAIMOVEMENT_API FMassInAvoidanceObstacleGridTag : public FComponentTag
+{
+	GENERATED_BODY()
+};
+
+/** Processor to update avoidance obstacle data */
+UCLASS()
+class MASSAIMOVEMENT_API UMassAvoidanceObstacleProcessor : public UPipeProcessor
+{
+	GENERATED_BODY()
+
+public:
+	UMassAvoidanceObstacleProcessor();
+
+protected:
+	virtual void ConfigureQueries() override;
+	virtual void Initialize(UObject& Owner) override;
+	virtual void Execute(UEntitySubsystem& EntitySubsystem, FLWComponentSystemExecutionContext& Context) override;
+
+private:
+	TWeakObjectPtr<UMassMovementSubsystem> WeakMovementSubsystem;
+	FLWComponentQuery AddToGridEntityQuery;
+	FLWComponentQuery UpdateGridEntityQuery;
+	FLWComponentQuery RemoveFromGridEntityQuery;
+};
+
+/** Experimental: navmesh edges gathering */
+UCLASS()
+class MASSAIMOVEMENT_API UMassNavigationBoundaryProcessor : public UPipeProcessor
+{
+	GENERATED_BODY()
+
+public:
+	UMassNavigationBoundaryProcessor();
+
+protected:
+	virtual void ConfigureQueries() override;
+	virtual void Initialize(UObject& Owner) override;
+	virtual void Execute(UEntitySubsystem& EntitySubsystem, FLWComponentSystemExecutionContext& Context) override;
+
+private:
+	TWeakObjectPtr<UWorld> WeakWorld;
+	TWeakObjectPtr<UMassMovementSubsystem> WeakMovementSubsystem;
+	TWeakObjectPtr<ANavigationData> WeakNavData;
+	FLWComponentQuery EntityQuery;
+};
+
+USTRUCT()
+struct MASSAIMOVEMENT_API FMassLastUpdatePositionFragment : public FLWComponentData
+{
+	GENERATED_BODY()
+	FVector Value;
+};
+
+USTRUCT()
+struct MASSAIMOVEMENT_API FMassAvoidanceBoundaryLastLaneHandleFragment : public FLWComponentData
+{
+	GENERATED_BODY()
+
+	/** Not replicated but computed and used on client and server */
+	FZoneGraphLaneHandle LaneHandle;
+};
+
+/** Experimental: lane borders gathering */
+UCLASS()
+class MASSAIMOVEMENT_API UMassLaneBoundaryProcessor : public UPipeProcessor
+{
+	GENERATED_BODY()
+
+public:
+	UMassLaneBoundaryProcessor();
+
+protected:
+	virtual void ConfigureQueries() override;
+	virtual void Initialize(UObject& Owner) override;
+	virtual void Execute(UEntitySubsystem& EntitySubsystem, FLWComponentSystemExecutionContext& Context) override;
+
+private:
+	TWeakObjectPtr<UWorld> WeakWorld;
+	TWeakObjectPtr<UZoneGraphSubsystem> WeakZoneGraph;
+	FLWComponentQuery EntityQuery;
+};
+
+USTRUCT()
+struct MASSAIMOVEMENT_API FMassLaneCacheBoundaryFragment : public FLWComponentData
+{
+	GENERATED_BODY()
+
+	/** Last update position. */
+	FVector LastUpdatePosition = FVector::ZeroVector;
+
+	/** Lane cached ID at last update. */
+	uint16 LastUpdateCacheID = 0;
+};
+
+/** ZoneGraph lane cache boundary processor */
+// @todo MassMovement: Make this signal based.
+UCLASS()
+class MASSAIMOVEMENT_API UMassLaneCacheBoundaryProcessor : public UPipeProcessor
+{
+	GENERATED_BODY()
+
+public:
+	UMassLaneCacheBoundaryProcessor();
+
+protected:
+	virtual void ConfigureQueries() override;
+	virtual void Initialize(UObject& Owner) override;
+	virtual void Execute(UEntitySubsystem& EntitySubsystem, FLWComponentSystemExecutionContext& Context) override;
+
+private:
+	TWeakObjectPtr<UWorld> WeakWorld;
+	FLWComponentQuery EntityQuery;
+};
\ No newline at end of file
diff --git a/Engine/Plugins/AI/MassAI/Source/MassAIMovement/Public/MassAvoidanceSettings.h b/Engine/Plugins/AI/MassAI/Source/MassAIMovement/Public/MassAvoidanceSettings.h
new file mode 100644
--- /dev/null
+++ b/Engine/Plugins/AI/MassAI/Source/MassAIMovement/Public/MassAvoidanceSettings.h
@@ -0,0 +1,116 @@
+// Copyright Epic Games, Inc. All Rights Reserved.
+
+#pragma once
+
+#include "CoreMinimal.h"
+#include "Engine/DeveloperSettings.h"
+#include "MassAvoidanceSettings.generated.h"
+
+USTRUCT()
+struct MASSAIMOVEMENT_API FMassAvoidanceVelocityFilter
+{
+	GENERATED_BODY()
+
+	/** Delay between changes. (seconds) */
+	UPROPERTY(Category = Zone, EditAnywhere, meta = (ClampMin = "0"))
+	double Delay = 0.4;
+
+	/** Minimum speed of the range (cm/s). */
+	UPROPERTY(Category = Zone, EditAnywhere, meta = (ClampMin = "0"))
+	float LowSpeed = 10.f;
+
+	/** Maximum speed of the range (cm/s). */
+	UPROPERTY(Category = Zone, EditAnywhere, meta = (ClampMin = "0"))
+	float HighSpeed = 50.f;
+};
+
+UCLASS(config = Engine, defaultconfig, meta = (DisplayName = "Mass Avoidance"))
+class MASSAIMOVEMENT_API UMassAvoidanceSettings : public UDeveloperSettings
+{
+	GENERATED_BODY()
+
+public:
+	static const UMassAvoidanceSettings* Get()
+	{
+		return GetDefault<UMassAvoidanceSettings>();
+	}
+
+	/** 1..3 Too high value makes close agent search slow, too small removes avoidance effect. Indoor humans 1.4, outdoor humans 2.4 (seconds). */
+	UPROPERTY(config, EditAnywhere, Category = "General", meta = (ClampMin = "0.1"))
+	float TimeHorizon = 2.5f;
+
+
+	/** 0..10 How far inside the circle the smooth collisions starts (distance). */
+	UPROPERTY(config, EditAnywhere, Category = "Collision", meta = (ClampMin = "0"))
+	float AgentCollisionInset = 5.f;
+
+	/** 0..10 How far inside the circle the smooth collisions starts (distance). */
+	UPROPERTY(config, EditAnywhere, Category = "Collision", meta = (ClampMin = "0"))
+	float ObstacleCollisionInset = 5.f;
+
+
+	/** 100..500 Separation force, even a little bit of separation can help to smooth out deadlocks in dense crowds. */
+	UPROPERTY(config, EditAnywhere, Category = "Separation", meta = (ClampMin = "0"))
+	float AgentSeparation = 200.f;
+
+	/** 100..1000 Separation force for obstacles growing near edges. */
+	UPROPERTY(config, EditAnywhere, Category = "Separation", meta = (ClampMin = "0"))
+	float AgentSeparationForExtendingColliders = 800.f;
+	
+	/** 0..100 How big the decay/buffer is for separation (distance). */
+	UPROPERTY(config, EditAnywhere, Category = "Separation", meta = (ClampMin = "0"))
+	float AgentSeparationBuffer = 75.f;
+
+	/** 0..100 Agent separation buffer near target location (distance). */
+	UPROPERTY(config, EditAnywhere, Category = "Separation", meta = (ClampMin = "0"))
+	float AgentSeparationBufferAtEnd = 15.f;
+	
+	/** 0..150 Separation buffer for obstacles growing near edges (distance). */
+	UPROPERTY(config, EditAnywhere, Category = "Separation", meta = (ClampMin = "0"))
+	float AgentSeparationBufferForExtendingColliders = 150.f;
+	
+	/** Distance threshold where the agent is considered near it's target location and removing agent separation. */
+	UPROPERTY(config, EditAnywhere, Category = "Separation", meta = (ClampMin = "0"))
+	float NearTargetLocationDistance = 150.f;
+
+	/** 200..1000 Prevents getting stuck obstacles, keeps space to take over close to obstacles (distance). */
+	UPROPERTY(config, EditAnywhere, Category = "Separation", meta = (ClampMin = "0"))
+	float ObstacleSeparation = 600.f;
+
+	/** 0..200 How big the decay/buffer is for separation (distance). */
+	UPROPERTY(config, EditAnywhere, Category = "Separation", meta = (ClampMin = "0"))
+	float ObstacleSeparationBuffer = 200.f;
+
+
+	/** 0..25 Allowed penetration, subtracted from the total of compared agents radius (distance). */
+	UPROPERTY(config, EditAnywhere, Category = "Predictive Avoidance", meta = (ClampMin = "0"))
+	float AvoidanceInset = 20.f;
+	
+	/** 0..150 Buffer added when computing distance to agents and obstacles (distance). */
+	UPROPERTY(config, EditAnywhere, Category = "Predictive Avoidance", meta = (ClampMin = "0"))
+	float AvoidanceBuffer = 100.f;
+
+	/** 0..150 Buffer added when computing distance to agents and obstacles growing near edges (distance). */
+	UPROPERTY(config, EditAnywhere, Category = "Predictive Avoidance", meta = (ClampMin = "0"))
+	float AvoidanceBufferForExtendingColliders = 100.f;
+	
+	/** 0..150 Avoidance buffer near target location (distance). */
+	UPROPERTY(config, EditAnywhere, Category = "Predictive Avoidance", meta = (ClampMin = "0"))
+	float AvoidanceBufferAtEnd = 10.f;
+
+	/** 400..1000 Multiplier on the agent-agent avoidance force (distance). */
+	UPROPERTY(config, EditAnywhere, Category = "Predictive Avoidance", meta = (ClampMin = "0"))
+	float AgentAvoidanceStiffness = 700.f;
+
+	/** 400..1000 Multiplier on the agent-obstacle avoidance force (distance). */
+	UPROPERTY(config, EditAnywhere, Category = "Predictive Avoidance", meta = (ClampMin = "0"))
+	float ObstacleAvoidanceStiffness = 400.f;
+
+	/** Time (in seconds) it takes a new heading direction to completely blend in. */
+	UPROPERTY(config, EditAnywhere, Category = "Orientation", meta = (ClampMin = "0"))
+	float OrientationSmoothingTime = 0.7f;
+	
+	/** Distance from end of the path when we start to blend into the desired orientation. */
+	UPROPERTY(config, EditAnywhere, Category = "Orientation", meta = (ClampMin = "0"))
+	float OrientationEndOfPathHeadingAnticipation = 100.0f;
+};
diff --git a/Engine/Plugins/AI/MassAI/Source/MassAIMovement/Public/MassAvoidanceTrait.h b/Engine/Plugins/AI/MassAI/Source/MassAIMovement/Public/MassAvoidanceTrait.h
new file mode 100644
--- /dev/null
+++ b/Engine/Plugins/AI/MassAI/Source/MassAIMovement/Public/MassAvoidanceTrait.h
@@ -0,0 +1,18 @@
+// Copyright Epic Games, Inc. All Rights Reserved.
+
+#pragma once 
+
+#include "MassEntityTraitBase.h"
+#include "MassAvoidanceTrait.generated.h"
+
+UCLASS(BlueprintType, EditInlineNew, CollapseCategories, meta = (DisplayName = "Avoidance"))
+class MASSAIMOVEMENT_API UMassAvoidanceTrait : public UMassEntityTraitBase
+{
+	GENERATED_BODY()
+
+protected:
+	virtual void BuildTemplate(FMassEntityTemplateBuildContext& BuildContext, UWorld& World) const override;
+
+	UPROPERTY(Category="Movement", EditAnywhere)
+	bool bUseZoneGraphMovement = false;
+};
diff --git a/Engine/Plugins/AI/MassAI/Source/MassAIMovement/Public/MassMovementDelegates.h b/Engine/Plugins/AI/MassAI/Source/MassAIMovement/Public/MassMovementDelegates.h
new file mode 100644
--- /dev/null
+++ b/Engine/Plugins/AI/MassAI/Source/MassAIMovement/Public/MassMovementDelegates.h
@@ -0,0 +1,14 @@
+﻿// Copyright Epic Games, Inc. All Rights Reserved.
+
+#pragma once
+
+#include "CoreMinimal.h"
+
+namespace UE::MassMovement::Delegates
+{
+#if WITH_EDITOR
+	/** Called when movement names have changed (UI update). */
+	DECLARE_MULTICAST_DELEGATE(FOnMassMovementNamesChanged);
+	extern MASSAIMOVEMENT_API FOnMassMovementNamesChanged OnMassMovementNamesChanged;
+#endif // WITH_EDITOR
+} // UE::MassMovement::Delegates
diff --git a/Engine/Plugins/AI/MassAI/Source/MassAIMovement/Public/MassMovementSettings.h b/Engine/Plugins/AI/MassAI/Source/MassAIMovement/Public/MassMovementSettings.h
new file mode 100644
--- /dev/null
+++ b/Engine/Plugins/AI/MassAI/Source/MassAIMovement/Public/MassMovementSettings.h
@@ -0,0 +1,42 @@
+// Copyright Epic Games, Inc. All Rights Reserved.
+
+#pragma once
+
+#include "CoreMinimal.h"
+#include "Engine/DeveloperSettings.h"
+#include "MassAIMovementTypes.h"
+#include "MassMovementSettings.generated.h"
+
+UCLASS(config = Engine, defaultconfig, meta = (DisplayName = "Mass Movement"))
+class MASSAIMOVEMENT_API UMassMovementSettings : public UDeveloperSettings
+{
+	GENERATED_BODY()
+
+public:
+	static const UMassMovementSettings* Get()
+	{
+		return GetDefault<UMassMovementSettings>();
+	}
+
+	TConstArrayView<FMassMovementStyle> GetMovementStyles() const { return MovementStyles; }
+	const FMassMovementStyle* GetMovementStyleByID(const FGuid ID) const; 
+	
+	TConstArrayView<FMassMovementConfig> GetMovementConfigs() const { return MovementConfigs; } 
+	const FMassMovementConfig* GetMovementConfigByID(const FGuid ID) const;
+	FMassMovementConfigHandle GetMovementConfigHandleByID(const FGuid ID) const;
+	const FMassMovementConfig* GetMovementConfigByHandle(const FMassMovementConfigHandle Handle) const { return Handle.IsValid() ? &MovementConfigs[Handle.Index] : nullptr; } 
+
+private:
+
+	void UpdateConfigs();
+	virtual void PostInitProperties() override;
+#if WITH_EDITOR
+	virtual void PostEditChangeChainProperty(FPropertyChangedChainEvent& PropertyChangedEvent) override;
+#endif
+
+	UPROPERTY(EditAnywhere, config, Category = Movement);
+    TArray<FMassMovementStyle> MovementStyles;
+
+	UPROPERTY(EditAnywhere, config, Category = Movement);
+	TArray<FMassMovementConfig> MovementConfigs;
+};
diff --git a/Engine/Plugins/AI/MassAI/Source/MassAIMovement/Public/MassMovementSubsystem.h b/Engine/Plugins/AI/MassAI/Source/MassAIMovement/Public/MassMovementSubsystem.h
new file mode 100644
--- /dev/null
+++ b/Engine/Plugins/AI/MassAI/Source/MassAIMovement/Public/MassMovementSubsystem.h
@@ -0,0 +1,35 @@
+// Copyright Epic Games, Inc. All Rights Reserved.
+
+#pragma once
+
+#include "MassEntitySystem.h"
+#include "HierarchicalHashGrid2D.h"
+#include "MassMovementTypes.h"
+#include "Subsystems/WorldSubsystem.h"
+#include "MassMovementSubsystem.generated.h"
+
+class UZoneGraphSubsystem;
+
+UCLASS()
+class MASSAIMOVEMENT_API UMassMovementSubsystem : public UWorldSubsystem
+{
+	GENERATED_BODY()
+
+public:
+	UMassMovementSubsystem();
+
+	virtual void Initialize(FSubsystemCollectionBase& Collection) override;
+
+	const FAvoidanceObstacleHashGrid2D& GetGrid() const { return AvoidanceObstacleGrid; }
+	FAvoidanceObstacleHashGrid2D& GetGridMutable() { return AvoidanceObstacleGrid; }
+
+protected:
+
+	UPROPERTY(Transient)
+	UPipeEntitySubsystem* EntitySubsystem;
+
+	UPROPERTY(Transient)
+	UZoneGraphSubsystem* ZoneGraphSubsystem;
+
+	FAvoidanceObstacleHashGrid2D AvoidanceObstacleGrid;
+};
diff --git a/Engine/Plugins/AI/MassAI/Source/MassAIMovement/Public/MassZoneGraphLaneObstacleAnnotations.h b/Engine/Plugins/AI/MassAI/Source/MassAIMovement/Public/MassZoneGraphLaneObstacleAnnotations.h
new file mode 100644
--- /dev/null
+++ b/Engine/Plugins/AI/MassAI/Source/MassAIMovement/Public/MassZoneGraphLaneObstacleAnnotations.h
@@ -0,0 +1,160 @@
+// Copyright Epic Games, Inc. All Rights Reserved.
+
+#pragma once
+
+#include "MassAIMovementTypes.h"
+#include "ZoneGraphAnnotationTypes.h"
+#include "ZoneGraphAnnotationComponent.h"
+
+#include "MassZoneGraphLaneObstacleAnnotations.generated.h"
+
+class UZoneGraphSubsystem;
+
+/**
+* Uniquely identified obstacles on lanes.
+*/
+struct MASSAIMOVEMENT_API FMassLaneObstacle
+{
+	FMassLaneObstacle() = default;
+	FMassLaneObstacle(const FMassLaneObstacleID NewID, const FZoneGraphLaneSection& Section)
+		: LaneSection(Section)
+		, ID(NewID)
+	{}
+
+	/** Returns the unique ID for that obstacle. */
+	FMassLaneObstacleID GetID() const { return ID; }
+
+	bool operator==(const FMassLaneObstacle& Other) const
+	{
+		return LaneSection == Other.LaneSection && ID == Other.ID;
+	}
+	
+	/** Part of the lane that is blocked. */
+	FZoneGraphLaneSection LaneSection;
+
+protected:
+	/** Lane obstacle ID. */
+	FMassLaneObstacleID ID;
+};
+
+/**
+ * Event indicating an obstacle change on a lane.
+ */
+USTRUCT()
+struct MASSAIMOVEMENT_API FZoneGraphLaneObstacleChangeEvent : public FZoneGraphAnnotationEventBase
+{
+	GENERATED_BODY()
+
+	FZoneGraphLaneObstacleChangeEvent() = default;
+	FZoneGraphLaneObstacleChangeEvent(const FMassLaneObstacle Obstacle, const EMassLaneObstacleEventAction Action)
+		: LaneObstacle(Obstacle)
+		, EventAction(Action)
+	{}
+
+	/** Part of the lane that is blocked. */
+	FMassLaneObstacle LaneObstacle;
+
+	/** Action for that lane obstacle. */
+	EMassLaneObstacleEventAction EventAction;
+};
+
+/**
+ * Container to get obstacles from LaneHandle
+ */
+class FMassLaneObstacleContainer
+{
+public:
+	/** Add a FMassLaneObstacle */
+	void Add(const FMassLaneObstacle& Obstacle)
+	{
+		IDToLaneObstacleMap.Add(Obstacle.GetID(), Obstacle);
+		
+		TArray<FMassLaneObstacle>& Obstacles = LaneObstaclesMap.FindOrAdd(Obstacle.LaneSection.LaneHandle);
+		Obstacles.Add(Obstacle);
+	}
+
+	/** Remove a FMassLaneObstacle using it's ID, returns true if it was the last one on the lane. */
+	bool Remove(const FMassLaneObstacle& Obstacle)
+	{
+		IDToLaneObstacleMap.Remove(Obstacle.GetID());
+		
+		TArray<FMassLaneObstacle>& Obstacles = LaneObstaclesMap.FindChecked(Obstacle.LaneSection.LaneHandle);
+		Obstacles.RemoveSwap(Obstacle);
+
+		return Obstacles.IsEmpty();
+	}
+
+	/** Find a FMassLaneObstacle from it's ID. */
+	const FMassLaneObstacle* Find(const FMassLaneObstacleID& LaneObstacle)
+	{
+		return IDToLaneObstacleMap.Find(LaneObstacle);
+	}
+
+	/** Returns the array of lane obstacle from a FZoneGraphLaneHandle (this is expected to be the most common access pattern). */
+	TArray<FMassLaneObstacle>* Find(const FZoneGraphLaneHandle LaneHandle)
+	{
+		return LaneObstaclesMap.Find(LaneHandle);
+	}
+
+	/** Reset data */
+	void Reset()
+	{
+		IDToLaneObstacleMap.Reset();
+		LaneObstaclesMap.Reset();
+	}
+	
+	TMap<FMassLaneObstacleID, FMassLaneObstacle> IDToLaneObstacleMap; // @todo: investigate if keeping the LaneHandle with the ID in the fragment is better than having this second map.
+	TMap<FZoneGraphLaneHandle, TArray<FMassLaneObstacle>> LaneObstaclesMap;
+};
+
+/** Container for movement lane data associated to a specific registered ZoneGraph data. */
+struct FMassRegisteredMovementLaneData
+{
+	void Reset()
+	{
+		DataHandle.Reset();
+		LaneObstacles.Reset();
+	}
+
+	/** Handle of the storage the data was initialized from. */
+	FZoneGraphDataHandle DataHandle;
+
+	/** 'Lane obstacle' data. */
+	FMassLaneObstacleContainer LaneObstacles;
+};
+
+/**
+ * Zone graph lane obstacle annotations.
+ */
+UCLASS(ClassGroup = AI, BlueprintType, meta = (BlueprintSpawnableComponent))
+class MASSAIMOVEMENT_API UZoneGraphLaneObstacleAnnotations : public UZoneGraphAnnotationComponent
+{
+	GENERATED_BODY()
+
+public:
+	virtual void PostZoneGraphDataAdded(const AZoneGraphData& ZoneGraphData) override;
+	virtual void PreZoneGraphDataRemoved(const AZoneGraphData& ZoneGraphData) override;
+	
+protected:
+	virtual void PostSubsystemsInitialized() override;
+	virtual FZoneGraphTagMask GetAnnotationTags() const override;
+	virtual void HandleEvents(TConstArrayView<const UScriptStruct*> AllEventStructs, const FInstancedStructStream& Events) override;
+	virtual void TickAnnotation(const float DeltaTime, FZoneGraphAnnotationTagContainer& BehaviorTagContainer) override;
+
+#if !UE_BUILD_SHIPPING && !UE_BUILD_TEST
+	virtual void DebugDraw(FZoneGraphAnnotationSceneProxy* DebugProxy) override;
+#endif
+	
+	UPROPERTY(Transient)
+	UZoneGraphSubsystem* ZoneGraphSubsystem;
+
+	/** Annotation Tag to mark lanes that have obstacles. */
+	UPROPERTY(EditAnywhere, Category = Tag)
+	FZoneGraphTag LaneObstacleTag;
+
+	/** Array of queued obstacle change events. */
+	TArray<FZoneGraphLaneObstacleChangeEvent> LaneObstacleChangeEvents;
+
+	/** Lane data for all registered ZoneGraph data, keeping track of 'lane obstacles' on lanes. */
+	TArray<FMassRegisteredMovementLaneData> RegisteredLaneData;
+};
diff --git a/Engine/Plugins/AI/MassAI/Source/MassAIMovement/Public/MassZoneGraphMovementFragments.h b/Engine/Plugins/AI/MassAI/Source/MassAIMovement/Public/MassZoneGraphMovementFragments.h
new file mode 100644
--- /dev/null
+++ b/Engine/Plugins/AI/MassAI/Source/MassAIMovement/Public/MassZoneGraphMovementFragments.h
@@ -0,0 +1,237 @@
+// Copyright Epic Games, Inc. All Rights Reserved.
+
+#pragma once
+
+#include "LWComponentTypes.h"
+#include "ZoneGraphTypes.h"
+#include "MassCommonTypes.h"
+#include "MassZoneGraphMovementTypes.h"
+#include "MassZoneGraphMovementFragments.generated.h"
+
+/** Stores path request associated to a new movement action. This is used to replicate actions. */
+USTRUCT()
+struct MASSAIMOVEMENT_API FMassZoneGraphPathRequestFragment : public FLWComponentData
+{
+	GENERATED_BODY()
+
+	/** Short path request Handle to current lane. */
+	UPROPERTY(Transient)
+	FZoneGraphShortPathRequest PathRequest;
+};
+
+/** Describes current location on ZoneGraph */ 
+USTRUCT()
+struct MASSAIMOVEMENT_API FMassZoneGraphLaneLocationFragment : public FLWComponentData
+{
+	GENERATED_BODY()
+
+	/** Handle to current lane. */
+	FZoneGraphLaneHandle LaneHandle;
+	
+	/** Distance along current lane. */
+	float DistanceAlongLane = 0.0f;
+	
+	/** Cached lane length, used for clamping and testing if at end of lane. */
+	float LaneLength = 0.0f;
+};
+
+/** Describes part of a ZoneGraph lane. */
+USTRUCT()
+struct MASSAIMOVEMENT_API FMassZoneGraphCachedLaneFragment : public FLWComponentData
+{
+	GENERATED_BODY()
+
+	static constexpr uint8 MaxPoints = 5;
+
+	void Reset()
+	{
+		LaneHandle.Reset();
+		LaneLength = 0.0f;
+		LaneWidth = FMassInt16Real(0.0f);
+		NumPoints = 0;
+	}
+
+	/** Caches portion of a lane from ZoneGraph. */
+	void CacheLaneData(const FZoneGraphStorage& ZoneGraphStorage, const FZoneGraphLaneHandle CurrentLaneHandle,
+					   const float CurrentDistanceAlongLane, const float TargetDistanceAlongLane, const float InflateDistance);
+
+	int32 FindSegmentIndexAtDistance(const float DistanceAlongPath) const
+	{
+		int32 SegmentIndex = 0;
+		while (SegmentIndex < ((int32)NumPoints - 2))
+		{
+			if (DistanceAlongPath < LanePointProgressions[SegmentIndex + 1].Get())
+			{
+				break;
+			}
+			SegmentIndex++;
+		}
+
+		return SegmentIndex;
+	}
+
+	float GetInterpolationTimeOnSegment(const int32 SegmentIndex, const float DistanceAlongPath) const
+	{
+		check(SegmentIndex >= 0 && SegmentIndex <= (int32)NumPoints - 2);
+		const float StartDistance = LanePointProgressions[SegmentIndex].Get();
+		const float EndDistance = LanePointProgressions[SegmentIndex + 1].Get();
+		const float SegLength = EndDistance - StartDistance;
+		const float InvSegLength = SegLength > KINDA_SMALL_NUMBER ? 1.0f / SegLength : 0.0f;
+		return FMath::Clamp((DistanceAlongPath - StartDistance) * InvSegLength, 0.0f, 1.0f);
+	}
+	
+	void InterpolatePointAndTangentOnSegment(const int32 SegmentIndex, const float DistanceAlongPath, FVector& OutPoint, FVector& OutTangent) const
+	{
+		const float T = GetInterpolationTimeOnSegment(SegmentIndex, DistanceAlongPath);
+		OutPoint = FMath::Lerp(LanePoints[SegmentIndex], LanePoints[SegmentIndex + 1], T);
+		OutTangent = FVector(FMath::Lerp(LaneTangentVectors[SegmentIndex].Get(), LaneTangentVectors[SegmentIndex + 1].Get(), T), 0.0f);
+	}
+
+	FVector InterpolatePointOnSegment(const int32 SegmentIndex, const float DistanceAlongPath) const
+	{
+		const float T = GetInterpolationTimeOnSegment(SegmentIndex, DistanceAlongPath);
+		return FMath::Lerp(LanePoints[SegmentIndex], LanePoints[SegmentIndex + 1], T);
+	}
+
+	void GetPointAndTangentAtDistance(const float DistanceAlongPath, FVector& OutPoint, FVector& OutTangent) const
+	{
+		if (NumPoints == 0)
+		{
+			OutPoint = FVector::ZeroVector;
+			OutTangent = FVector::ForwardVector;
+			return;
+		}
+		if (NumPoints == 1)
+		{
+			OutPoint = LanePoints[0];
+			OutTangent = FVector(LaneTangentVectors[0].Get(), 0.0f);
+			return;
+		}
+
+		const int32 SegmentIndex = FindSegmentIndexAtDistance(DistanceAlongPath);
+		InterpolatePointAndTangentOnSegment(SegmentIndex, DistanceAlongPath, OutPoint, OutTangent);
+	}
+
+	FVector GetPointAtDistance(const float DistanceAlongPath) const
+	{
+		if (NumPoints == 0)
+		{
+			return FVector::ZeroVector;
+		}
+		if (NumPoints == 1)
+		{
+			return LanePoints[0];
+		}
+
+		const int32 SegmentIndex = FindSegmentIndexAtDistance(DistanceAlongPath);
+		return InterpolatePointOnSegment(SegmentIndex, DistanceAlongPath);
+	}
+
+	bool IsDistanceAtLaneExtrema(const float Distance) const
+	{
+		static constexpr float Epsilon = 0.1f;
+		return Distance <= Epsilon || (Distance - LaneLength) >= -Epsilon;
+	}
+
+	FZoneGraphLaneHandle LaneHandle;
+	
+	/** Lane points */
+	TStaticArray<FVector, MaxPoints> LanePoints;
+
+	/** Cached length of the lane. */
+	float LaneLength = 0.0f;
+
+	/** Lane tangents */
+	TStaticArray<FMassSnorm8Vector2D, MaxPoints> LaneTangentVectors;
+
+	/** lane Advance distances */
+	TStaticArray<FMassInt16Real10, MaxPoints> LanePointProgressions;
+
+	/** Cached width of the lane. */
+	FMassInt16Real LaneWidth = FMassInt16Real(0.0f);
+
+	/** Additional space left of the lane */
+	FMassInt16Real LaneLeftSpace = FMassInt16Real(0.0f);
+
+	/** Additional space right of the lane */
+	FMassInt16Real LaneRightSpace = FMassInt16Real(0.0f);
+
+	/** ID incremented each time the cache is updated. */
+	uint16 CacheID = 0;
+	
+	/** Number of points on path. */
+	uint8 NumPoints = 0;
+};
+
+/** Describes short path along ZoneGraph */
+// @todo MassMovement: it should be possible to prune this down to 64bytes
+// - remove debug lane handle, and replace other with index
+// - see if we can remove move tangent?
+USTRUCT()
+struct MASSAIMOVEMENT_API FMassZoneGraphShortPathFragment : public FLWComponentData
+{
+	GENERATED_BODY()
+
+	FMassZoneGraphShortPathFragment() = default;
+	
+	static constexpr uint8 MaxPoints = 3;
+
+	void Reset()
+	{
+#if WITH_MASS_DEBUG
+		DebugLaneHandle.Reset();
+#endif
+		NextLaneHandle.Reset();
+		NextExitLinkType = EZoneLaneLinkType::None;
+		ProgressDistance = 0.0f;
+		NumPoints = 0;
+		bMoveReverse = false;
+		EndOfPathIntent = EMassMovementAction::Stand;
+		bPartialResult = false;
+		bDone = false;
+	}
+
+	/** Requests path along the current lane */
+	bool RequestPath(const FMassZoneGraphCachedLaneFragment& CachedLane, const FZoneGraphShortPathRequest& Request, const float CurrentDistanceAlongLane, const float AgentRadius);
+
+	/** Requests path to stand at current position. */
+	bool RequestStand(const FMassZoneGraphCachedLaneFragment& CachedLane, const float CurrentDistanceAlongLane, const FVector& CurrentPosition);
+	
+	bool IsDone() const
+	{
+		// @todo MassMovement: should we remove NumPoints == 0? The logic used to be quite different when it was really needed.
+		return NumPoints == 0 || bDone;
+	}
+
+#if WITH_MASS_DEBUG
+	/** Current lane handle, for debug */
+	FZoneGraphLaneHandle DebugLaneHandle;
+#endif
+	
+	/** If valid, the this lane will be set as current lane after the path follow is completed. */
+	FZoneGraphLaneHandle NextLaneHandle;
+	
+	/** Current progress distance along the lane. */
+	float ProgressDistance = 0.0f;
+	
+	/** Path points */
+	TStaticArray<FMassZoneGraphPathPoint, MaxPoints> Points;
+
+	/** If next lane is set, this is how to reach the lane from current lane. */
+	EZoneLaneLinkType NextExitLinkType = EZoneLaneLinkType::None;
+	
+	/** Number of points on path. */
+	uint8 NumPoints = 0;
+	
+	/** Intent at the end of the path. */
+	EMassMovementAction EndOfPathIntent = EMassMovementAction::Stand;
+
+	/** True if we're moving reverse */
+	uint8 bMoveReverse : 1;
+
+	/** True if the path was partial. */
+	uint8 bPartialResult : 1;
+
+	/** True when path follow is completed. */
+	uint8 bDone : 1;
+};
diff --git a/Engine/Plugins/AI/MassAI/Source/MassAIMovement/Public/MassZoneGraphMovementProcessors.h b/Engine/Plugins/AI/MassAI/Source/MassAIMovement/Public/MassZoneGraphMovementProcessors.h
new file mode 100644
--- /dev/null
+++ b/Engine/Plugins/AI/MassAI/Source/MassAIMovement/Public/MassZoneGraphMovementProcessors.h
@@ -0,0 +1,80 @@
+// Copyright Epic Games, Inc. All Rights Reserved.
+
+#pragma once
+
+#include "MassProcessor.h"
+#include "MassTranslator.h"
+#include "MassZoneGraphMovementProcessors.generated.h"
+
+class UMassSignalSubsystem;
+class UZoneGraphSubsystem;
+
+/**
+ * Processor for initializing nearest location on ZoneGraph.
+ */
+UCLASS()
+class MASSAIMOVEMENT_API UMassZoneGraphLocationInitializer : public UMassFragmentInitializer
+{
+	GENERATED_BODY()
+	
+public:
+	UMassZoneGraphLocationInitializer();
+
+protected:
+	virtual void Initialize(UObject& Owner) override;
+	virtual void ConfigureQueries() override;
+	virtual void Execute(UEntitySubsystem& EntitySubsystem, FLWComponentSystemExecutionContext& Context) override;
+
+	FLWComponentQuery EntityQuery;
+
+	UPROPERTY(Transient)
+	UZoneGraphSubsystem* ZoneGraphSubsystem = nullptr;
+
+	UPROPERTY(Transient)
+	UMassSignalSubsystem* SignalSubsystem = nullptr;
+};
+
+/** 
+ * Processor for updating move target on ZoneGraph path.
+ */
+UCLASS()
+class MASSAIMOVEMENT_API UMassZoneGraphPathFollowProcessor : public UPipeProcessor
+{
+	GENERATED_BODY()
+
+protected:
+	UMassZoneGraphPathFollowProcessor();
+	
+	virtual void Initialize(UObject& Owner) override;
+	virtual void ConfigureQueries() override;
+	virtual void Execute(UEntitySubsystem& EntitySubsystem, FLWComponentSystemExecutionContext& Context) override;
+
+	FLWComponentQuery EntityQuery_Conditional;
+
+	UPROPERTY(Transient)
+	UZoneGraphSubsystem* ZoneGraphSubsystem = nullptr;
+
+	UPROPERTY(Transient)
+	UMassSignalSubsystem* SignalSubsystem = nullptr;
+};
+
+/** 
+ * Processor for updating steering towards MoveTarget.
+ */
+UCLASS()
+class MASSAIMOVEMENT_API UMassZoneGraphSteeringProcessor : public UPipeProcessor
+{
+	GENERATED_BODY()
+
+protected:
+	UMassZoneGraphSteeringProcessor();
+	
+	virtual void Initialize(UObject& Owner) override;
+	virtual void ConfigureQueries() override;
+	virtual void Execute(UEntitySubsystem& EntitySubsystem, FLWComponentSystemExecutionContext& Context) override;
+
+	FLWComponentQuery EntityQuery;
+
+	UPROPERTY(Transient)
+	UMassSignalSubsystem* SignalSubsystem = nullptr;
+};
diff --git a/Engine/Plugins/AI/MassAI/Source/MassAIMovement/Public/MassZoneGraphMovementTrait.h b/Engine/Plugins/AI/MassAI/Source/MassAIMovement/Public/MassZoneGraphMovementTrait.h
new file mode 100644
--- /dev/null
+++ b/Engine/Plugins/AI/MassAI/Source/MassAIMovement/Public/MassZoneGraphMovementTrait.h
@@ -0,0 +1,20 @@
+﻿// Copyright Epic Games, Inc. All Rights Reserved.
+
+#pragma once 
+
+#include "MassEntityTraitBase.h"
+#include "MassMovementSubsystem.h"
+#include "MassAIMovementTypes.h"
+#include "MassZoneGraphMovementTrait.generated.h"
+
+UCLASS(BlueprintType, EditInlineNew, CollapseCategories, meta = (DisplayName = "ZoneGraph Movement"))
+class MASSAIMOVEMENT_API UMassZoneGraphMovementTrait : public UMassEntityTraitBase
+{
+	GENERATED_BODY()
+
+protected:
+	virtual void BuildTemplate(FMassEntityTemplateBuildContext& BuildContext, UWorld& World) const override;
+
+	UPROPERTY(Category="Movement", EditAnywhere)
+	FMassMovementConfigRef Config;
+};
diff --git a/Engine/Plugins/AI/MassAI/Source/MassAIMovement/Public/MassZoneGraphMovementTypes.h b/Engine/Plugins/AI/MassAI/Source/MassAIMovement/Public/MassZoneGraphMovementTypes.h
new file mode 100644
--- /dev/null
+++ b/Engine/Plugins/AI/MassAI/Source/MassAIMovement/Public/MassZoneGraphMovementTypes.h
@@ -0,0 +1,101 @@
+// Copyright Epic Games, Inc. All Rights Reserved.
+
+#pragma once
+
+#include "MassCommonTypes.h"
+#include "MassAIMovementTypes.h"
+#include "MassZoneGraphMovementTypes.generated.h"
+
+/** Describes path request along one ZoneGraph lane. If the NextLaneHandle is set, lane is changed when path finishes. */
+USTRUCT()
+struct MASSAIMOVEMENT_API FZoneGraphShortPathRequest
+{
+	GENERATED_BODY()
+
+	FZoneGraphShortPathRequest()
+		: bMoveReverse(false)
+		, bIsEndOfPathPositionSet(false)
+		, bIsEndOfPathDirectionSet(false)
+	{
+	}
+
+	FString ToString() const
+	{
+		return FString::Printf(TEXT("%s to distance %.1f Next lane: %s of type %s. End of path intent:%s"),
+			bMoveReverse ? TEXT("forward") : TEXT("reverse"), TargetDistance,
+			NextLaneHandle.IsValid() ? *NextLaneHandle.ToString() : TEXT("unset"),
+			NextLaneHandle.IsValid() ? *UEnum::GetValueAsString(NextExitLinkType) : TEXT("Unset"),
+			*UEnum::GetValueAsString(EndOfPathIntent));
+	}
+
+	/** Position used as the start of the path*/
+	UPROPERTY(Transient)
+	FVector StartPosition = FVector::ZeroVector;
+
+	/** Optional specific point at the end of the path. */
+	UPROPERTY(Transient)
+	FVector EndOfPathPosition = FVector::ZeroVector;
+
+	/** If set, the lane will be switched when path finishes. */
+	UPROPERTY(Transient)
+	FZoneGraphLaneHandle NextLaneHandle;
+
+	/** Distance to move. */
+	UPROPERTY(Transient)
+	float TargetDistance = 0.0f;
+
+	/** Optional specific direction at the end of the path. Used only if EndOfPathPosition is set. */
+	UPROPERTY(Transient)
+	FMassSnorm8Vector EndOfPathDirection;
+
+	/** If start or end of path is off-lane, the distance along the lane is pushed forward/back along the lane to make smoother transition. */
+	UPROPERTY(Transient)
+	uint16 AnticipationDistance = 50;
+
+	UPROPERTY(Transient)
+	FMassInt16Real EndOfPathOffset = FMassInt16Real(0.0f);
+
+	/** Movement intent at the end of the path. */
+	UPROPERTY(Transient)
+	EMassMovementAction EndOfPathIntent = EMassMovementAction::Stand;
+
+	/** How the NextLaneHandle links to current lane. */
+	UPROPERTY(Transient)
+	EZoneLaneLinkType NextExitLinkType = EZoneLaneLinkType::None;
+
+	/** If true, move backwards along the lane. */
+	UPROPERTY(Transient)
+	uint8 bMoveReverse : 1;
+
+	/** Indicates if the optional end of path position is set. */
+	UPROPERTY(Transient)
+	uint8 bIsEndOfPathPositionSet : 1;
+
+	/** Indicates if the optional end of path direction is set. Used only if EndOfPathPosition is set. */
+	UPROPERTY(Transient)
+	uint8 bIsEndOfPathDirectionSet : 1;
+};
+
+USTRUCT()
+struct MASSAIMOVEMENT_API FMassZoneGraphPathPoint
+{
+	GENERATED_BODY()
+
+	/** Position of the path. */
+	FVector Position = FVector::ZeroVector;
+
+	/** Tangent direction of the path. */
+	FMassSnorm8Vector2D Tangent;
+
+	/** Position of the point along the original path. (Could potentially be uint16 at 10cm accuracy) */
+	FMassInt16Real10 DistanceAlongLane = FMassInt16Real10(0.0f);
+
+	/** Distance along the offset path from first point. (Could potentially be uint16 at 10cm accuracy) */
+	FMassInt16Real Distance = FMassInt16Real(0.0f);
+
+	/** True if this point is assumed to be off lane. */
+	uint8 bOffLane : 1;
+
+	/** True if this point is lane start or end point. */
+	uint8 bIsLaneExtrema : 1;
+};
diff --git a/Engine/Plugins/AI/MassAI/Source/MassAIMovement/Public/MassZoneGraphMovementUtils.h b/Engine/Plugins/AI/MassAI/Source/MassAIMovement/Public/MassZoneGraphMovementUtils.h
new file mode 100644
--- /dev/null
+++ b/Engine/Plugins/AI/MassAI/Source/MassAIMovement/Public/MassZoneGraphMovementUtils.h
@@ -0,0 +1,45 @@
+﻿// Copyright Epic Games, Inc. All Rights Reserved.
+
+#pragma once
+
+#include "LWComponentTypes.h"
+
+struct FZoneGraphShortPathRequest;
+struct FMassZoneGraphCachedLaneFragment;
+struct FMassZoneGraphShortPathFragment;
+struct FMassMoveTargetFragment;
+struct FMassZoneGraphPathRequestFragment;
+struct FMassZoneGraphLaneLocationFragment;
+class UMassSignalSubsystem;
+class UZoneGraphSubsystem;
+class UZoneGraphAnnotationSubsystem;
+
+namespace UE::MassMovement
+{
+	MASSAIMOVEMENT_API bool ActivateActionMove(const UWorld& World,
+											 const UObject* Requester,
+											 const FLWEntity Entity,
+											 const UZoneGraphSubsystem& ZoneGraphSubsystem,
+											 const FMassZoneGraphLaneLocationFragment& LaneLocation,
+											 const FZoneGraphShortPathRequest& PathRequest,
+											 const float AgentRadius,
+											 const float DesiredSpeed,
+											 FMassMoveTargetFragment& MoveTarget,
+											 FMassZoneGraphShortPathFragment& ShortPath,
+											 FMassZoneGraphCachedLaneFragment& CachedLane);
+
+	MASSAIMOVEMENT_API bool ActivateActionStand(const UWorld& World,
+											  const UObject* Requester,
+											  const FLWEntity Entity,
+											  const UZoneGraphSubsystem& ZoneGraphSubsystem,
+											  const FMassZoneGraphLaneLocationFragment& LaneLocation,
+											  const float DesiredSpeed,
+											  FMassMoveTargetFragment& MoveTarget,
+											  FMassZoneGraphShortPathFragment& ShortPath,
+											  FMassZoneGraphCachedLaneFragment& CachedLane);
+
+	MASSAIMOVEMENT_API bool ActivateActionAnimate(const UWorld& World,
+												const UObject* Requester,
+												const FLWEntity Entity,
+												FMassMoveTargetFragment& MoveTarget);
+};
diff --git a/Engine/Plugins/AI/MassAI/Source/MassAIMovement/Public/MovementProcessor.h b/Engine/Plugins/AI/MassAI/Source/MassAIMovement/Public/MovementProcessor.h
new file mode 100644
--- /dev/null
+++ b/Engine/Plugins/AI/MassAI/Source/MassAIMovement/Public/MovementProcessor.h
@@ -0,0 +1,67 @@
+// Copyright Epic Games, Inc. All Rights Reserved.
+
+#pragma once
+
+#include "MassProcessor.h"
+#include "MassTranslator.h"
+#include "MassMovementSubsystem.h"
+#include "MovementProcessor.generated.h"
+
+/** Move and orient */ 
+UCLASS()
+class MASSAIMOVEMENT_API UMassProcessor_Movement : public UPipeProcessor
+{
+	GENERATED_BODY()
+
+public:
+	UMassProcessor_Movement();
+
+protected:
+	virtual void ConfigureQueries() override;
+	virtual void Execute(UEntitySubsystem& EntitySubsystem, FLWComponentSystemExecutionContext& Context) override;
+
+	FLWComponentQuery EntityQuery;
+};
+
+/** Base movement processor for 'grid localized circular agent' */
+UCLASS()
+class MASSAIMOVEMENT_API UMassProcessor_AgentMovement : public UMassProcessor_Movement
+{
+	GENERATED_BODY()
+public:
+	UMassProcessor_AgentMovement();
+
+protected:
+	virtual void ConfigureQueries() override;
+};
+
+/** Destructor processor to remove avoidance obstacles from the avoidance obstacle grid */
+UCLASS()
+class MASSAIMOVEMENT_API UMassAvoidanceObstacleRemoverFragmentDestructor : public UMassFragmentDestructor
+{
+	GENERATED_BODY()
+
+	UMassAvoidanceObstacleRemoverFragmentDestructor();
+	TWeakObjectPtr<UMassMovementSubsystem> WeakMovementSubsystem;
+
+protected:
+	virtual void ConfigureQueries() override;
+	virtual void Initialize(UObject& Owner) override;
+	virtual void Execute(UEntitySubsystem& EntitySubsystem, FLWComponentSystemExecutionContext& Context) override;
+
+	FLWComponentQuery EntityQuery;
+};
+
+UCLASS()
+class MASSAIMOVEMENT_API UVelocityFragmentInitializer : public UMassFragmentInitializer
+{
+	GENERATED_BODY()
+public:
+	UVelocityFragmentInitializer();
+
+protected:
+	virtual void ConfigureQueries() override;
+	virtual void Execute(UEntitySubsystem& EntitySubsystem, FLWComponentSystemExecutionContext& Context) override;
+
+	FLWComponentQuery EntityQuery;
+};
diff --git a/Engine/Plugins/AI/MassAI/Source/MassAIMovement/Public/NavigationProcessor.h b/Engine/Plugins/AI/MassAI/Source/MassAIMovement/Public/NavigationProcessor.h
new file mode 100644
--- /dev/null
+++ b/Engine/Plugins/AI/MassAI/Source/MassAIMovement/Public/NavigationProcessor.h
@@ -0,0 +1,59 @@
+// Copyright Epic Games, Inc. All Rights Reserved.
+
+#pragma once
+
+#include "MassProcessor.h"
+#include "NavigationProcessor.generated.h"
+
+struct FMassDynamicObstacleFragment;
+class UMassSignalSubsystem;
+class UMassMovementSubsystem;
+
+/**
+ * Update agent's transform fragments from post avoidance velocity.
+ * Sets agents Z position based on MoveTarget fragment.
+ * @todo: separate apply velocity and height adjustment
+ */
+UCLASS()
+class MASSAIMOVEMENT_API UMassApplyVelocityMoveTargetProcessor : public UPipeProcessor
+{
+	GENERATED_BODY()
+
+public:
+	UMassApplyVelocityMoveTargetProcessor();
+
+protected:
+	virtual void ConfigureQueries() override;
+	virtual void Execute(UEntitySubsystem& EntitySubsystem, FLWComponentSystemExecutionContext& Context) override;
+
+private:
+	FLWComponentQuery HighResEntityQuery;
+	FLWComponentQuery LowResEntityQuery_Conditional;
+};
+
+/** Handle dynamic obstacles. */
+UCLASS()
+class MASSAIMOVEMENT_API UMassDynamicObstacleProcessor : public UPipeProcessor
+{
+	GENERATED_BODY()
+
+public:
+	UMassDynamicObstacleProcessor();
+
+	/** Delay before sending the stop notification once the entity has stop moving. */
+	UPROPERTY(Category = Settings, EditAnywhere, meta = (ClampMin = "0"))
+	float DelayBeforeStopNotification = 0.3f;
+
+	/** Distance within which the obstacle is considered not moving. */
+	UPROPERTY(Category = Settings, EditAnywhere, meta = (ClampMin = "0"))
+	float DistanceBuffer = 10.f;
+	
+protected:
+	virtual void ConfigureQueries() override;
+	virtual void Execute(UEntitySubsystem& EntitySubsystem, FLWComponentSystemExecutionContext& Context) override;
+	virtual void OnStop(FMassDynamicObstacleFragment& OutObstacle, const float BlockingRadius) {}
+	virtual void OnMove(FMassDynamicObstacleFragment& OutObstacle) {}
+
+private:
+	FLWComponentQuery EntityQuery_Conditional;
+};
diff --git a/Engine/Plugins/AI/MassAI/Source/MassAIMovement/Public/SnapToNavigationProcessor.h b/Engine/Plugins/AI/MassAI/Source/MassAIMovement/Public/SnapToNavigationProcessor.h
new file mode 100644
--- /dev/null
+++ b/Engine/Plugins/AI/MassAI/Source/MassAIMovement/Public/SnapToNavigationProcessor.h
@@ -0,0 +1,25 @@
+// Copyright Epic Games, Inc. All Rights Reserved.
+
+#pragma once
+
+#include "MassProcessor.h"
+#include "SnapToNavigationProcessor.generated.h"
+
+class ANavigationData;
+
+UCLASS()
+class MASSAIMOVEMENT_API USnapToNavigationProcessor : public UPipeProcessor
+{
+	GENERATED_BODY()
+public:
+	USnapToNavigationProcessor();
+
+protected:	
+	virtual void ConfigureQueries() override;
+	virtual void Execute(UEntitySubsystem& EntitySubsystem, FLWComponentSystemExecutionContext& Context) override;
+	virtual void Initialize(UObject& InOwner) override;
+
+	TWeakObjectPtr<ANavigationData> WeakNavData;
+
+	FLWComponentQuery EntityQuery;
+};
diff --git a/Engine/Plugins/AI/MassAI/Source/MassAIMovementEditor/MassAIMovementEditor.Build.cs b/Engine/Plugins/AI/MassAI/Source/MassAIMovementEditor/MassAIMovementEditor.Build.cs
new file mode 100644
--- /dev/null
+++ b/Engine/Plugins/AI/MassAI/Source/MassAIMovementEditor/MassAIMovementEditor.Build.cs
@@ -0,0 +1,48 @@
+// Copyright Epic Games, Inc. All Rights Reserved.
+
+namespace UnrealBuildTool.Rules
+{
+	public class MassAIMovementEditor : ModuleRules
+	{
+		public MassAIMovementEditor(ReadOnlyTargetRules Target) : base(Target)
+		{
+			PublicIncludePaths.AddRange(
+			new string[] {
+			}
+			);
+
+			PublicDependencyModuleNames.AddRange(
+			new string[] {
+				"Core",
+				"CoreUObject",
+				"Engine",
+				"InputCore",
+				"AssetTools",
+				"EditorFramework",
+				"UnrealEd",
+				"RHI",
+				"Slate",
+				"SlateCore",
+				"EditorStyle",
+				"PropertyEditor",
+				"MassEntity",
+				"DetailCustomizations",
+				"MassCommon",
+				"MassAIMovement",
+				"ZoneGraph",
+			}
+			);
+
+			PrivateDependencyModuleNames.AddRange(
+			new string[] {
+				"RenderCore",
+				"KismetWidgets",
+				"ToolMenus",
+				"AppFramework",
+				"Projects",
+			}
+			);
+		}
+
+	}
+}
diff --git a/Engine/Plugins/AI/MassAI/Source/MassAIMovementEditor/Private/MassMovementTestingActor.cpp b/Engine/Plugins/AI/MassAI/Source/MassAIMovementEditor/Private/MassMovementTestingActor.cpp
new file mode 100644
--- /dev/null
+++ b/Engine/Plugins/AI/MassAI/Source/MassAIMovementEditor/Private/MassMovementTestingActor.cpp
@@ -0,0 +1,335 @@
+﻿// Copyright Epic Games, Inc. All Rights Reserved.
+
+#include "MassMovementTestingActor.h"
+#include "UObject/ConstructorHelpers.h"
+#include "DebugRenderSceneProxy.h"
+#include "ZoneGraphSubsystem.h"
+#include "ZoneGraphDelegates.h"
+#include "ZoneGraphRenderingUtilities.h"
+#include "ZoneGraphData.h"
+
+#if !UE_BUILD_SHIPPING && !UE_BUILD_TEST
+
+//////////////////////////////////////////////////////////////////////////
+// FMassMovementTestingSceneProxy
+FMassMovementTestingSceneProxy::FMassMovementTestingSceneProxy(const UPrimitiveComponent& InComponent)
+	: FDebugRenderSceneProxy(&InComponent)
+{
+	DrawType = WireMesh;
+}
+
+SIZE_T FMassMovementTestingSceneProxy::GetTypeHash() const
+{
+	static size_t UniquePointer;
+	return reinterpret_cast<size_t>(&UniquePointer);
+}
+
+FPrimitiveViewRelevance FMassMovementTestingSceneProxy::GetViewRelevance(const FSceneView* View) const
+{
+	FPrimitiveViewRelevance Result;
+	Result.bDrawRelevance = IsShown(View);
+	Result.bDynamicRelevance = true;
+	Result.bShadowRelevance = IsShadowCast(View);
+	Result.bEditorPrimitiveRelevance = UseEditorCompositing(View);
+	return Result;
+}
+
+uint32 FMassMovementTestingSceneProxy::GetMemoryFootprint(void) const
+{
+	return sizeof(*this) + FDebugRenderSceneProxy::GetAllocatedSize();
+}
+
+#endif // !UE_BUILD_SHIPPING && !UE_BUILD_TEST
+
+//////////////////////////////////////////////////////////////////////////
+// UMassMovementTestingComponent
+
+UMassMovementTestingComponent::UMassMovementTestingComponent(const FObjectInitializer& ObjectInitialize)
+	: Super(ObjectInitialize)
+{
+	SearchExtent = FVector(150.0f);
+}
+
+#if WITH_EDITOR
+void UMassMovementTestingComponent::PostEditChangeProperty(FPropertyChangedEvent& PropertyChangedEvent)
+{
+	Super::PostEditChangeProperty(PropertyChangedEvent);
+	UpdateTests();
+}
+#endif
+
+void UMassMovementTestingComponent::OnRegister()
+{
+	Super::OnRegister();
+#if WITH_EDITOR
+	// Force to update tests when ever the data changes.
+	OnDataChangedHandle = UE::ZoneGraphDelegates::OnZoneGraphDataBuildDone.AddUObject(this, &UMassMovementTestingComponent::OnZoneGraphDataBuildDone);
+#endif
+	OnDataAddedHandle = UE::ZoneGraphDelegates::OnPostZoneGraphDataAdded.AddUObject(this, &UMassMovementTestingComponent::OnZoneGraphDataChanged);
+	OnDataRemovedHandle = UE::ZoneGraphDelegates::OnPreZoneGraphDataRemoved.AddUObject(this, &UMassMovementTestingComponent::OnZoneGraphDataChanged);
+
+	ZoneGraph = UWorld::GetSubsystem<UZoneGraphSubsystem>(GetWorld());
+
+	UpdateTests();
+}
+
+void UMassMovementTestingComponent::OnUnregister()
+{
+	Super::OnUnregister();
+#if WITH_EDITOR
+	UE::ZoneGraphDelegates::OnZoneGraphDataBuildDone.Remove(OnDataChangedHandle);
+#endif
+	UE::ZoneGraphDelegates::OnPostZoneGraphDataAdded.Remove(OnDataAddedHandle);
+	UE::ZoneGraphDelegates::OnPreZoneGraphDataRemoved.Remove(OnDataRemovedHandle);
+}
+
+void UMassMovementTestingComponent::OnZoneGraphDataChanged(const AZoneGraphData* ZoneGraphData)
+{
+	// Only consider valid graph from our world
+	if (ZoneGraphData == nullptr || ZoneGraphData->GetWorld() != GetWorld())
+	{
+		return;
+	}
+
+	UpdateTests();
+}
+
+#if WITH_EDITOR
+void UMassMovementTestingComponent::OnZoneGraphDataBuildDone(const struct FZoneGraphBuildData& BuildData)
+{
+	UpdateTests();
+}
+#endif
+
+FBoxSphereBounds UMassMovementTestingComponent::CalcBounds(const FTransform& LocalToWorld) const
+{
+	const FVector ActorPosition = LocalToWorld.GetTranslation();
+	return FBox(ActorPosition - SearchExtent, ActorPosition + SearchExtent);
+}
+
+void UMassMovementTestingComponent::UpdateTests()
+{
+	if (!ZoneGraph)
+	{
+		return;
+	}
+
+	const FVector WorldPosition = GetOwner()->GetActorLocation();
+
+	// Find nearest
+	float DistanceSqr = 0.0f;
+	if (PinnedLane.IsValid())
+	{
+		ZoneGraph->FindNearestLocationOnLane(LaneLocation.LaneHandle, FBox(WorldPosition - SearchExtent, WorldPosition + SearchExtent), LaneLocation, DistanceSqr);
+	}
+	else
+	{
+		ZoneGraph->FindNearestLane(FBox(WorldPosition - SearchExtent, WorldPosition + SearchExtent), QueryFilter, LaneLocation, DistanceSqr);
+	}
+
+	CachedLane.Reset();
+	ShortPaths.Reset();
+
+	GoalLaneLocation.Reset();
+
+	if (LaneLocation.IsValid())
+	{
+		const FVector GoalWorldPosition = WorldPosition + GoalPosition;
+		ZoneGraph->FindNearestLocationOnLane(LaneLocation.LaneHandle, FBox(GoalWorldPosition - SearchExtent, GoalWorldPosition + SearchExtent), GoalLaneLocation, DistanceSqr);
+
+		if (const FZoneGraphStorage* ZoneStorage = ZoneGraph->GetZoneGraphStorage(LaneLocation.LaneHandle.DataHandle))
+		{
+			constexpr float InflateDistance = 200.0f;
+			CachedLane.CacheLaneData(*ZoneStorage, LaneLocation.LaneHandle, LaneLocation.DistanceAlongLane, GoalLaneLocation.DistanceAlongLane, InflateDistance);
+
+			// Path find has smaller buffer than cache, calculate up to 5 paths. Also allows to test that one path find can be continued to next one.
+			FZoneGraphShortPathRequest PathRequest;
+			PathRequest.bMoveReverse = GoalLaneLocation.DistanceAlongLane < LaneLocation.DistanceAlongLane;
+			PathRequest.TargetDistance = GoalLaneLocation.DistanceAlongLane;
+			if (bHasSpecificEndPoint)
+			{
+				PathRequest.EndOfPathPosition = GoalWorldPosition;
+			}
+			PathRequest.AnticipationDistance = AnticipationDistance;
+
+			PathRequest.StartPosition = WorldPosition;
+			float StartDistanceAlongPath = LaneLocation.DistanceAlongLane;
+			for (int32 Iter = 0; Iter < 5; Iter++)
+			{
+				FMassZoneGraphShortPathFragment& ShortPath = ShortPaths.AddDefaulted_GetRef();
+				ShortPath.RequestPath(CachedLane, PathRequest, StartDistanceAlongPath, AgentRadius);
+
+				if (ShortPath.NumPoints == 0 || !ShortPath.bPartialResult)
+				{
+					break;
+				}
+
+				// Restart from last lane end.
+				PathRequest.StartPosition = ShortPath.Points[ShortPath.NumPoints - 1].Position;
+				StartDistanceAlongPath = ShortPath.Points[ShortPath.NumPoints - 1].DistanceAlongLane.Get();
+			}
+		}
+	}
+	
+	MarkRenderStateDirty();
+}
+
+void UMassMovementTestingComponent::PinLane()
+{
+	if (LaneLocation.IsValid())
+	{
+		PinnedLane = LaneLocation.LaneHandle;
+	}
+	UpdateTests();
+}
+
+void UMassMovementTestingComponent::ClearPinnedLane()
+{
+	PinnedLane.Reset();
+	UpdateTests();
+}
+
+#if !UE_BUILD_SHIPPING
+void UMassMovementTestingComponent::DestroyRenderState_Concurrent()
+{
+	DebugDrawDelegateHelper.UnregisterDebugDrawDelegate();
+	Super::DestroyRenderState_Concurrent();
+}
+
+FPrimitiveSceneProxy* UMassMovementTestingComponent::CreateSceneProxy()
+{
+	FMassMovementTestingSceneProxy* DebugProxy = new FMassMovementTestingSceneProxy(*this);
+
+	const FTransform& ActorTransform = GetOwner()->GetActorTransform();
+	const FVector Center = ActorTransform.GetLocation(); 
+	const FVector Goal = Center + GoalPosition; 
+	
+	static constexpr float TickSize = 25.0f;
+
+	DebugProxy->Boxes.Emplace(FBox(Center - SearchExtent, Center + SearchExtent), FColorList::SkyBlue);
+
+	DebugProxy->Lines.Emplace(Center + FVector(-TickSize, 0, 0), Center + FVector(TickSize, 0, 0), FColorList::SkyBlue);
+	DebugProxy->Lines.Emplace(Center + FVector(0, -TickSize, 0), Center + FVector(0, TickSize, 0), FColorList::SkyBlue);
+	DebugProxy->Lines.Emplace(Center + FVector(0, 0, -TickSize), Center + FVector(0, 0, TickSize), FColorList::SkyBlue);
+
+	DebugProxy->Lines.Emplace(Goal + FVector(-TickSize, 0, 0), Goal + FVector(TickSize, 0, 0), FColorList::YellowGreen);
+	DebugProxy->Lines.Emplace(Goal + FVector(0, -TickSize, 0), Goal + FVector(0, TickSize, 0), FColorList::YellowGreen);
+	DebugProxy->Lines.Emplace(Goal + FVector(0, 0, -TickSize), Goal + FVector(0, 0, TickSize), FColorList::YellowGreen);
+
+	const FVector ZOffset(0,0,10);
+	const FVector ZOffsetPath(0,0,20);
+	
+	if (LaneLocation.IsValid())
+	{
+		DebugProxy->Lines.Emplace(LaneLocation.Position - LaneLocation.Up * TickSize, LaneLocation.Position + LaneLocation.Up * TickSize * 3.0f, FColorList::SkyBlue);
+		DebugProxy->Lines.Emplace(Center, LaneLocation.Position, FColorList::SkyBlue);
+	}
+
+	if (GoalLaneLocation.IsValid())
+	{
+		DebugProxy->Lines.Emplace(GoalLaneLocation.Position - GoalLaneLocation.Up * TickSize, GoalLaneLocation.Position + GoalLaneLocation.Up * TickSize * 3.0f, FColorList::YellowGreen);
+		DebugProxy->Lines.Emplace(Goal, GoalLaneLocation.Position, FColorList::YellowGreen);
+	}
+
+	if (CachedLane.NumPoints > 1)
+	{
+		const float LeftSpace = CachedLane.LaneWidth.Get() * 0.5f + CachedLane.LaneLeftSpace.Get() - AgentRadius;
+		const float RightSpace = CachedLane.LaneWidth.Get() * 0.5f + CachedLane.LaneRightSpace.Get() - AgentRadius;
+		
+		for (uint8 PointIndex = 0; PointIndex < CachedLane.NumPoints - 1; PointIndex++)
+		{
+			DebugProxy->Lines.Emplace(CachedLane.LanePoints[PointIndex] + ZOffset, CachedLane.LanePoints[PointIndex + 1] + ZOffset, FColorList::Grey, 4.0f);
+
+			// Draw boundaries
+			const FVector StartTangent = CachedLane.LaneTangentVectors[PointIndex].GetVector();
+			const FVector StartLeftDir = FVector::CrossProduct(StartTangent, FVector::UpVector);
+			const FVector StartLeftPos = CachedLane.LanePoints[PointIndex] + ZOffset + StartLeftDir * LeftSpace;
+			const FVector StartRightPos = CachedLane.LanePoints[PointIndex] + ZOffset + StartLeftDir * -RightSpace;
+
+			const FVector EndTangent = CachedLane.LaneTangentVectors[PointIndex + 1].GetVector();
+			const FVector EndLeftDir = FVector::CrossProduct(EndTangent, FVector::UpVector);
+			const FVector EndLeftPos = CachedLane.LanePoints[PointIndex + 1] + ZOffset + EndLeftDir * LeftSpace;
+			const FVector EndRightPos = CachedLane.LanePoints[PointIndex + 1] + ZOffset + EndLeftDir * -RightSpace;
+
+			DebugProxy->Lines.Emplace(StartLeftPos, EndLeftPos, FColorList::LightGrey, 1.0f);
+			DebugProxy->Lines.Emplace(StartRightPos, EndRightPos, FColorList::LightGrey, 1.0f);
+
+			if (PointIndex == 0)
+			{
+				DebugProxy->Lines.Emplace(StartLeftPos, StartRightPos, FColorList::LightGrey, 1.0f);
+			}
+			else if (PointIndex == (CachedLane.NumPoints - 1))
+			{
+				DebugProxy->Lines.Emplace(EndLeftPos, EndRightPos, FColorList::LightGrey, 1.0f);
+			}
+			
+		}
+	}
+
+	for (int32 PathIndex = 0; PathIndex < ShortPaths.Num(); PathIndex++)
+	{
+		const FMassZoneGraphShortPathFragment& ShortPath = ShortPaths[PathIndex];
+		const FColor PathColor = (PathIndex & 1) ? FColor::Orange : FColor::Emerald;
+	
+		if (ShortPath.NumPoints > 1)
+		{
+			for (uint8 PointIndex = 0; PointIndex < ShortPath.NumPoints; PointIndex++)
+			{
+				const FVector Position = ShortPath.Points[PointIndex].Position + ZOffsetPath;
+				const FVector LeftDir = FVector::CrossProduct(ShortPath.Points[PointIndex].Tangent.GetVector(), FVector::UpVector);
+
+				DebugProxy->Lines.Emplace(Position - LeftDir * TickSize, Position + LeftDir * TickSize, PathColor, 2.0f);
+				DebugProxy->Texts.Emplace(FString::Printf(TEXT("%d"), (int32)PointIndex), Position, PathColor);
+			}
+
+			for (uint8 PointIndex = 0; PointIndex < ShortPath.NumPoints - 1; PointIndex++)
+			{
+				DebugProxy->Lines.Emplace(ShortPath.Points[PointIndex].Position + ZOffsetPath, ShortPath.Points[PointIndex + 1].Position + ZOffsetPath, PathColor, 4.0f);
+			}
+		}
+	}
+
+	DebugDrawDelegateHelper.InitDelegateHelper(DebugProxy);
+	DebugDrawDelegateHelper.RegisterDebugDrawDelegate();
+
+	return DebugProxy;
+}
+#endif // !UE_BUILD_SHIPPING
+
+//////////////////////////////////////////////////////////////////////////
+// UMassMovementTestingComponent
+
+AMassMovementTestingActor::AMassMovementTestingActor(const FObjectInitializer& ObjectInitializer)
+	: Super(ObjectInitializer)
+{
+	DebugComp = CreateDefaultSubobject<UMassMovementTestingComponent>(TEXT("DebugComp"));
+	RootComponent = DebugComp;
+
+	SetCanBeDamaged(false);
+}
+
+#if WITH_EDITOR
+void AMassMovementTestingActor::PostEditMove(bool bFinished)
+{
+	if (DebugComp)
+	{
+		DebugComp->UpdateTests();
+	}
+}
+#endif
+
+void AMassMovementTestingActor::PinLane()
+{
+	if (DebugComp)
+	{
+		DebugComp->PinLane();
+	}
+}
+	
+void AMassMovementTestingActor::ClearPinnedLane()
+{
+	if (DebugComp)
+	{
+		DebugComp->ClearPinnedLane();
+	}
+}
diff --git a/Engine/Plugins/AI/MassAI/Source/MassAIMovementEditor/Private/MassMovementTestingActor.h b/Engine/Plugins/AI/MassAI/Source/MassAIMovementEditor/Private/MassMovementTestingActor.h
new file mode 100644
--- /dev/null
+++ b/Engine/Plugins/AI/MassAI/Source/MassAIMovementEditor/Private/MassMovementTestingActor.h
@@ -0,0 +1,129 @@
+// Copyright Epic Games, Inc. All Rights Reserved.
+
+#pragma once
+
+#include "GameFramework/Actor.h"
+#include "Components/PrimitiveComponent.h"
+#include "DebugRenderSceneProxy.h"
+#include "MassAvoidanceProcessors.h"
+#include "ZoneGraphTypes.h"
+#include "MassZoneGraphMovementFragments.h"
+#include "MassMovementTestingActor.generated.h"
+
+class AZoneGraphData;
+class UZoneGraphTestingComponent;
+class UZoneGraphSubsystem;
+class AMassMovementTestingActor;
+
+#if !UE_BUILD_SHIPPING && !UE_BUILD_TEST
+class MASSAIMOVEMENTEDITOR_API FMassMovementTestingSceneProxy final : public FDebugRenderSceneProxy
+{
+public:
+	FMassMovementTestingSceneProxy(const UPrimitiveComponent& InComponent);
+	
+	virtual SIZE_T GetTypeHash() const override;
+	virtual FPrimitiveViewRelevance GetViewRelevance(const FSceneView* View) const override;
+	virtual uint32 GetMemoryFootprint(void) const override;
+};
+#endif
+
+/** Component for testing MassMovement functionality. */
+UCLASS(ClassGroup = Custom, hidecategories = (Physics, Collision, Rendering, Cooking, Lighting, Navigation, Tags, HLOD, Mobile, AssetUserData, Activation))
+class MASSAIMOVEMENTEDITOR_API UMassMovementTestingComponent : public UPrimitiveComponent
+{
+	GENERATED_BODY()
+public:
+	UMassMovementTestingComponent(const FObjectInitializer& ObjectInitializer = FObjectInitializer::Get());
+
+#if WITH_EDITOR
+	virtual void PostEditChangeProperty(FPropertyChangedEvent& PropertyChangedEvent) override;
+#endif // WITH_EDITOR
+
+	virtual void OnRegister() override;
+	virtual void OnUnregister() override;
+
+	virtual FBoxSphereBounds CalcBounds(const FTransform& LocalToWorld) const override;
+
+#if !UE_BUILD_SHIPPING
+	//~ Begin UPrimitiveComponent Interface.
+	virtual void DestroyRenderState_Concurrent() override;
+	virtual FPrimitiveSceneProxy* CreateSceneProxy() override;
+	//~ End UPrimitiveComponent Interface.
+#endif
+
+	void UpdateTests();
+	void PinLane();
+	void ClearPinnedLane();
+
+protected:
+
+#if WITH_EDITOR
+	void OnZoneGraphDataBuildDone(const struct FZoneGraphBuildData& BuildData);
+#endif
+	void OnZoneGraphDataChanged(const AZoneGraphData* ZoneGraphData);
+
+#if !UE_BUILD_SHIPPING && !UE_BUILD_TEST
+	FDebugDrawDelegateHelper DebugDrawDelegateHelper;
+#endif
+
+#if WITH_EDITOR
+	FDelegateHandle OnDataChangedHandle;
+#endif
+	FDelegateHandle OnDataAddedHandle;
+	FDelegateHandle OnDataRemovedHandle;
+
+	UPROPERTY(Transient)
+	UZoneGraphSubsystem* ZoneGraph;
+
+	UPROPERTY(Transient)
+	FZoneGraphLaneLocation LaneLocation;
+
+	UPROPERTY(Transient)
+	FZoneGraphLaneLocation GoalLaneLocation;
+
+	UPROPERTY(EditAnywhere, Category = Default);
+	FVector SearchExtent;
+
+	UPROPERTY(EditAnywhere, Category = Default);
+	float AnticipationDistance = 50.0f;
+
+	UPROPERTY(EditAnywhere, Category = Default);
+	float AgentRadius = 40.0f;
+
+	UPROPERTY(EditAnywhere, Category = Default);
+	bool bHasSpecificEndPoint = true;
+
+	UPROPERTY(EditAnywhere, Category = Default);
+	FZoneGraphTagFilter QueryFilter;
+
+	UPROPERTY(EditAnywhere, Category = Default, meta = (MakeEditWidget=true))
+	FVector GoalPosition;
+
+	FZoneGraphLaneHandle PinnedLane;
+	
+	FMassZoneGraphCachedLaneFragment CachedLane;
+	TArray<FMassZoneGraphShortPathFragment> ShortPaths;
+};
+
+/** Debug actor to visually test zone graph. */
+UCLASS(hidecategories = (Actor, Input, Collision, Rendering, Replication, Partition, HLOD, Cooking))
+class MASSAIMOVEMENTEDITOR_API AMassMovementTestingActor : public AActor
+{
+	GENERATED_BODY()
+public:
+	AMassMovementTestingActor(const FObjectInitializer& ObjectInitializer = FObjectInitializer::Get());
+
+#if WITH_EDITOR
+	virtual void PostEditMove(bool bFinished) override;
+#endif // WITH_EDITOR
+
+	UFUNCTION(BlueprintCallable, CallInEditor, Category = "Default")
+	void PinLane();
+	
+	UFUNCTION(BlueprintCallable, CallInEditor, Category = "Default")
+	void ClearPinnedLane();
+
+protected:
+	UPROPERTY(Category = Default, VisibleAnywhere, meta = (AllowPrivateAccess = "true"))
+	UMassMovementTestingComponent* DebugComp;
+};
diff --git a/Engine/Plugins/AI/MassAI/Source/MassAIMovementEditor/Public/IMassMovementEditor.h b/Engine/Plugins/AI/MassAI/Source/MassAIMovementEditor/Public/IMassMovementEditor.h
new file mode 100644
--- /dev/null
+++ b/Engine/Plugins/AI/MassAI/Source/MassAIMovementEditor/Public/IMassMovementEditor.h
@@ -0,0 +1,15 @@
+// Copyright Epic Games, Inc. All Rights Reserved.
+
+#pragma once
+
+#include "CoreMinimal.h"
+#include "Toolkits/AssetEditorToolkit.h"
+
+/** MassMovement Editor public interface */
+class MASSAIMOVEMENTEDITOR_API IMassMovementEditor : public FAssetEditorToolkit
+{
+public:
+
+};
+
+
diff --git a/Engine/Plugins/AI/MassAI/Source/MassAIMovementEditor/Public/MassMovementConfigRefDetails.cpp b/Engine/Plugins/AI/MassAI/Source/MassAIMovementEditor/Public/MassMovementConfigRefDetails.cpp
new file mode 100644
--- /dev/null
+++ b/Engine/Plugins/AI/MassAI/Source/MassAIMovementEditor/Public/MassMovementConfigRefDetails.cpp
@@ -0,0 +1,145 @@
+// Copyright Epic Games, Inc. All Rights Reserved.
+#include "MassMovementConfigRefDetails.h"
+#include "Widgets/DeclarativeSyntaxSupport.h"
+#include "DetailWidgetRow.h"
+#include "DetailLayoutBuilder.h"
+#include "IPropertyUtilities.h"
+#include "IDetailChildrenBuilder.h"
+#include "Widgets/Input/SComboButton.h"
+#include "Framework/MultiBox/MultiBoxBuilder.h"
+#include "ScopedTransaction.h"
+#include "MassMovementSettings.h"
+#include "Modules/ModuleManager.h"
+#include "ISettingsModule.h"
+#include "MassMovementPropertyUtils.h"
+
+#define LOCTEXT_NAMESPACE "MassMovementEditor"
+
+TSharedRef<IPropertyTypeCustomization> FMassMovementConfigRefDetails::MakeInstance()
+{
+	return MakeShareable(new FMassMovementConfigRefDetails);
+}
+
+void FMassMovementConfigRefDetails::CustomizeHeader(TSharedRef<class IPropertyHandle> StructPropertyHandle, class FDetailWidgetRow& HeaderRow, IPropertyTypeCustomizationUtils& StructCustomizationUtils)
+{
+	StructProperty = StructPropertyHandle;
+	PropUtils = StructCustomizationUtils.GetPropertyUtilities().Get();
+	
+	NameProperty = StructProperty->GetChildHandle(TEXT("Name"));
+	IDProperty = StructProperty->GetChildHandle(TEXT("ID"));
+
+	HeaderRow
+		.NameContent()
+		[
+			StructPropertyHandle->CreatePropertyNameWidget()
+		]
+		.ValueContent()
+		[
+			SNew(SComboButton)
+			.OnGetMenuContent(this, &FMassMovementConfigRefDetails::OnGetProfileContent)
+			.ButtonContent()
+			[
+				SNew(STextBlock)
+				.Text(this, &FMassMovementConfigRefDetails::GetCurrentProfileDesc)
+				.Font(IDetailLayoutBuilder::GetDetailFont())
+			]
+		];
+}
+
+void FMassMovementConfigRefDetails::CustomizeChildren(TSharedRef<class IPropertyHandle> StructPropertyHandle, class IDetailChildrenBuilder& StructBuilder, IPropertyTypeCustomizationUtils& StructCustomizationUtils)
+{
+}
+
+void FMassMovementConfigRefDetails::OnProfileComboChange(int32 Idx)
+{
+	const UMassMovementSettings* Settings = GetDefault<UMassMovementSettings>();
+	check(Settings);
+
+	if (Idx == -1)
+	{
+		// Goto settings to create new Profile
+		FModuleManager::LoadModuleChecked<ISettingsModule>("Settings").ShowViewer(Settings->GetContainerName(), Settings->GetCategoryName(), Settings->GetSectionName());
+		return;
+	}
+
+	TConstArrayView<FMassMovementConfig> MovementConfigs = Settings->GetMovementConfigs(); 
+	if (MovementConfigs.IsValidIndex(Idx))
+	{
+		const FMassMovementConfig& Config = MovementConfigs[Idx];
+
+		FScopedTransaction Transaction(FText::Format(LOCTEXT("SetPropertyValue", "Set {0}"), StructProperty->GetPropertyDisplayName()));
+
+		if (NameProperty)
+		{
+			NameProperty->SetValue(Config.Name, EPropertyValueSetFlags::NotTransactable);
+		}
+
+		if (IDProperty)
+		{
+			UE::MassMovement::PropertyUtils::SetValue<FGuid>(IDProperty, Config.ID, EPropertyValueSetFlags::NotTransactable);
+		}
+
+		if (PropUtils)
+		{
+			PropUtils->ForceRefresh();
+		}
+	}
+}
+
+TSharedRef<SWidget> FMassMovementConfigRefDetails::OnGetProfileContent() const
+{
+	FMenuBuilder MenuBuilder(true, NULL);
+	const UMassMovementSettings* Settings = GetDefault<UMassMovementSettings>();
+	check(Settings);
+
+	FUIAction NewItemAction(FExecuteAction::CreateSP(const_cast<FMassMovementConfigRefDetails*>(this), &FMassMovementConfigRefDetails::OnProfileComboChange, -1));
+	MenuBuilder.AddMenuEntry(LOCTEXT("CreateOrEditConfigs", "Create or Edit Movement Configs..."), TAttribute<FText>(), FSlateIcon(), NewItemAction);
+	MenuBuilder.AddMenuSeparator();
+
+	TConstArrayView<FMassMovementConfig> MovementConfigs = Settings->GetMovementConfigs();
+	for (int32 Index = 0; Index < MovementConfigs.Num(); Index++)
+	{
+		const FMassMovementConfig& Config = MovementConfigs[Index];
+		FUIAction ItemAction(FExecuteAction::CreateSP(const_cast<FMassMovementConfigRefDetails*>(this), &FMassMovementConfigRefDetails::OnProfileComboChange, Index));
+		MenuBuilder.AddMenuEntry(FText::FromName(Config.Name), TAttribute<FText>(), FSlateIcon(), ItemAction);
+	}
+	return MenuBuilder.MakeWidget();
+}
+
+FText FMassMovementConfigRefDetails::GetCurrentProfileDesc() const
+{
+	TOptional<FGuid> IDOpt = UE::MassMovement::PropertyUtils::GetValue<FGuid>(IDProperty);
+	if (IDOpt.IsSet())
+	{
+		const FGuid ID = IDOpt.GetValue();
+		if (ID.IsValid())
+		{
+			const UMassMovementSettings* Settings = GetDefault<UMassMovementSettings>();
+			check(Settings);
+
+			const FMassMovementConfig* Config = Settings->GetMovementConfigByID(ID);
+			if (Config)
+			{
+				return FText::FromName(Config->Name);
+			}
+			else
+			{
+				FName OldProfileName;
+				if (NameProperty && NameProperty->GetValue(OldProfileName) == FPropertyAccess::Success)
+				{
+					FFormatNamedArguments Args;
+					Args.Add(TEXT("Identifier"), FText::FromName(OldProfileName));
+					return FText::Format(LOCTEXT("InvalidConfig", "Invalid Config {Identifier}"), Args);
+				}
+			}
+		}
+		else
+		{
+			return LOCTEXT("Invalid", "Invalid");
+		}
+	}
+	// TODO: handle multiple values
+	return FText();
+}
+
+#undef LOCTEXT_NAMESPACE
\ No newline at end of file
diff --git a/Engine/Plugins/AI/MassAI/Source/MassAIMovementEditor/Public/MassMovementConfigRefDetails.h b/Engine/Plugins/AI/MassAI/Source/MassAIMovementEditor/Public/MassMovementConfigRefDetails.h
new file mode 100644
--- /dev/null
+++ b/Engine/Plugins/AI/MassAI/Source/MassAIMovementEditor/Public/MassMovementConfigRefDetails.h
@@ -0,0 +1,33 @@
+// Copyright Epic Games, Inc. All Rights Reserved.
+#pragma once
+#include "CoreMinimal.h"
+#include "IPropertyTypeCustomization.h"
+
+class IPropertyHandle;
+
+/**
+ * Type customization for FMassMovementConfigRef.
+ */
+class FMassMovementConfigRefDetails : public IPropertyTypeCustomization
+{
+public:
+	/** Makes a new instance of this detail layout class for a specific detail view requesting it */
+	static TSharedRef<IPropertyTypeCustomization> MakeInstance();
+
+	/** IPropertyTypeCustomization interface */
+	virtual void CustomizeHeader(TSharedRef<IPropertyHandle> StructPropertyHandle, class FDetailWidgetRow& HeaderRow, IPropertyTypeCustomizationUtils& StructCustomizationUtils) override;
+	virtual void CustomizeChildren(TSharedRef<IPropertyHandle> StructPropertyHandle, class IDetailChildrenBuilder& StructBuilder, IPropertyTypeCustomizationUtils& StructCustomizationUtils) override;
+
+private:
+
+	void OnProfileComboChange(int32 Idx);
+	TSharedRef<SWidget> OnGetProfileContent() const;
+	FText GetCurrentProfileDesc() const;
+
+	TSharedPtr<IPropertyHandle> NameProperty;
+	TSharedPtr<IPropertyHandle> IDProperty;
+	TSharedPtr<IPropertyHandle> LanesProperty;
+
+	class IPropertyUtilities* PropUtils;
+	TSharedPtr<IPropertyHandle> StructProperty;
+};
\ No newline at end of file
diff --git a/Engine/Plugins/AI/MassAI/Source/MassAIMovementEditor/Public/MassMovementEditorModule.h b/Engine/Plugins/AI/MassAI/Source/MassAIMovementEditor/Public/MassMovementEditorModule.h
new file mode 100644
--- /dev/null
+++ b/Engine/Plugins/AI/MassAI/Source/MassAIMovementEditor/Public/MassMovementEditorModule.h
@@ -0,0 +1,25 @@
+// Copyright Epic Games, Inc. All Rights Reserved.
+
+#pragma  once
+
+#include "CoreMinimal.h"
+#include "Modules/ModuleManager.h"
+#include "AssetTypeCategories.h"
+#include "Toolkits/IToolkitHost.h"
+#include "Toolkits/AssetEditorToolkit.h"
+
+
+class IMassMovementEditor;
+
+/**
+* The public interface to this module
+*/
+class MASSAIMOVEMENTEDITOR_API FMassMovementEditorModule : public IModuleInterface//, public IHasMenuExtensibility, public IHasToolBarExtensibility
+{
+public:
+	// Begin IModuleInterface
+	virtual void StartupModule() override;
+	virtual void ShutdownModule() override;
+
+private:
+};
diff --git a/Engine/Plugins/AI/MassAI/Source/MassAIReplication/Private/MassReplicationPathHandlers.cpp b/Engine/Plugins/AI/MassAI/Source/MassAIReplication/Private/MassReplicationPathHandlers.cpp
new file mode 100644
--- /dev/null
+++ b/Engine/Plugins/AI/MassAI/Source/MassAIReplication/Private/MassReplicationPathHandlers.cpp
@@ -0,0 +1,154 @@
+// Copyright Epic Games, Inc. All Rights Reserved.
+
+#include "MassReplicationPathHandlers.h"
+#include "EntityQuery.h"
+#include "MassCommonFragments.h"
+#include "MassMovementTypes.h"
+#include "MassSimulationSubsystem.h"
+#include "MassZoneGraphMovementProcessors.h"
+#include "MassZoneGraphMovementUtils.h"
+#include "ZoneGraphSubsystem.h"
+#include "GameFramework/GameStateBase.h"
+#include "VisualLogger/VisualLogger.h"
+
+void FMassReplicationProcessorPathHandler::AddRequirements(FLWComponentQuery& InQuery)
+{
+	InQuery.AddRequirement<FMassZoneGraphPathRequestFragment>(ELWComponentAccess::ReadOnly);
+	InQuery.AddRequirement<FMassMoveTargetFragment>(ELWComponentAccess::ReadOnly);
+	InQuery.AddRequirement<FMassZoneGraphLaneLocationFragment>(ELWComponentAccess::ReadOnly);
+}
+
+void FMassReplicationProcessorPathHandler::CacheComponentViews(FLWComponentSystemExecutionContext& ExecContext)
+{
+	PathRequestList = ExecContext.GetMutableComponentView<FMassZoneGraphPathRequestFragment>();
+	MoveTargetList = ExecContext.GetMutableComponentView<FMassMoveTargetFragment>();
+	LaneLocationList = ExecContext.GetMutableComponentView<FMassZoneGraphLaneLocationFragment>();
+}
+
+void FMassReplicationProcessorPathHandler::AddEntity(const int32 EntityIdx, FReplicatedAgentPathData& InOutReplicatedPathData) const
+{
+	const FMassZoneGraphPathRequestFragment& RequestFragment = PathRequestList[EntityIdx];
+	const FMassMoveTargetFragment& MoveTargetFragment = MoveTargetList[EntityIdx];
+	const FMassZoneGraphLaneLocationFragment& LaneLocationFragment = LaneLocationList[EntityIdx];
+
+	InOutReplicatedPathData = FReplicatedAgentPathData(RequestFragment, MoveTargetFragment, LaneLocationFragment);
+}
+
+FReplicatedAgentPathData::FReplicatedAgentPathData(const FMassZoneGraphPathRequestFragment& RequestFragment,
+	const FMassMoveTargetFragment& MoveTargetFragment,
+	const FMassZoneGraphLaneLocationFragment& LaneLocationFragment)
+{
+	// Move target
+	ActionID = MoveTargetFragment.GetCurrentActionID();
+	Action = MoveTargetFragment.GetCurrentAction();
+	ActionServerStartTime = MoveTargetFragment.GetCurrentActionStartTime();
+	DesiredSpeed = MoveTargetFragment.DesiredSpeed;
+
+	// Lane Location
+	LaneHandle = LaneLocationFragment.LaneHandle;
+	DistanceAlongLane = LaneLocationFragment.DistanceAlongLane;
+	LaneLength = LaneLocationFragment.LaneLength;
+
+	// Path request
+	PathRequest = RequestFragment.PathRequest;
+}
+
+void FReplicatedAgentPathData::InitEntity(const UWorld& InWorld,
+										  const FEntityView& InEntityView,
+										  FMassZoneGraphLaneLocationFragment& OutLaneLocation,
+										  FMassMoveTargetFragment& OutMoveTarget,
+										  FMassZoneGraphPathRequestFragment& OutActionRequest) const
+{
+	const FLWEntity Entity = InEntityView.GetEntity();
+
+	const UZoneGraphSubsystem* ZoneGraphSubsystem = InWorld.GetSubsystem<UZoneGraphSubsystem>();
+	const UMassSimulationSubsystem* SimulationSubsystem = InWorld.GetSubsystem<UMassSimulationSubsystem>();
+	if (ZoneGraphSubsystem == nullptr || SimulationSubsystem == nullptr)
+	{
+		UE_CVLOG(ZoneGraphSubsystem == nullptr, &InWorld, LogMassNavigation, Error, TEXT("Entity [%s] no ZoneGraphSubsystem to process request %s"),
+			*Entity.DebugGetDescription(), *PathRequest.ToString());
+		UE_CVLOG(SimulationSubsystem == nullptr, &InWorld, LogMassNavigation, Error, TEXT("Entity [%s] no MassSimulationSubsystem to process request %s"),
+			*Entity.DebugGetDescription(), *PathRequest.ToString());
+		return;
+	}
+
+	UE_VLOG(SimulationSubsystem, LogMassNavigation, Log, TEXT("%s InitEntity"), *Entity.DebugGetDescription());
+
+	// Setup initial lane location
+	OutLaneLocation.LaneHandle = LaneHandle;
+	OutLaneLocation.DistanceAlongLane = DistanceAlongLane;
+	OutLaneLocation.LaneLength = LaneLength;
+
+	// Setup initial move target
+	FZoneGraphLaneLocation LaneLocation;
+	ZoneGraphSubsystem->CalculateLocationAlongLane(LaneHandle, DistanceAlongLane, LaneLocation);
+	OutMoveTarget.DesiredSpeed = DesiredSpeed;
+	OutMoveTarget.Center = LaneLocation.Position;
+	OutMoveTarget.Forward = LaneLocation.Tangent;
+	OutMoveTarget.DistanceToGoal = 0.0f;
+	OutMoveTarget.SlackRadius = 0.0f;
+
+	// Setup initial action request
+	OutActionRequest.PathRequest = PathRequest;
+
+	ApplyToEntity(InWorld, InEntityView);
+}
+
+void FReplicatedAgentPathData::ApplyToEntity(const UWorld& InWorld, const FEntityView& InEntityView) const
+{
+	const FLWEntity Entity = InEntityView.GetEntity();
+	FMassMoveTargetFragment& MoveTarget = InEntityView.GetComponentData<FMassMoveTargetFragment>();
+	if (MoveTarget.GetCurrentActionID() == ActionID)
+	{
+		return;
+	}
+
+	const UZoneGraphSubsystem* ZoneGraphSubsystem = InWorld.GetSubsystem<UZoneGraphSubsystem>();
+	const UMassSimulationSubsystem* SimulationSubsystem = InWorld.GetSubsystem<UMassSimulationSubsystem>();
+	if (ZoneGraphSubsystem == nullptr || SimulationSubsystem == nullptr)
+	{
+		UE_CVLOG(ZoneGraphSubsystem == nullptr, &InWorld, LogMassNavigation, Error, TEXT("Entity [%s] no ZoneGraphSubsystem to process request %s"),
+			*Entity.DebugGetDescription(), *PathRequest.ToString());
+		UE_CVLOG(SimulationSubsystem == nullptr, &InWorld, LogMassNavigation, Error, TEXT("Entity [%s] no MassSimulationSubsystem to process request %s"),
+			*Entity.DebugGetDescription(), *PathRequest.ToString());
+		return;
+	}
+
+	UE_VLOG(SimulationSubsystem, LogMassNavigation, Log, TEXT("Entity [%s] apply replicated data to entity"), *Entity.DebugGetDescription());
+
+	FMassZoneGraphShortPathFragment& ShortPath = InEntityView.GetComponentData<FMassZoneGraphShortPathFragment>();
+	FMassZoneGraphCachedLaneFragment& CachedLane = InEntityView.GetComponentData<FMassZoneGraphCachedLaneFragment>();
+	FMassZoneGraphLaneLocationFragment& LaneLocation = InEntityView.GetComponentData<FMassZoneGraphLaneLocationFragment>();
+	const FDataFragment_AgentRadius& AgentRadius = InEntityView.GetComponentData<FDataFragment_AgentRadius>();
+
+	MoveTarget.CreateReplicatedAction(Action, ActionID, InWorld.GetTimeSeconds(), ActionServerStartTime);
+	MoveTarget.DesiredSpeed = DesiredSpeed;
+
+	// Force current lane to build same path as server
+	if (LaneLocation.LaneHandle != LaneHandle)
+	{
+		UE_VLOG(SimulationSubsystem, LogMassNavigation, Verbose, TEXT("Entity [%s] Force lane location from %s - %.1f to %s - %.1f to build similar path."),
+			*Entity.DebugGetDescription(),
+			*LaneLocation.LaneHandle.ToString(), LaneLocation.DistanceAlongLane,
+			*LaneHandle.ToString(),	DistanceAlongLane);
+
+		LaneLocation.LaneHandle = LaneHandle;
+		LaneLocation.DistanceAlongLane = DistanceAlongLane;
+		LaneLocation.LaneLength = LaneLength;
+	}
+
+	switch (Action)
+	{
+	case EMassMovementAction::Stand:
+		UE::MassMovement::ActivateActionStand(InWorld, SimulationSubsystem, Entity, *ZoneGraphSubsystem, LaneLocation, DesiredSpeed.Get(), MoveTarget, ShortPath, CachedLane);
+		break;
+	case EMassMovementAction::Move:
+		UE::MassMovement::ActivateActionMove(InWorld, SimulationSubsystem, Entity, *ZoneGraphSubsystem, LaneLocation, PathRequest, AgentRadius.Radius, DesiredSpeed.Get(), MoveTarget, ShortPath, CachedLane);
+		break;
+	case EMassMovementAction::Animate:
+		UE::MassMovement::ActivateActionAnimate(InWorld, SimulationSubsystem, Entity, MoveTarget);
+		break;
+	default:
+		ensureMsgf(false, TEXT("Unhandled action type: %s"), *UEnum::GetValueAsString(Action));
+	}
+}
