AggregateTicking and Pipe merged into Engine/Plugins/Runtime/MassEntity

diff --git a/Engine/Plugins/AI/MassAI/Source/MassAIReplication/Public/MassReplicationPathHandlers.h b/Engine/Plugins/AI/MassAI/Source/MassAIReplication/Public/MassReplicationPathHandlers.h
new file mode 100644
--- /dev/null
+++ b/Engine/Plugins/AI/MassAI/Source/MassAIReplication/Public/MassReplicationPathHandlers.h
@@ -0,0 +1,238 @@
+// Copyright Epic Games, Inc. All Rights Reserved.
+
+#pragma once
+
+#include "MassReplicationTypes.h"
+#include "EntityView.h"
+#include "MassClientBubbleHandler.h"
+#include "MassAIMovementFragments.h"
+#include "MassZoneGraphMovementFragments.h"
+
+#include "MassReplicationPathHandlers.generated.h"
+
+struct FLWComponentQuery;
+
+//////////////////////////////////////////////////////////////////////////// FReplicatedAgentPathData ////////////////////////////////////////////////////////////////////////////
+/**
+ * To replicate path following make a member variable of this class in your FReplicatedAgentBase derived class. In the FReplicatedAgentBase derived class you must also provide an accessor function
+ * FReplicatedAgentPathData& GetReplicatedPathDataMutable().
+ */
+USTRUCT()
+struct MASSAIREPLICATION_API FReplicatedAgentPathData
+{
+	GENERATED_BODY()
+
+	friend class FMassClientBubblePathHandlerBase;
+
+	FReplicatedAgentPathData() = default;
+	explicit FReplicatedAgentPathData(const FMassZoneGraphPathRequestFragment& RequestFragment,
+		const FMassMoveTargetFragment& MoveTargetFragment,
+		const FMassZoneGraphLaneLocationFragment& LaneLocationFragment);
+
+	void InitEntity(const UWorld& InWorld,
+					const FEntityView& InEntityView,
+					FMassZoneGraphLaneLocationFragment& OutLaneLocation,
+					FMassMoveTargetFragment& OutMoveTarget,
+					FMassZoneGraphPathRequestFragment& OutActionRequest) const;
+
+	void ApplyToEntity(const UWorld& InWorld, const FEntityView& InEntityView) const;
+
+	UPROPERTY(Transient)
+	mutable FZoneGraphShortPathRequest PathRequest;
+
+	/** Handle to current lane. */
+	UPROPERTY(Transient)
+	FZoneGraphLaneHandle LaneHandle;
+
+	/** Server time in seconds when the action started. */
+	UPROPERTY(Transient)
+	float ActionServerStartTime = 0.0f;
+
+	/** Distance along current lane. */
+	UPROPERTY(Transient)
+	float DistanceAlongLane = 0.0f;
+
+	/** Cached lane length, used for clamping and testing if at end of lane. */
+	UPROPERTY(Transient)
+	float LaneLength = 0.0f;
+
+	/** Requested movement speed. */
+	UPROPERTY(Transient)
+	FMassInt16Real DesiredSpeed = FMassInt16Real(0.0f);
+
+	UPROPERTY(Transient)
+	uint16 ActionID = 0;
+
+	/** Current movement action. */
+	UPROPERTY(Transient)
+	EMassMovementAction Action = EMassMovementAction::Move;
+};
+
+//////////////////////////////////////////////////////////////////////////// TMassClientBubblePathHandler ////////////////////////////////////////////////////////////////////////////
+/**
+ * To replicate path following make a member variable of this class in your TClientBubbleHandlerBase derived class. This class is a friend of TMassClientBubblePathHandler.
+ * It is meant to have access to the protected data declared there.
+ */
+template<typename AgentArrayItem>
+class TMassClientBubblePathHandler
+{
+public:
+	TMassClientBubblePathHandler(TClientBubbleHandlerBase<AgentArrayItem>& InOwnerHandler)
+		: OwnerHandler(InOwnerHandler)
+	{}
+
+#if UE_REPLICATION_COMPILE_SERVER_CODE
+	/** Sets the path data in the client bubble on the server */
+	void SetBubblePathData(const FMassReplicatedAgentHandle Handle,
+		const FMassZoneGraphPathRequestFragment& PathRequest,
+		const FMassMoveTargetFragment& MoveTargetFragment,
+		const FMassZoneGraphLaneLocationFragment& LaneLocationFragment);
+#endif // UE_REPLICATION_COMPILE_SERVER_CODE
+
+#if UE_REPLICATION_COMPILE_CLIENT_CODE
+	/**
+	 * When entities are spawned in Mass by the replication system on the client, a spawn query is used to set the data on the spawned entities.
+	 * The following functions are used to configure the query and then set that data for path following.
+	 */
+	static void AddRequirementsForSpawnQuery(FLWComponentQuery& InQuery);
+	void CacheComponentViewsForSpawnQuery(FLWComponentSystemExecutionContext& InExecContext);
+	void ClearComponentViewsForSpawnQuery();
+
+	void SetSpawnedEntityData(const FEntityView& EntityView, const FReplicatedAgentPathData& ReplicatedPathData, const int32 EntityIdx) const;
+
+	/** Call this when an Entity that has already been spawned is modified on the client */
+	void SetModifiedEntityData(const FEntityView& EntityView, const FReplicatedAgentPathData& ReplicatedPathData) const;
+#endif // UE_REPLICATION_COMPILE_CLIENT_CODE
+
+protected:
+	TArrayView<FMassZoneGraphPathRequestFragment> PathRequestList;
+	TArrayView<FMassMoveTargetFragment> MoveTargetList;
+	TArrayView<FMassZoneGraphLaneLocationFragment> LaneLocationList;
+
+	TClientBubbleHandlerBase<AgentArrayItem>& OwnerHandler;
+};
+
+#if UE_REPLICATION_COMPILE_SERVER_CODE
+template<typename AgentArrayItem>
+void TMassClientBubblePathHandler<AgentArrayItem>::SetBubblePathData(const FMassReplicatedAgentHandle Handle,
+	const FMassZoneGraphPathRequestFragment& PathRequestFragment,
+	const FMassMoveTargetFragment& MoveTargetFragment,
+	const FMassZoneGraphLaneLocationFragment& LaneLocationFragment)
+{
+	check(OwnerHandler.AgentHandleManager.IsValidHandle(Handle));
+
+	const int32 AgentsIdx = OwnerHandler.AgentLookupArray[Handle.GetIndex()].AgentsIdx;
+	AgentArrayItem& Item = (*OwnerHandler.Agents)[AgentsIdx];
+
+	checkf(Item.Agent.GetNetID().IsValid(), TEXT("Pos should not be updated on FCrowdFastArrayItem's that have an Invalid ID! First Add the Agent!"));
+
+	// GetReplicatedPathDataMutable() must be defined in your FReplicatedAgentBase derived class
+	FReplicatedAgentPathData& ReplicatedPath = Item.Agent.GetReplicatedPathDataMutable();
+
+	if (ReplicatedPath.ActionID != MoveTargetFragment.GetCurrentActionID())
+	{
+		ReplicatedPath = FReplicatedAgentPathData(PathRequestFragment, MoveTargetFragment, LaneLocationFragment);
+
+		OwnerHandler.Serializer->MarkItemDirty(Item);
+	}
+}
+#endif //UE_REPLICATION_COMPILE_SERVER_CODE
+
+#if UE_REPLICATION_COMPILE_CLIENT_CODE
+template<typename AgentArrayItem>
+void TMassClientBubblePathHandler<AgentArrayItem>::AddRequirementsForSpawnQuery(FLWComponentQuery& InQuery)
+{
+	InQuery.AddRequirement<FMassZoneGraphPathRequestFragment>(ELWComponentAccess::ReadWrite);
+	InQuery.AddRequirement<FMassMoveTargetFragment>(ELWComponentAccess::ReadWrite);
+	InQuery.AddRequirement<FMassZoneGraphLaneLocationFragment>(ELWComponentAccess::ReadWrite);
+}
+#endif // UE_REPLICATION_COMPILE_CLIENT_CODE
+
+#if UE_REPLICATION_COMPILE_CLIENT_CODE
+template<typename AgentArrayItem>
+void TMassClientBubblePathHandler<AgentArrayItem>::CacheComponentViewsForSpawnQuery(FLWComponentSystemExecutionContext& InExecContext)
+{
+	PathRequestList = InExecContext.GetMutableComponentView<FMassZoneGraphPathRequestFragment>();
+	MoveTargetList = InExecContext.GetMutableComponentView<FMassMoveTargetFragment>();
+	LaneLocationList = InExecContext.GetMutableComponentView<FMassZoneGraphLaneLocationFragment>();
+}
+#endif // UE_REPLICATION_COMPILE_CLIENT_CODE
+
+#if UE_REPLICATION_COMPILE_CLIENT_CODE
+template<typename AgentArrayItem>
+void TMassClientBubblePathHandler<AgentArrayItem>::ClearComponentViewsForSpawnQuery()
+{
+	LaneLocationList = TArrayView<FMassZoneGraphLaneLocationFragment>();
+	MoveTargetList = TArrayView<FMassMoveTargetFragment>();
+	PathRequestList = TArrayView<FMassZoneGraphPathRequestFragment>();
+}
+#endif // UE_REPLICATION_COMPILE_CLIENT_CODE
+
+#if UE_REPLICATION_COMPILE_CLIENT_CODE
+template<typename AgentArrayItem>
+void TMassClientBubblePathHandler<AgentArrayItem>::SetSpawnedEntityData(const FEntityView& EntityView, const FReplicatedAgentPathData& ReplicatedPathData, const int32 EntityIdx) const
+{
+	UWorld* World = OwnerHandler.Serializer->GetWorld();
+	check(World);
+	ReplicatedPathData.InitEntity(*World, EntityView, LaneLocationList[EntityIdx], MoveTargetList[EntityIdx], PathRequestList[EntityIdx]);
+}
+#endif // UE_REPLICATION_COMPILE_CLIENT_CODE
+
+#if UE_REPLICATION_COMPILE_CLIENT_CODE
+template<typename AgentArrayItem>
+void TMassClientBubblePathHandler<AgentArrayItem>::SetModifiedEntityData(const FEntityView& EntityView, const FReplicatedAgentPathData& ReplicatedPathData) const
+{
+	UWorld* World = OwnerHandler.Serializer->GetWorld();
+	check(World);
+	ReplicatedPathData.ApplyToEntity(*World, EntityView);
+}
+#endif // UE_REPLICATION_COMPILE_CLIENT_CODE
+
+//////////////////////////////////////////////////////////////////////////// FMassReplicationProcessorPathHandler ////////////////////////////////////////////////////////////////////////////
+/**
+ * Used to replicate path following by your UMassReplicationProcessorBase derived class. This class should only get used on the server.
+ * @todo add #if UE_REPLICATION_COMPILE_SERVER_CODE
+ */
+class MASSAIREPLICATION_API FMassReplicationProcessorPathHandler
+{
+public:
+	/** Adds the requirements for the path following to the query. */
+	static void AddRequirements(FLWComponentQuery& InQuery);
+
+	/** Cache any component views you want to, this will get called before we iterate through entities. */
+	void CacheComponentViews(FLWComponentSystemExecutionContext& ExecContext);
+
+	/**
+	 * Set the replicated path data when we are adding an entity to the client bubble.
+	 * @param EntityIdx the index of the entity in component views that have been cached.
+	 * @param InOUtReplicatedPathData the data to set.
+	 */
+	void AddEntity(const int32 EntityIdx, FReplicatedAgentPathData& InOUtReplicatedPathData) const;
+
+	/**
+	 * Set the replicated path data when we are modifying an entity that already exists in the client bubble.
+	 * @param Handle to the agent in the TMassClientBubbleHandler (that TMassClientBubblePathHandler is a member variable of).
+	 * @param EntityIdx the index of the entity in component views that have been cached.
+	 * @param BubblePathHandler handler to actually set the data in the client bubble
+	 */
+	template<typename AgentArrayItem>
+	void ModifyEntity(const FMassReplicatedAgentHandle Handle, const int32 EntityIdx, TMassClientBubblePathHandler<AgentArrayItem>& BubblePathHandler);
+
+	TArrayView<FMassZoneGraphPathRequestFragment> PathRequestList;
+	TArrayView<FMassMoveTargetFragment> MoveTargetList;
+	TArrayView<FMassZoneGraphLaneLocationFragment> LaneLocationList;
+};
+
+template<typename AgentArrayItem>
+void FMassReplicationProcessorPathHandler::ModifyEntity(const FMassReplicatedAgentHandle Handle, const int32 EntityIdx, TMassClientBubblePathHandler<AgentArrayItem>& BubblePathHandler)
+{
+	const FMassZoneGraphPathRequestFragment& PathRequest = PathRequestList[EntityIdx];
+	FMassMoveTargetFragment& MoveTargetFragment = MoveTargetList[EntityIdx];
+	const FMassZoneGraphLaneLocationFragment& LaneLocationFragment = LaneLocationList[EntityIdx];
+
+	if (MoveTargetFragment.GetNetDirty())
+	{
+		BubblePathHandler.SetBubblePathData(Handle, PathRequest, MoveTargetFragment, LaneLocationFragment);
+		MoveTargetFragment.ResetNetDirty();
+	}
+}
diff --git a/Engine/Plugins/AI/MassCrowd/Config/BaseMassCrowd.ini b/Engine/Plugins/AI/MassCrowd/Config/BaseMassCrowd.ini
new file mode 100644
--- /dev/null
+++ b/Engine/Plugins/AI/MassCrowd/Config/BaseMassCrowd.ini
@@ -0,0 +1,3 @@
+[CoreRedirects]
++PackageRedirects=(OldName="/Script/MassCrowdAI5", NewName="/Script/MassCrowd")
++StructRedirects=(OldName="/Script/MassAIBehavior.MassZoneGraphFindWanderTarget",NewName="/Script/MassCrowd.MassZoneGraphFindWanderTarget")
\ No newline at end of file
diff --git a/Engine/Plugins/AI/MassCrowd/Source/MassCrowd/Private/MassCrowdBubble.cpp b/Engine/Plugins/AI/MassCrowd/Source/MassCrowd/Private/MassCrowdBubble.cpp
new file mode 100644
--- /dev/null
+++ b/Engine/Plugins/AI/MassCrowd/Source/MassCrowd/Private/MassCrowdBubble.cpp
@@ -0,0 +1,192 @@
+// Copyright Epic Games, Inc. All Rights Reserved.
+
+#include "MassCrowdBubble.h"
+#include "Net/UnrealNetwork.h"
+#include "DrawDebugHelpers.h"
+#include "MassSpawnerTypes.h"
+#include "EntityView.h"
+#include "MassCommonFragments.h"
+#include "MassMovementFragments.h"
+
+namespace UE::Mass::Crowd
+{
+	bool bDebugReplicationPositions = false;
+	FAutoConsoleVariableRef CVarbDebugReplication(TEXT("ai.debug.CrowdReplicationPositions"), bDebugReplicationPositions, TEXT("Crowd debug replication positions on server and client"), ECVF_Cheat);
+
+#if WITH_MASS_DEBUG && WITH_EDITOR
+	// @todo provide a better way of selecting agents to debug
+	constexpr int32 MaxAgentsDraw = 300;
+
+	void DebugDrawReplicatedAgent(FLWEntity Entity, const UPipeEntitySubsystem& EntitySystem)
+	{
+		static const FVector DebugCylinderHeight = FVector(0.f, 0.f, 200.f);
+		static constexpr float DebugCylinderRadius = 50.f;
+
+		const FEntityView EntityView(EntitySystem, Entity);
+
+		const FDataFragment_Transform& TransformFragment = EntityView.GetComponentData<FDataFragment_Transform>();
+		const FMassNetworkIDFragment& NetworkIDFragment = EntityView.GetComponentData<FMassNetworkIDFragment>();
+
+		const FVector& Pos = TransformFragment.GetTransform().GetLocation();
+		const uint32 NetworkID = NetworkIDFragment.NetID.GetValue();
+
+		// Multiply by a largeish number that is not a multiple of 256 to separate out the color shades a bit
+		const uint32 InitialColor = NetworkID * 20001;
+
+		const uint8 NetworkIdMod3 = NetworkID % 3;
+		FColor DebugCylinderColor;
+
+		// Make a deterministic color from the Mod by 3 to vary how we create it
+		if (NetworkIdMod3 == 0)
+		{
+			DebugCylinderColor = FColor(InitialColor % 256, InitialColor / 256 % 256, InitialColor / 256 / 256 % 256);
+		}
+		else if (NetworkIdMod3 == 1)
+		{
+			DebugCylinderColor = FColor(InitialColor / 256 / 256 % 256, InitialColor % 256, InitialColor / 256 % 256);
+		}
+		else
+		{
+			DebugCylinderColor = FColor(InitialColor / 256 % 256, InitialColor / 256 / 256 % 256, InitialColor % 256);
+		}
+
+		const UWorld* World = EntitySystem.GetWorld();
+		if (World != nullptr && World->GetNetMode() == NM_Client)
+		{
+			DrawDebugCylinder(World, Pos, Pos + 0.5f * DebugCylinderHeight, DebugCylinderRadius, /*segments = */24, DebugCylinderColor);
+		}
+		else
+		{
+			DrawDebugCylinder(World, Pos + 0.5f * DebugCylinderHeight, Pos + DebugCylinderHeight, DebugCylinderRadius, /*segments = */24, DebugCylinderColor);
+		}
+	}
+#endif // WITH_MASS_DEBUG && WITH_EDITOR
+}
+
+#if WITH_MASS_DEBUG && WITH_EDITOR
+void FMassCrowdClientBubbleHandler::DebugValidateBubbleOnServer()
+{
+	Super::DebugValidateBubbleOnServer();
+
+	if (UE::Mass::Crowd::bDebugReplicationPositions)
+	{
+		UPipeEntitySubsystem* EntitySystem = Serializer->GetEntitySystem();
+		check(EntitySystem);
+
+		// @todo cap at MaxAgentsDraw for now
+		const int32 MaxAgentsDraw = FMath::Min(UE::Mass::Crowd::MaxAgentsDraw, (*Agents).Num());
+
+		for (int32 Idx = 0; Idx < MaxAgentsDraw; ++Idx)
+		{
+			const FCrowdFastArrayItem& CrowdItem = (*Agents)[Idx];
+
+			const FMassAgentLookupData& LookupData = AgentLookupArray[CrowdItem.GetHandle().GetIndex()];
+
+			check(LookupData.Entity.IsSet());
+
+			UE::Mass::Crowd::DebugDrawReplicatedAgent(LookupData.Entity, *EntitySystem);
+		}
+	}
+}
+#endif // WITH_MASS_DEBUG && WITH_EDITOR
+
+#if WITH_MASS_DEBUG && WITH_EDITOR
+void FMassCrowdClientBubbleHandler::DebugValidateBubbleOnClient()
+{
+	Super::DebugValidateBubbleOnClient();
+
+	if (UE::Mass::Crowd::bDebugReplicationPositions)
+	{
+		UPipeEntitySubsystem* EntitySystem = Serializer->GetEntitySystem();
+		check(EntitySystem);
+
+		UMassReplicationManager* ReplicationManager = Serializer->GetReplicationManager();
+		check(ReplicationManager);
+
+		// @todo cap at MaxAgentsDraw for now
+		const int32 MaxAgentsDraw = FMath::Min(UE::Mass::Crowd::MaxAgentsDraw, (*Agents).Num());
+
+		for (int32 Idx = 0; Idx < MaxAgentsDraw; ++Idx)
+		{
+			const FCrowdFastArrayItem& CrowdItem = (*Agents)[Idx];
+
+			const FMassReplicationEntityInfo* EntityInfo = ReplicationManager->FindMassEntityInfo(CrowdItem.Agent.GetNetID());
+
+			check(EntityInfo->Entity.IsSet());
+
+			UE::Mass::Crowd::DebugDrawReplicatedAgent(EntityInfo->Entity, *EntitySystem);
+		}
+	}
+}
+#endif // WITH_MASS_DEBUG && WITH_EDITOR
+
+#if UE_REPLICATION_COMPILE_CLIENT_CODE
+void FMassCrowdClientBubbleHandler::PostReplicatedAdd(const TArrayView<int32> AddedIndices, int32 FinalSize)
+{
+	auto AddRequirementsForSpawnQuery = [this](FLWComponentQuery& InQuery)
+	{
+		PathHandler.AddRequirementsForSpawnQuery(InQuery);
+		TransformHandler.AddRequirementsForSpawnQuery(InQuery);
+	};
+
+	auto CacheComponentViewsForSpawnQuery = [this](FLWComponentSystemExecutionContext& InExecContext)
+	{
+		PathHandler.CacheComponentViewsForSpawnQuery(InExecContext);
+		TransformHandler.CacheComponentViewsForSpawnQuery(InExecContext);
+	};
+
+	auto SetSpawnedEntityData = [this](const FEntityView& EntityView, const FReplicatedCrowdAgent& ReplicatedEntity, const int32 EntityIdx)
+	{
+		PathHandler.SetSpawnedEntityData(EntityView, ReplicatedEntity.GetReplicatedPathData(), EntityIdx);
+		TransformHandler.SetSpawnedEntityData(EntityIdx, ReplicatedEntity.GetReplicatedPositionYawData());
+	};
+
+	auto SetModifiedEntityData = [this](const FEntityView& EntityView, const FReplicatedCrowdAgent& Item)
+	{
+		PostReplicatedChangeEntity(EntityView, Item);
+	};
+
+	PostReplicatedAddHelper(AddedIndices, AddRequirementsForSpawnQuery, CacheComponentViewsForSpawnQuery, SetSpawnedEntityData, SetModifiedEntityData);
+
+	PathHandler.ClearComponentViewsForSpawnQuery();
+	TransformHandler.ClearComponentViewsForSpawnQuery();
+}
+#endif //UE_REPLICATION_COMPILE_SERVER_CODE
+
+#if UE_REPLICATION_COMPILE_CLIENT_CODE
+void FMassCrowdClientBubbleHandler::FMassCrowdClientBubbleHandler::PostReplicatedChange(const TArrayView<int32> ChangedIndices, int32 FinalSize)
+{
+	auto SetModifiedEntityData = [this](const FEntityView& EntityView, const FReplicatedCrowdAgent& Item)
+	{
+		PostReplicatedChangeEntity(EntityView, Item);
+	};
+
+	PostReplicatedChangeHelper(ChangedIndices, SetModifiedEntityData);
+}
+#endif //UE_REPLICATION_COMPILE_SERVER_CODE
+
+#if UE_REPLICATION_COMPILE_CLIENT_CODE
+void FMassCrowdClientBubbleHandler::PostReplicatedChangeEntity(const FEntityView& EntityView, const FReplicatedCrowdAgent& Item) const
+{
+	PathHandler.SetModifiedEntityData(EntityView, Item.GetReplicatedPathData());
+
+	// No need to call TransformHandler as that only gets replicated when an agent is added to the bubble
+}
+#endif // UE_REPLICATION_COMPILE_CLIENT_CODE
+
+AMassCrowdClientBubbleInfo::AMassCrowdClientBubbleInfo(const FObjectInitializer& ObjectInitializer)
+	: Super(ObjectInitializer)
+{
+	Serializers.Add(&CrowdSerializer);
+}
+
+void AMassCrowdClientBubbleInfo::GetLifetimeReplicatedProps(TArray<FLifetimeProperty>& OutLifetimeProps) const
+{
+	Super::GetLifetimeReplicatedProps(OutLifetimeProps);
+
+	FDoRepLifetimeParams SharedParams;
+	SharedParams.bIsPushBased = true;
+
+	// Technically, this doesn't need to be PushModel based because it's a FastArray and they ignore it.
+	DOREPLIFETIME_WITH_PARAMS_FAST(AMassCrowdClientBubbleInfo, CrowdSerializer, SharedParams);
+}
diff --git a/Engine/Plugins/AI/MassCrowd/Source/MassCrowd/Private/MassCrowdMemberTrait.cpp b/Engine/Plugins/AI/MassCrowd/Source/MassCrowd/Private/MassCrowdMemberTrait.cpp
new file mode 100644
--- /dev/null
+++ b/Engine/Plugins/AI/MassCrowd/Source/MassCrowd/Private/MassCrowdMemberTrait.cpp
@@ -0,0 +1,12 @@
+﻿// Copyright Epic Games, Inc. All Rights Reserved.
+
+#include "MassCrowdMemberTrait.h"
+#include "MassCrowdFragments.h"
+#include "MassEntityTemplateRegistry.h"
+#include "MassCrowdNavigationProcessor.h"
+
+void UMassCrowdMemberTrait::BuildTemplate(FMassEntityTemplateBuildContext& BuildContext, UWorld& World) const
+{
+	BuildContext.AddTag<FTagFragment_MassCrowd>();
+	BuildContext.AddFragmentWithDefaultInitializer<FMassCrowdLaneTrackingFragment>();	
+}
diff --git a/Engine/Plugins/AI/MassCrowd/Source/MassCrowd/Private/MassCrowdNavigationProcessor.cpp b/Engine/Plugins/AI/MassCrowd/Source/MassCrowd/Private/MassCrowdNavigationProcessor.cpp
new file mode 100644
--- /dev/null
+++ b/Engine/Plugins/AI/MassCrowd/Source/MassCrowd/Private/MassCrowdNavigationProcessor.cpp
@@ -0,0 +1,182 @@
+// Copyright Epic Games, Inc. All Rights Reserved.
+
+#include "MassCrowdNavigationProcessor.h"
+#include "MassAIBehaviorTypes.h"
+#include "MassCrowdSubsystem.h"
+#include "MassCrowdFragments.h"
+#include "MassCrowdSettings.h"
+#include "MassAIMovementFragments.h"
+#include "MassZoneGraphLaneObstacleAnnotations.h"
+#include "ZoneGraphSubsystem.h"
+#include "ZoneGraphAnnotationSubsystem.h"
+#include "MassZoneGraphMovementFragments.h"
+#include "Annotations/ZoneGraphDisturbanceAnnotation.h"
+#include "VisualLogger/VisualLogger.h"
+
+//----------------------------------------------------------------------//
+// UMassCrowdLaneTrackingSignalProcessor
+//----------------------------------------------------------------------//
+UMassCrowdLaneTrackingSignalProcessor::UMassCrowdLaneTrackingSignalProcessor()
+{
+	ExecutionOrder.ExecuteBefore.Add(UE::Mass::ProcessorGroupNames::Behavior);
+}
+
+void UMassCrowdLaneTrackingSignalProcessor::ConfigureQueries()
+{
+	EntityQuery.AddTagRequirement<FTagFragment_MassCrowd>(ELWComponentPresence::All);
+	EntityQuery.AddRequirement<FMassCrowdLaneTrackingFragment>(ELWComponentAccess::ReadWrite);
+	EntityQuery.AddRequirement<FMassZoneGraphLaneLocationFragment>(ELWComponentAccess::ReadOnly);
+}
+
+void UMassCrowdLaneTrackingSignalProcessor::Initialize(UObject& Owner)
+{
+	Super::Initialize(Owner);
+	
+	MassCrowdSubsystem = UWorld::GetSubsystem<UMassCrowdSubsystem>(Owner.GetWorld());
+	checkf(MassCrowdSubsystem != nullptr, TEXT("UMassCrowdSubsystem is mandatory when using MassCrowd processor."));
+
+	SubscribeToSignal(UE::Mass::Signals::CurrentLaneChanged);
+}
+
+void UMassCrowdLaneTrackingSignalProcessor::SignalEntities(UEntitySubsystem& EntitySubsystem, FLWComponentSystemExecutionContext& Context, FMassSignalNameLookup& EntitySignals)
+{
+	if (!MassCrowdSubsystem)
+	{
+		return;
+	}
+
+	EntityQuery.ForEachEntityChunk(EntitySubsystem, Context, [this](FLWComponentSystemExecutionContext& Context)
+	{
+		const int32 NumEntities = Context.GetEntitiesNum();
+		const TConstArrayView<FMassZoneGraphLaneLocationFragment> LaneLocationList = Context.GetComponentView<FMassZoneGraphLaneLocationFragment>();
+		const TArrayView<FMassCrowdLaneTrackingFragment> LaneTrackingList = Context.GetMutableComponentView<FMassCrowdLaneTrackingFragment>();
+
+		for (int32 EntityIndex = 0; EntityIndex < NumEntities; ++EntityIndex)
+		{
+			const FMassZoneGraphLaneLocationFragment& LaneLocation = LaneLocationList[EntityIndex];
+			FMassCrowdLaneTrackingFragment& LaneTracking = LaneTrackingList[EntityIndex];
+			if (LaneTracking.TrackedLaneHandle != LaneLocation.LaneHandle)
+			{
+				MassCrowdSubsystem->OnEntityLaneChanged(Context.GetEntity(EntityIndex), LaneTracking.TrackedLaneHandle, LaneLocation.LaneHandle);
+				LaneTracking.TrackedLaneHandle = LaneLocation.LaneHandle;
+			}
+		}
+	});
+}
+
+//----------------------------------------------------------------------//
+// UMassCrowdLaneTrackingDestructor
+//----------------------------------------------------------------------//
+UMassCrowdLaneTrackingDestructor::UMassCrowdLaneTrackingDestructor()
+{
+	ExecutionFlags = (int32)(EProcessorExecutionFlags::Standalone | EProcessorExecutionFlags::Server);
+	FragmentType = FMassCrowdLaneTrackingFragment::StaticStruct();
+}
+
+void UMassCrowdLaneTrackingDestructor::Initialize(UObject& Owner)
+{
+	Super::Initialize(Owner);
+	MassCrowdSubsystem = UWorld::GetSubsystem<UMassCrowdSubsystem>(Owner.GetWorld());
+	checkf(MassCrowdSubsystem != nullptr, TEXT("UMassCrowdSubsystem is mandatory when using MassCrowd processor."));
+}
+
+void UMassCrowdLaneTrackingDestructor::ConfigureQueries()
+{
+	EntityQuery.AddTagRequirement<FTagFragment_MassCrowd>(ELWComponentPresence::All);
+	EntityQuery.AddRequirement<FMassCrowdLaneTrackingFragment>(ELWComponentAccess::ReadOnly);
+}
+
+void UMassCrowdLaneTrackingDestructor::Execute(UEntitySubsystem& EntitySubsystem, FLWComponentSystemExecutionContext& Context)
+{
+	EntityQuery.ForEachEntityChunk(EntitySubsystem, Context, [this](const FLWComponentSystemExecutionContext& Context)
+	{
+		const int32 NumEntities = Context.GetEntitiesNum();
+		const TConstArrayView<FMassCrowdLaneTrackingFragment> LaneTrackingList = Context.GetComponentView<FMassCrowdLaneTrackingFragment>();
+
+		for (int32 EntityIndex = 0; EntityIndex < NumEntities; ++EntityIndex)
+		{
+			const FMassCrowdLaneTrackingFragment& LaneTracking = LaneTrackingList[EntityIndex];
+			if (LaneTracking.TrackedLaneHandle.IsValid())
+			{
+				MassCrowdSubsystem->OnEntityLaneChanged(Context.GetEntity(EntityIndex), LaneTracking.TrackedLaneHandle, FZoneGraphLaneHandle());
+			}
+
+		}
+	});
+}
+
+//----------------------------------------------------------------------//
+// UMassCrowdDynamicObstacleProcessor
+//----------------------------------------------------------------------//
+UMassCrowdDynamicObstacleProcessor::UMassCrowdDynamicObstacleProcessor()
+{
+	bAutoRegisterWithProcessingPhases = true;
+
+	ExecutionOrder.ExecuteBefore.Add(UE::Mass::ProcessorGroupNames::UpdateAnnotationTags);
+}
+
+void UMassCrowdDynamicObstacleProcessor::Initialize(UObject& Owner)
+{
+	Super::Initialize(Owner);
+	
+	ZoneGraphSubsystem = UWorld::GetSubsystem<UZoneGraphSubsystem>(Owner.GetWorld());
+	ZoneGraphAnnotationSubsystem = UWorld::GetSubsystem<UZoneGraphAnnotationSubsystem>(Owner.GetWorld());
+	
+	CrowdSettings = GetDefault<UMassCrowdSettings>();
+	checkf(CrowdSettings, TEXT("Settings default object is always expected to be valid"));
+}
+
+void UMassCrowdDynamicObstacleProcessor::OnStop(FMassDynamicObstacleFragment& OutObstacle, const float Radius)
+{
+	Super::OnStop(OutObstacle, Radius);
+
+	if (ZoneGraphSubsystem == nullptr || ZoneGraphAnnotationSubsystem == nullptr)
+	{
+		return;
+	}
+
+	ensureMsgf(OutObstacle.LaneObstacleIDs.IsEmpty(), TEXT("Blocked lanes have not been cleared."));
+
+	// Only send the a disturbance event if we are fully blocking a pedestrian lanes. 
+	const FZoneGraphTag CrowdTag = CrowdSettings->CrowdTag;
+	FZoneGraphTagFilter TagFilter;
+	TagFilter.AnyTags.Add(CrowdTag);
+	TArray<FZoneGraphLaneSection> OutLaneSections;
+	ZoneGraphSubsystem->FindLaneOverlaps(OutObstacle.LastPosition, Radius, TagFilter, OutLaneSections);
+	if (!OutLaneSections.IsEmpty())
+	{
+		// Add an obstacle disturbance.
+		FZoneGraphObstacleDisturbanceArea Disturbance;
+		Disturbance.Position = OutObstacle.LastPosition;
+		constexpr float EffectRadius = 1000.f;
+		Disturbance.Radius = EffectRadius;
+		Disturbance.ObstacleRadius = Radius;
+		Disturbance.ObstacleID = FMassLaneObstacleID::GetNextUniqueID();
+		Disturbance.Action = EZoneGraphObstacleDisturbanceAreaAction::Add;
+		ZoneGraphAnnotationSubsystem->SendEvent(Disturbance);
+
+		// Keep track of our obstacles.
+		OutObstacle.LaneObstacleIDs.Add(Disturbance.ObstacleID);
+	}
+}
+
+void UMassCrowdDynamicObstacleProcessor::OnMove(FMassDynamicObstacleFragment& OutObstacle)
+{
+	Super::OnMove(OutObstacle);
+
+	if (ZoneGraphAnnotationSubsystem == nullptr)
+	{
+		return;
+	}
+
+	// Unblock previously blocked lanes.
+	for (const FMassLaneObstacleID ID : OutObstacle.LaneObstacleIDs)
+	{
+		FZoneGraphObstacleDisturbanceArea Disturbance;
+		Disturbance.ObstacleID = ID;
+		Disturbance.Action = EZoneGraphObstacleDisturbanceAreaAction::Remove;
+		ZoneGraphAnnotationSubsystem->SendEvent(Disturbance);
+	}
+
+	OutObstacle.LaneObstacleIDs.Reset();
+}
diff --git a/Engine/Plugins/AI/MassCrowd/Source/MassCrowd/Private/MassCrowdReplicationProcessor.cpp b/Engine/Plugins/AI/MassCrowd/Source/MassCrowd/Private/MassCrowdReplicationProcessor.cpp
new file mode 100644
--- /dev/null
+++ b/Engine/Plugins/AI/MassCrowd/Source/MassCrowd/Private/MassCrowdReplicationProcessor.cpp
@@ -0,0 +1,119 @@
+// Copyright Epic Games, Inc. All Rights Reserved.
+
+#include "MassCrowdReplicationProcessor.h"
+#include "MassReplicationManager.h"
+#include "MassReplicationTypes.h"
+#include "MassClientBubbleHandler.h"
+#include "MassCrowdBubble.h"
+#include "MassCommonFragments.h"
+#include "MassLODManager.h"
+#include "MassCrowdFragments.h"
+#include "MassReplicationPathHandlers.h"
+#include "MassReplicationTransformHandlers.h"
+
+namespace UE { namespace Mass { namespace Crowd
+{
+	int32 bDebugReplicationViewerLOD = 0;
+	FAutoConsoleVariableRef CVarDebugReplicationViewerLOD(TEXT("ai.debug.CrowdReplicationViewerLOD"), bDebugReplicationViewerLOD, TEXT("Crowd Debug Replication Viewer LOD"), ECVF_Cheat);
+}}}
+
+//----------------------------------------------------------------------//
+//  UMassProcessor_Replication
+//----------------------------------------------------------------------//
+UMassCrowdReplicationProcessor::UMassCrowdReplicationProcessor()
+{
+#if !UE_ALLOW_DEBUG_REPLICATION_BUBBLES_STANDALONE
+	ExecutionFlags = int32(EProcessorExecutionFlags::Server);
+#else
+	ExecutionFlags = int32(EProcessorExecutionFlags::All);
+#endif // UE_ALLOW_DEBUG_REPLICATION_BUBBLES_STANDALONE
+}
+
+void UMassCrowdReplicationProcessor::ConfigureQueries()
+{
+	Super::ConfigureQueries();
+
+	FMassReplicationProcessorPositionYawHandler::AddRequirements(EntityQuery);
+	FMassReplicationProcessorPathHandler::AddRequirements(EntityQuery);
+
+	EntityQuery.AddTagRequirement<FTagFragment_MassCrowd>(ELWComponentPresence::All);
+}
+
+void UMassCrowdReplicationProcessor::Initialize(UObject& Owner)
+{
+	Super::Initialize(Owner);
+
+	check(ReplicationManager);
+
+	BubbleInfoClassHandle = ReplicationManager->GetBubbleInfoClassHandle(AMassCrowdClientBubbleInfo::StaticClass());
+}
+
+void UMassCrowdReplicationProcessor::Execute(UEntitySubsystem& EntitySubsystem, FLWComponentSystemExecutionContext& Context)
+{
+	Super::Execute(EntitySubsystem, Context);
+
+	ProcessClientReplication(EntitySubsystem, Context);
+
+#if WITH_MASS_DEBUG
+	// Optional debug display
+	if (UE::Mass::Crowd::bDebugReplicationViewerLOD)
+	{
+		EntityQuery.ForEachEntityChunk(EntitySubsystem, Context, [this](FLWComponentSystemExecutionContext& Context)
+			{
+				const TConstArrayView<FDataFragment_Transform> TransformList = Context.GetComponentView<FDataFragment_Transform>();
+				const TConstArrayView<FMassReplicationViewerLODFragment> ViewerLODList = Context.GetComponentView<FMassReplicationViewerLODFragment>();
+				LODCalculator.DebugDisplayLOD(Context, ViewerLODList, TransformList, World);
+			});
+	}
+#endif // WITH_MASS_DEBUG
+}
+
+void UMassCrowdReplicationProcessor::ProcessClientReplication(UEntitySubsystem& EntitySubsystem, FLWComponentSystemExecutionContext& Context)
+{
+#if UE_REPLICATION_COMPILE_SERVER_CODE
+
+	FMassReplicationProcessorPathHandler PathHandler;
+	FMassReplicationProcessorPositionYawHandler PositionYawHandler;
+
+	auto CacheViewsCallback = [&Context, &PathHandler, &PositionYawHandler]()
+	{
+		PathHandler.CacheComponentViews(Context);
+		PositionYawHandler.CacheComponentViews(Context);
+	};
+
+	auto AddEntityCallback = [this, &Context, &PathHandler, &PositionYawHandler](const int32 EntityIdx, FReplicatedCrowdAgent& InReplicatedAgent, const FMassClientHandle ClientHandle)->FMassReplicatedAgentHandle
+	{
+		AMassCrowdClientBubbleInfo& CrowdBubbleInfo = GetTypedClientBubbleInfoChecked<AMassCrowdClientBubbleInfo>(ClientHandle);
+
+		PathHandler.AddEntity(EntityIdx, InReplicatedAgent.GetReplicatedPathDataMutable());
+		PositionYawHandler.AddEntity(EntityIdx, InReplicatedAgent.GetReplicatedPositionYawDataMutable());
+
+		return CrowdBubbleInfo.GetCrowdSerializer().Bubble.AddAgent(Context.GetEntity(EntityIdx), InReplicatedAgent);
+	};
+
+	auto ModifyEntityCallback = [this, &PathHandler](const int32 EntityIdx, const EMassLOD::Type LOD, const float Time, const FMassReplicatedAgentHandle Handle, const FMassClientHandle ClientHandle)
+	{
+		AMassCrowdClientBubbleInfo& CrowdBubbleInfo = GetTypedClientBubbleInfoChecked<AMassCrowdClientBubbleInfo>(ClientHandle);
+
+		FMassCrowdClientBubbleHandler& Bubble = CrowdBubbleInfo.GetCrowdSerializer().Bubble;
+
+		PathHandler.ModifyEntity<FCrowdFastArrayItem>(Handle, EntityIdx, Bubble.GetPathHandlerMutable());
+
+		// Don't call the PositionYawHandler here as we currently only replicate the position and yaw when we add an entity to Mass
+	};
+
+	auto RemoveEntityCallback = [this](const FMassReplicatedAgentHandle Handle, const FMassClientHandle ClientHandle)
+	{
+		AMassCrowdClientBubbleInfo& CrowdBubbleInfo = GetTypedClientBubbleInfoChecked<AMassCrowdClientBubbleInfo>(ClientHandle);
+
+		CrowdBubbleInfo.GetCrowdSerializer().Bubble.RemoveAgentChecked(Handle);
+	};
+
+	QUICK_SCOPE_CYCLE_COUNTER(UMassCrowdReplicationProcessor_ProcessClientReplication);
+
+	EntityQuery.ForEachEntityChunk(EntitySubsystem, Context, [this, &CacheViewsCallback, &AddEntityCallback, &ModifyEntityCallback, &RemoveEntityCallback](FLWComponentSystemExecutionContext& Context)
+	{
+		CalculateClientReplication<FCrowdFastArrayItem>(Context, CacheViewsCallback, AddEntityCallback, ModifyEntityCallback, RemoveEntityCallback);
+	});
+#endif // UE_REPLICATION_COMPILE_SERVER_CODE
+}
diff --git a/Engine/Plugins/AI/MassCrowd/Source/MassCrowd/Private/MassCrowdRepresentationProcessor.cpp b/Engine/Plugins/AI/MassCrowd/Source/MassCrowd/Private/MassCrowdRepresentationProcessor.cpp
new file mode 100644
--- /dev/null
+++ b/Engine/Plugins/AI/MassCrowd/Source/MassCrowd/Private/MassCrowdRepresentationProcessor.cpp
@@ -0,0 +1,152 @@
+// Copyright Epic Games, Inc. All Rights Reserved.
+
+#include "MassCrowdRepresentationProcessor.h"
+#include "MassCrowdFragments.h"
+#include "MassAIMovementFragments.h"
+#include "Components/SkeletalMeshComponent.h"
+#include "Components/CapsuleComponent.h"
+#include "GameFramework/CharacterMovementComponent.h"
+#include "GameFramework/Character.h"
+#include "Translators/MassCharacterMovementTranslators.h"
+#include "MassAgentComponent.h"
+#include "MassCrowdRepresentationSubsystem.h"
+
+UMassCrowdRepresentationProcessor::UMassCrowdRepresentationProcessor()
+{
+	bAutoRegisterWithProcessingPhases = true;
+
+	ExecutionOrder.ExecuteAfter.Add(UE::Mass::ProcessorGroupNames::SyncWorldToMass);
+}
+
+void UMassCrowdRepresentationProcessor::Initialize(UObject& Owner)
+{
+	Super::Initialize(Owner);
+
+	// Override default representation subsystem by the crowd one to support parallelization
+	RepresentationSubsystem = UWorld::GetSubsystem<UMassCrowdRepresentationSubsystem>(Owner.GetWorld());
+}
+
+
+void UMassCrowdRepresentationProcessor::ConfigureQueries()
+{
+	Super::ConfigureQueries();
+	EntityQuery.AddTagRequirement<FTagFragment_MassCrowd>(ELWComponentPresence::All);
+
+	CharacterMovementEntitiesQuery_Conditional = EntityQuery;
+	CharacterMovementEntitiesQuery_Conditional.AddRequirement<FDataFragment_CharacterMovementComponentWrapper>(ELWComponentAccess::ReadWrite);
+	CharacterMovementEntitiesQuery_Conditional.AddRequirement<FMassVelocityFragment>(ELWComponentAccess::ReadOnly);
+	CharacterMovementEntitiesQuery_Conditional.SetChunkFilter(&FMassVisualizationChunkFragment::AreAnyEntitiesVisibleInChunk);
+}
+
+void UMassCrowdRepresentationProcessor::InitializeVelocity(UEntitySubsystem& EntitySubsystem, FLWComponentSystemExecutionContext& Context)
+{
+	CharacterMovementEntitiesQuery_Conditional.ForEachEntityChunk(EntitySubsystem, Context, [this](FLWComponentSystemExecutionContext& Context)
+		{
+			const TConstArrayView<FMassRepresentationFragment> VisualizationList = Context.GetComponentView<FMassRepresentationFragment>();
+			const TConstArrayView<FMassVelocityFragment> VelocityList = Context.GetComponentView<FMassVelocityFragment>();
+			const TArrayView<FDataFragment_CharacterMovementComponentWrapper> CharMoveWrapperList = Context.GetMutableComponentView<FDataFragment_CharacterMovementComponentWrapper>();
+
+			const int32 NumEntities = Context.GetEntitiesNum();
+			for (int32 EntityIdx = 0; EntityIdx < NumEntities; EntityIdx++)
+			{
+				const FMassRepresentationFragment& Visualization = VisualizationList[EntityIdx];
+
+				// @todo: Add syncing skeletal mesh <-> static mesh here.
+				if (Visualization.PrevRepresentation != Visualization.CurrentRepresentation)
+				{
+					const FMassVelocityFragment& Velocity = VelocityList[EntityIdx];
+					FDataFragment_CharacterMovementComponentWrapper& CharMoveWrapper = CharMoveWrapperList[EntityIdx];
+					UCharacterMovementComponent* MovementComp = CharMoveWrapper.Component.Get();
+
+					const bool bWasUsingMovementComponent = (Visualization.PrevRepresentation == ERepresentationType::LowResSpawnedActor || Visualization.PrevRepresentation == ERepresentationType::HighResSpawnedActor);
+					const bool bIsUsingMovementComponent = (Visualization.CurrentRepresentation == ERepresentationType::LowResSpawnedActor || Visualization.CurrentRepresentation == ERepresentationType::HighResSpawnedActor);
+					if (MovementComp && !bWasUsingMovementComponent && bIsUsingMovementComponent)
+					{
+						MovementComp->Velocity = Velocity.Value;
+					}
+				}
+			}
+		});
+}
+
+void UMassCrowdRepresentationProcessor::SetActorEnabled(const EActorEnabledType EnabledType, AActor& Actor, const int32 EntityIdx, FLWComponentSystemExecutionContext& Context)
+{
+	Super::SetActorEnabled(EnabledType, Actor, EntityIdx, Context);
+
+	const bool bEnabled = EnabledType != EActorEnabledType::Disabled;
+
+	USkeletalMeshComponent* SkeletalMeshComponent = Actor.FindComponentByClass<USkeletalMeshComponent>();
+	if (SkeletalMeshComponent)
+	{
+		// Enable/disable the ticking and visibility of SkeletalMesh and its children
+		SkeletalMeshComponent->SetVisibility(bEnabled);
+		SkeletalMeshComponent->SetComponentTickEnabled(bEnabled);
+		const TArray<USceneComponent*>& AttachedChildren = SkeletalMeshComponent->GetAttachChildren();
+		if (AttachedChildren.Num() > 0)
+		{
+			TInlineComponentArray<USceneComponent*, NumInlinedActorComponents> ComponentStack;
+
+			ComponentStack.Append(AttachedChildren);
+			while (ComponentStack.Num() > 0)
+			{
+				USceneComponent* const CurrentComp = ComponentStack.Pop(/*bAllowShrinking=*/false);
+				if (CurrentComp)
+				{
+					ComponentStack.Append(CurrentComp->GetAttachChildren());
+					CurrentComp->SetVisibility(bEnabled);
+					if (bEnabled)
+					{
+						// Re-enable only if it was enabled at startup
+						CurrentComp->SetComponentTickEnabled(CurrentComp->PrimaryComponentTick.bStartWithTickEnabled);
+					}
+					else
+					{
+						CurrentComp->SetComponentTickEnabled(false);
+					}
+				}
+			}
+		}
+	}
+
+	// Enable/disable the ticking of CharacterMovementComponent as well
+	ACharacter* Character = Cast<ACharacter>(&Actor);
+	UCharacterMovementComponent* MovementComp = Character != nullptr ? Character->GetCharacterMovement() : nullptr;
+	if (MovementComp != nullptr)
+	{
+		MovementComp->SetComponentTickEnabled(bEnabled);
+	}
+
+	// when we "suspend" the puppet actor we need to let the agent subsystem know by unregistering the agent component
+	// associated with the actor. This will result in removing all the puppet-actor-specific fragments which in turn
+	// will exclude the owner entity from being processed by puppet-specific processors (usually translators).
+	if (UMassAgentComponent* AgentComp = Actor.FindComponentByClass<UMassAgentComponent>())
+	{
+		AgentComp->PausePuppet(!bEnabled);
+	}
+}
+
+AActor* UMassCrowdRepresentationProcessor::GetOrSpawnActor(const FMassHandle MassAgent, FDataFragment_Actor& ActorInfo, const FTransform& Transform, const int16 TemplateActorIndex, FMassHandle_ActorSpawnRequest& SpawnRequestHandle, const float Priority)
+{
+	FTransform RootTransform = Transform;
+	
+	if (const AActor* DefaultActor = RepresentationSubsystem->GetTemplateActorClass(TemplateActorIndex).GetDefaultObject())
+	{
+		if (const UCapsuleComponent* CapsuleComp = DefaultActor->FindComponentByClass<UCapsuleComponent>())
+		{
+			RootTransform.AddToTranslation(FVector(0.0f, 0.0f, CapsuleComp->GetScaledCapsuleHalfHeight()));
+		}
+	}
+
+	return Super::GetOrSpawnActor(MassAgent, ActorInfo, RootTransform, TemplateActorIndex, SpawnRequestHandle, Priority);
+}
+
+void UMassCrowdRepresentationProcessor::TeleportActor(const FTransform& Transform, AActor& Actor, FLWComponentSystemExecutionContext& Context)
+{
+	FTransform RootTransform = Transform;
+
+	if (const UCapsuleComponent* CapsuleComp = Actor.FindComponentByClass<UCapsuleComponent>())
+	{
+		RootTransform.AddToTranslation(FVector(0.0f, 0.0f, CapsuleComp->GetScaledCapsuleHalfHeight()));
+	}
+	Super::TeleportActor(RootTransform, Actor, Context);
+}
diff --git a/Engine/Plugins/AI/MassCrowd/Source/MassCrowd/Private/MassCrowdServerRepresentationLODProcessor.cpp b/Engine/Plugins/AI/MassCrowd/Source/MassCrowd/Private/MassCrowdServerRepresentationLODProcessor.cpp
new file mode 100644
--- /dev/null
+++ b/Engine/Plugins/AI/MassCrowd/Source/MassCrowd/Private/MassCrowdServerRepresentationLODProcessor.cpp
@@ -0,0 +1,91 @@
+// Copyright Epic Games, Inc. All Rights Reserved.
+
+#include "MassCrowdServerRepresentationLODProcessor.h"
+#include "MassCommonFragments.h"
+#include "MassRepresentationFragments.h"
+
+namespace UE::MassCrowd
+{
+	int32 bDebugCrowdServerRepresentationLOD = 0;
+	FAutoConsoleVariableRef CVarDebugServerRepresentationLODTest(TEXT("ai.debug.CrowdServerRepresentationLOD"), bDebugCrowdServerRepresentationLOD, TEXT("Debug Crowd ServerRepresentation LOD"), ECVF_Cheat);
+} // UE::MassCrowd
+
+UMassCrowdServerRepresentationLODProcessor::UMassCrowdServerRepresentationLODProcessor()
+{
+	ExecutionFlags = (int32)EProcessorExecutionFlags::Server;
+
+	ExecutionOrder.ExecuteInGroup = UE::Mass::ProcessorGroupNames::LOD;
+	ExecutionOrder.ExecuteAfter.Add(UE::Mass::ProcessorGroupNames::LODCollector);
+
+	LODDistance[EMassLOD::High] = 0.0f;
+	LODDistance[EMassLOD::Medium] = 5000.0f;
+	LODDistance[EMassLOD::Low] = 5000.0f;
+	LODDistance[EMassLOD::Off] = 5000.0f;
+	
+	LODMaxCount[EMassLOD::High] = 50;
+	LODMaxCount[EMassLOD::Medium] = 0;
+	LODMaxCount[EMassLOD::Low] = 0;
+	LODMaxCount[EMassLOD::Off] = INT_MAX;
+}
+
+void UMassCrowdServerRepresentationLODProcessor::ConfigureQueries()
+{
+	EntityQuery.AddTagRequirement<FTagFragment_MassCrowd>(ELWComponentPresence::All);
+	EntityQuery.AddRequirement<FDataFragment_Transform>(ELWComponentAccess::ReadOnly);
+	EntityQuery.AddRequirement<FMassLODInfoFragment>(ELWComponentAccess::ReadOnly);
+	EntityQuery.AddRequirement<FMassRepresentationLODFragment>(ELWComponentAccess::ReadWrite);
+}
+
+void UMassCrowdServerRepresentationLODProcessor::Initialize(UObject& InOwner)
+{
+	LODCalculator.Initialize(LODDistance, BufferHysteresisOnDistancePercentage / 100.0f, LODMaxCount);
+
+	Super::Initialize(InOwner);
+}
+
+void UMassCrowdServerRepresentationLODProcessor::Execute(UEntitySubsystem& EntitySubsystem, FLWComponentSystemExecutionContext& Context)
+{
+	TRACE_CPUPROFILER_EVENT_SCOPE(TEXT("CrowdServerRepresentationLOD"))
+
+	check(LODManager);
+	const TArray<FViewerInfo>& Viewers = LODManager->GetViewers();
+	LODCalculator.PrepareExecution(Viewers);
+	
+	{
+		TRACE_CPUPROFILER_EVENT_SCOPE(TEXT("CalculateLOD"))
+		
+		EntityQuery.ForEachEntityChunk(EntitySubsystem, Context, [this](FLWComponentSystemExecutionContext& Context)
+		{
+			const TConstArrayView<FMassLODInfoFragment> ViewersInfoList = Context.GetComponentView<FMassLODInfoFragment>();
+			const TArrayView<FMassRepresentationLODFragment> RepresentationLODFragments = Context.GetMutableComponentView<FMassRepresentationLODFragment>();
+			LODCalculator.CalculateLOD(Context, RepresentationLODFragments, ViewersInfoList);
+		});
+	}
+
+	{
+		TRACE_CPUPROFILER_EVENT_SCOPE(TEXT("AdjustDistancesAndLODFromCount"))
+		
+		if (LODCalculator.AdjustDistancesFromCount())
+		{
+			EntityQuery.ForEachEntityChunk(EntitySubsystem, Context, [this](FLWComponentSystemExecutionContext& Context)
+			{
+				const TConstArrayView<FMassLODInfoFragment> ViewersInfoList = Context.GetComponentView<FMassLODInfoFragment>();
+				const TArrayView<FMassRepresentationLODFragment> RepresentationLODFragments = Context.GetMutableComponentView<FMassRepresentationLODFragment>();
+				LODCalculator.AdjustLODFromCount(Context, RepresentationLODFragments, ViewersInfoList);
+			});
+		}
+	}
+
+	// Optional debug display
+	if (UE::MassCrowd::bDebugCrowdServerRepresentationLOD)
+	{
+		TRACE_CPUPROFILER_EVENT_SCOPE(TEXT("DebugDisplayLOD"))
+		
+		EntityQuery.ForEachEntityChunk(EntitySubsystem, Context, [this](FLWComponentSystemExecutionContext& Context)
+		{
+			const TConstArrayView<FDataFragment_Transform> LocationList = Context.GetComponentView<FDataFragment_Transform>();
+			const TConstArrayView<FMassRepresentationLODFragment> RepresentationLODFragments = Context.GetComponentView<FMassRepresentationLODFragment>();
+			LODCalculator.DebugDisplayLOD(Context, RepresentationLODFragments, LocationList, World);
+		});
+	}
+}
\ No newline at end of file
diff --git a/Engine/Plugins/AI/MassCrowd/Source/MassCrowd/Private/MassCrowdServerRepresentationProcessor.cpp b/Engine/Plugins/AI/MassCrowd/Source/MassCrowd/Private/MassCrowdServerRepresentationProcessor.cpp
new file mode 100644
--- /dev/null
+++ b/Engine/Plugins/AI/MassCrowd/Source/MassCrowd/Private/MassCrowdServerRepresentationProcessor.cpp
@@ -0,0 +1,28 @@
+// Copyright Epic Games, Inc. All Rights Reserved.
+
+#include "MassCrowdServerRepresentationProcessor.h"
+#include "MassCrowdFragments.h"
+#include "MassActorSubsystem.h"
+
+UMassCrowdServerRepresentationProcessor::UMassCrowdServerRepresentationProcessor()
+{
+	ExecutionFlags = (int32)EProcessorExecutionFlags::Server;
+
+	ExecutionOrder.ExecuteAfter.Add(UE::Mass::ProcessorGroupNames::LOD);
+
+	LODRepresentation[EMassLOD::High] = ERepresentationType::HighResSpawnedActor;
+	LODRepresentation[EMassLOD::Medium] = ERepresentationType::None;
+	LODRepresentation[EMassLOD::Low] = ERepresentationType::None;
+	LODRepresentation[EMassLOD::Off] = ERepresentationType::None;
+}
+
+void UMassCrowdServerRepresentationProcessor::Execute(UEntitySubsystem& EntitySubsystem, FLWComponentSystemExecutionContext& Context)
+{
+	EntityQuery.ForEachEntityChunk(EntitySubsystem, Context, [this](FLWComponentSystemExecutionContext& Context)
+	{
+		UpdateRepresentation(Context);
+	});
+
+	// @todo, we should use the new translators to do that initialization
+	InitializeVelocity(EntitySubsystem, Context);
+}
\ No newline at end of file
diff --git a/Engine/Plugins/AI/MassCrowd/Source/MassCrowd/Private/MassCrowdServerRepresentationTrait.cpp b/Engine/Plugins/AI/MassCrowd/Source/MassCrowd/Private/MassCrowdServerRepresentationTrait.cpp
new file mode 100644
--- /dev/null
+++ b/Engine/Plugins/AI/MassCrowd/Source/MassCrowd/Private/MassCrowdServerRepresentationTrait.cpp
@@ -0,0 +1,34 @@
+// Copyright Epic Games, Inc. All Rights Reserved.
+
+#include "MassCrowdServerRepresentationTrait.h"
+#include "MassEntityTemplateRegistry.h"
+#include "MassCrowdRepresentationSubsystem.h"
+#include "MassCommonFragments.h"
+#include "MassRepresentationFragments.h"
+#include "Engine/World.h"
+
+void UMassCrowdServerRepresentationTrait::BuildTemplate(FMassEntityTemplateBuildContext& BuildContext, UWorld& World) const
+{
+	// This should only be ran on NM_DedicatedServer network mode
+	if (!World.IsNetMode(NM_DedicatedServer))
+	{
+		return;
+	}
+	
+	// the following needs to be always there for mesh vis to work. Adding following fragments after already 
+	// adding Config.AdditionalDataFragments to let user configure the fragments first. Calling BuildContext.Add() 
+	// won't override any fragments that are already there
+	BuildContext.AddFragment<FDataFragment_Transform>();
+
+	if (UMassCrowdRepresentationSubsystem* RepresentationSubsystem = World.GetSubsystem<UMassCrowdRepresentationSubsystem>())
+	{
+		FMassRepresentationFragment& RepresentationFragment = BuildContext.AddFragmentWithDefaultInitializer_GetRef<FMassRepresentationFragment>();
+		RepresentationFragment.StaticMeshDescIndex = INDEX_NONE;
+		RepresentationFragment.HighResTemplateActorIndex = TemplateActor.Get() ? RepresentationSubsystem->FindOrAddTemplateActor(TemplateActor.Get()) : INDEX_NONE;
+		RepresentationFragment.LowResTemplateActorIndex = INDEX_NONE;
+
+		BuildContext.AddFragment<FMassRepresentationLODFragment>();
+
+		BuildContext.AddChunkFragment<FMassVisualizationChunkFragment>();
+	}
+}
\ No newline at end of file
diff --git a/Engine/Plugins/AI/MassCrowd/Source/MassCrowd/Private/MassCrowdVisualizationLODProcessor.cpp b/Engine/Plugins/AI/MassCrowd/Source/MassCrowd/Private/MassCrowdVisualizationLODProcessor.cpp
new file mode 100644
--- /dev/null
+++ b/Engine/Plugins/AI/MassCrowd/Source/MassCrowd/Private/MassCrowdVisualizationLODProcessor.cpp
@@ -0,0 +1,161 @@
+// Copyright Epic Games, Inc. All Rights Reserved.
+
+#include "MassCrowdVisualizationLODProcessor.h"
+#include "MassCommonFragments.h"
+#include "MassCrowdFragments.h"
+
+namespace UE::MassCrowd
+{
+
+	MASSCROWD_API int32 GCrowdTurnOffVisualization = 0;
+	FAutoConsoleVariableRef CVarCrowdTurnOffVisualization(TEXT("Mass.CrowdTurnOffVisualization"), GCrowdTurnOffVisualization, TEXT("Turn off crowd visualization"));
+
+	int32 bDebugCrowdVisualizationLOD = 0;
+	int32 bDebugShowISMUnderSpecifiedRange = 0;
+
+	FAutoConsoleVariableRef ConsoleVariables[] =
+	{
+		FAutoConsoleVariableRef(TEXT("ai.debug.CrowdVisualizationLOD"), bDebugCrowdVisualizationLOD, TEXT("Debug crowd visualization LOD"), ECVF_Cheat),
+		FAutoConsoleVariableRef(TEXT("ai.debug.ShowISMUnderSpecifiedRange"), bDebugShowISMUnderSpecifiedRange, TEXT("Show ISM under a specified range (meters)"), ECVF_Cheat)
+	};
+
+#if WITH_EDITOR
+	int32 DebugCrowdMaxCountHigh = -1;
+	int32 DebugCrowdMaxCountMedium = -1;
+	int32 DebugCrowdMaxCountLow = -1;
+
+	FAutoConsoleVariableRef EditorConsoleVariables[] =
+	{
+		FAutoConsoleVariableRef(TEXT("ai.debug.CrowdMaxCountHigh"), DebugCrowdMaxCountHigh, TEXT("Forced MASS crowd max count high"), ECVF_Cheat),
+		FAutoConsoleVariableRef(TEXT("ai.debug.CrowdMaxCountMedium"), DebugCrowdMaxCountMedium, TEXT("Forced MASS crowd max count medium"), ECVF_Cheat),
+		FAutoConsoleVariableRef(TEXT("ai.debug.CrowdMaxCountLow"), DebugCrowdMaxCountLow, TEXT("Forced MASS crowd max count low"), ECVF_Cheat)
+	};
+
+	int32 LastDebugCrowdMaxCountHigh = -1;
+	int32 LastDebugCrowdMaxCountMedium = -1;
+	int32 LastDebugCrowdMaxCountLow = -1;
+#endif // WITH_EDITOR
+} // UE::MassCrowd
+
+UMassCrowdVisualizationLODProcessor::UMassCrowdVisualizationLODProcessor()
+{
+	ExecutionFlags = (int32)(EProcessorExecutionFlags::Client | EProcessorExecutionFlags::Standalone);
+
+	ExecutionOrder.ExecuteInGroup = UE::Mass::ProcessorGroupNames::LOD;
+	ExecutionOrder.ExecuteAfter.Add(UE::Mass::ProcessorGroupNames::LODCollector);
+
+#if WITH_EDITOR
+	UE::MassCrowd::LastDebugCrowdMaxCountHigh = -1;
+	UE::MassCrowd::LastDebugCrowdMaxCountMedium = -1;
+	UE::MassCrowd::LastDebugCrowdMaxCountLow = -1;
+#endif // WITH_EDITOR
+}
+
+void UMassCrowdVisualizationLODProcessor::ConfigureQueries()
+{
+	// @todo remove, no need for this anymore since we have the common LOD collector
+	// Do not call super as we do have our own LODInfo fragment, so need to duplicate
+
+	CloseEntityQuery.AddTagRequirement<FTagFragment_MassCrowd>(ELWComponentPresence::All);
+	CloseEntityQuery.AddRequirement<FMassLODInfoFragment>(ELWComponentAccess::ReadOnly);
+	CloseEntityQuery.AddRequirement<FMassRepresentationFragment>(ELWComponentAccess::ReadOnly);
+	CloseEntityQuery.AddRequirement<FMassRepresentationLODFragment>(ELWComponentAccess::ReadWrite);
+	CloseEntityQuery.AddRequirement<FDataFragment_Transform>(ELWComponentAccess::ReadOnly);
+
+	FarEntityQuery = CloseEntityQuery;
+	CloseEntityQuery.AddTagRequirement<FMassVisibilityCulledByDistanceTag>(ELWComponentPresence::None);
+	FarEntityQuery.AddTagRequirement<FMassVisibilityCulledByDistanceTag>(ELWComponentPresence::All);
+	FarEntityQuery.AddChunkRequirement<FMassVisualizationChunkFragment>(ELWComponentAccess::ReadOnly);
+}
+
+void UMassCrowdVisualizationLODProcessor::Execute(UEntitySubsystem& EntitySubsystem, FLWComponentSystemExecutionContext& Context)
+{
+	ForceOffLOD((bool)UE::MassCrowd::GCrowdTurnOffVisualization);
+
+	TRACE_CPUPROFILER_EVENT_SCOPE(TEXT("CrowdVisualizationLOD"))
+
+	Super::ExecuteInternal<FMassLODInfoFragment>(EntitySubsystem, Context);
+	
+	if (UE::MassCrowd::bDebugCrowdVisualizationLOD)
+	{
+		TRACE_CPUPROFILER_EVENT_SCOPE(TEXT("DebugDisplayLOD"))
+
+		auto DebugDisplayLOD = [this](FLWComponentSystemExecutionContext& Context)
+		{
+			const TConstArrayView<FDataFragment_Transform> LocationList = Context.GetComponentView<FDataFragment_Transform>();
+			const TConstArrayView<FMassRepresentationLODFragment> VisualizationLODList = Context.GetComponentView<FMassRepresentationLODFragment>();
+			LODCalculator.DebugDisplayLOD(Context, VisualizationLODList, LocationList, World);
+		};
+
+		CloseEntityQuery.ForEachEntityChunk(EntitySubsystem, Context, DebugDisplayLOD);
+		FarEntityQuery.ForEachEntityChunk(EntitySubsystem, Context, DebugDisplayLOD);
+	}
+
+	if (UE::MassCrowd::bDebugShowISMUnderSpecifiedRange > 0)
+	{
+		TRACE_CPUPROFILER_EVENT_SCOPE(TEXT("ShowISMUnderSpecifiedRange"))
+
+		auto ShowISMUnderSpecifiedRange = [this](const FLWComponentSystemExecutionContext& Context)
+		{
+			const TConstArrayView<FDataFragment_Transform> LocationList = Context.GetComponentView<FDataFragment_Transform>();
+			const TConstArrayView<FMassRepresentationFragment> RepresentationFragmentList = Context.GetComponentView<FMassRepresentationFragment>();
+			const TConstArrayView<FMassRepresentationLODFragment> RepresentationLODFragmentList = Context.GetComponentView<FMassRepresentationLODFragment>();
+			const int32 NumEntities = Context.GetEntitiesNum();
+			const float SpecifiedRangeSquaredCentimeters = FMath::Square(UE::MassCrowd::bDebugShowISMUnderSpecifiedRange * 100);
+			for (int EntityIdx = 0; EntityIdx < NumEntities; EntityIdx++)
+			{
+				const FMassRepresentationFragment& RepresentationFragment = RepresentationFragmentList[EntityIdx];
+				const FMassRepresentationLODFragment& RepresentationLODFragment = RepresentationLODFragmentList[EntityIdx];
+				if (RepresentationFragment.CurrentRepresentation == ERepresentationType::StaticMeshInstance && SpecifiedRangeSquaredCentimeters > RepresentationLODFragment.ClosestViewerDistanceSq)
+				{
+					const FDataFragment_Transform& EntityLocation = LocationList[EntityIdx];
+					DrawDebugSolidBox(World, EntityLocation.GetTransform().GetLocation() + FVector(0.0f, 0.0f, 150.0f), FVector(50.0f), FColor::Red);
+				}
+			}
+		};
+		CloseEntityQuery.ForEachEntityChunk(EntitySubsystem, Context, ShowISMUnderSpecifiedRange);
+		FarEntityQuery.ForEachEntityChunk(EntitySubsystem, Context, ShowISMUnderSpecifiedRange);
+	}
+
+#if WITH_EDITOR
+	if (UE::MassCrowd::LastDebugCrowdMaxCountHigh != UE::MassCrowd::DebugCrowdMaxCountHigh)
+	{
+		if (UE::MassCrowd::DebugCrowdMaxCountHigh >= 0)
+		{
+			LODMaxCount[0] = UE::MassCrowd::DebugCrowdMaxCountHigh;
+		}
+		else
+		{
+			LODMaxCount[0] = GetClass()->GetDefaultObject<UMassCrowdVisualizationLODProcessor>()->LODMaxCount[0];
+		}
+		UE::MassCrowd::LastDebugCrowdMaxCountHigh = UE::MassCrowd::DebugCrowdMaxCountHigh;
+		LODCalculator.Initialize(BaseLODDistance, BufferHysteresisOnDistancePercentage / 100.f, LODMaxCount, nullptr, VisibleLODDistance);
+	}
+	if (UE::MassCrowd::LastDebugCrowdMaxCountMedium != UE::MassCrowd::DebugCrowdMaxCountMedium)
+	{
+		if (UE::MassCrowd::DebugCrowdMaxCountMedium >= 0)
+		{
+			LODMaxCount[1] = UE::MassCrowd::DebugCrowdMaxCountMedium;
+		}
+		else
+		{
+			LODMaxCount[1] = GetClass()->GetDefaultObject<UMassCrowdVisualizationLODProcessor>()->LODMaxCount[1];
+		}
+		UE::MassCrowd::LastDebugCrowdMaxCountMedium = UE::MassCrowd::DebugCrowdMaxCountMedium;
+		LODCalculator.Initialize(BaseLODDistance, BufferHysteresisOnDistancePercentage / 100.f, LODMaxCount, nullptr, VisibleLODDistance);
+	}
+	if (UE::MassCrowd::LastDebugCrowdMaxCountLow != UE::MassCrowd::DebugCrowdMaxCountLow)
+	{
+		if (UE::MassCrowd::DebugCrowdMaxCountLow >= 0)
+		{
+			LODMaxCount[2] = UE::MassCrowd::DebugCrowdMaxCountLow;
+		}
+		else
+		{
+			LODMaxCount[2] = GetClass()->GetDefaultObject<UMassCrowdVisualizationLODProcessor>()->LODMaxCount[2];
+		}
+		UE::MassCrowd::LastDebugCrowdMaxCountLow = UE::MassCrowd::DebugCrowdMaxCountLow;
+		LODCalculator.Initialize(BaseLODDistance, BufferHysteresisOnDistancePercentage / 100.f, LODMaxCount, nullptr, VisibleLODDistance);
+	}
+#endif // WITH_EDITOR
+}
diff --git a/Engine/Plugins/AI/MassCrowd/Source/MassCrowd/Private/MassCrowdVisualizationProcessor.cpp b/Engine/Plugins/AI/MassCrowd/Source/MassCrowd/Private/MassCrowdVisualizationProcessor.cpp
new file mode 100644
--- /dev/null
+++ b/Engine/Plugins/AI/MassCrowd/Source/MassCrowd/Private/MassCrowdVisualizationProcessor.cpp
@@ -0,0 +1,87 @@
+// Copyright Epic Games, Inc. All Rights Reserved.
+
+#include "MassCrowdVisualizationProcessor.h"
+#include "MassCrowdFragments.h"
+#include "Components/CapsuleComponent.h"
+#include "GameFramework/CharacterMovementComponent.h"
+#include "MassActorSubsystem.h"
+#include "MassCrowdRepresentationSubsystem.h"
+
+namespace UE::MassCrowd
+{
+	int32 bDebugCrowdVisualType = 0;
+	FAutoConsoleVariableRef CVarDebugVisualType(TEXT("ai.debug.CrowdVisualType"), bDebugCrowdVisualType, TEXT("Debug Crowd Visual Type"), ECVF_Cheat);
+
+	FColor CrowdRepresentationTypesColors[] =
+	{
+		FColor::Red,
+		FColor::Yellow,
+		FColor::Emerald,
+		FColor::White,
+	};
+} // UE::MassCrowd
+
+UMassCrowdVisualizationProcessor::UMassCrowdVisualizationProcessor()
+{
+	ExecutionFlags = (int32)(EProcessorExecutionFlags::Client | EProcessorExecutionFlags::Standalone);
+
+	ExecutionOrder.ExecuteAfter.Add(UE::Mass::ProcessorGroupNames::LOD);
+
+	bRequiresGameThreadExecution = true;
+}
+
+void UMassCrowdVisualizationProcessor::ConfigureQueries()
+{
+	Super::ConfigureQueries();
+}
+
+void UMassCrowdVisualizationProcessor::Execute(UEntitySubsystem& EntitySubsystem, FLWComponentSystemExecutionContext& Context)
+{
+	EntityQuery.ForEachEntityChunk(EntitySubsystem, Context, [this](FLWComponentSystemExecutionContext& Context)
+	{
+		UpdateVisualization(Context);
+	});
+
+	// @todo, we should use the new translators to do that initialization
+	InitializeVelocity(EntitySubsystem, Context);
+
+	if (UE::MassCrowd::bDebugCrowdVisualType)
+	{
+		check(World);
+
+		TRACE_CPUPROFILER_EVENT_SCOPE(TEXT("DebugDisplayVisualType"))
+
+		EntityQuery.ForEachEntityChunk(EntitySubsystem, Context, [this](const FLWComponentSystemExecutionContext& Context)
+		{
+			const TConstArrayView<FMassRepresentationFragment> VisualizationList = Context.GetComponentView<FMassRepresentationFragment>();
+			const TConstArrayView<FDataFragment_Actor> ActorList = Context.GetComponentView<FDataFragment_Actor>();
+			const TConstArrayView<FDataFragment_Transform> EntityLocationList = Context.GetComponentView<FDataFragment_Transform>();
+
+			const int32 NumEntities = Context.GetEntitiesNum();
+			for (int EntityIdx = 0; EntityIdx < NumEntities; EntityIdx++)
+			{
+				const FDataFragment_Transform& EntityLocation = EntityLocationList[EntityIdx];
+				const FMassRepresentationFragment& Visualization = VisualizationList[EntityIdx];
+				const FDataFragment_Actor& ActorInfo = ActorList[EntityIdx];
+				const int32 RepresentationTypeIdx = (int32)Visualization.CurrentRepresentation;
+				// Show replicated actors
+				if (ActorInfo.IsValid() && !ActorInfo.IsOwnedByMass())
+				{
+					DrawDebugBox(World, EntityLocation.GetTransform().GetLocation() + FVector(0.0f, 0.0f, 120.0f), FVector(25.0f), UE::MassCrowd::CrowdRepresentationTypesColors[RepresentationTypeIdx]);
+				}
+				else
+				{ 
+					DrawDebugSolidBox(World, EntityLocation.GetTransform().GetLocation() + FVector(0.0f, 0.0f, 120.0f), FVector(25.0f), UE::MassCrowd::CrowdRepresentationTypesColors[RepresentationTypeIdx]);
+				}
+			}
+		});
+	}
+}
+
+void UMassCrowdRepresentationFragmentDestructor::Initialize(UObject& Owner)
+{
+	Super::Initialize(Owner);
+
+	// Override default representation subsystem by the crowd one to support parallelization
+	RepresentationSubsystem = UWorld::GetSubsystem<UMassCrowdRepresentationSubsystem>(Owner.GetWorld());
+}
diff --git a/Engine/Plugins/AI/MassCrowd/Source/MassCrowd/Private/MassCrowdVisualizationTrait.cpp b/Engine/Plugins/AI/MassCrowd/Source/MassCrowd/Private/MassCrowdVisualizationTrait.cpp
new file mode 100644
--- /dev/null
+++ b/Engine/Plugins/AI/MassCrowd/Source/MassCrowd/Private/MassCrowdVisualizationTrait.cpp
@@ -0,0 +1,12 @@
+// Copyright Epic Games, Inc. All Rights Reserved.
+
+#include "MassCrowdVisualizationTrait.h"
+#include "MassCrowdRepresentationSubsystem.h"
+#include "MassCrowdVisualizationProcessor.h"
+
+UMassCrowdVisualizationTrait::UMassCrowdVisualizationTrait()
+{
+	// Override the subsystem and deinitializer to support parallelization of the crowd
+	RepresentationSubsystemClass = UMassCrowdRepresentationSubsystem::StaticClass();
+	RepresentationFragmentDeinitializerClass = UMassCrowdRepresentationFragmentDestructor::StaticClass();
+}
diff --git a/Engine/Plugins/AI/MassCrowd/Source/MassCrowd/Private/Tasks/MassCrowdClaimWaitSlotTask.cpp b/Engine/Plugins/AI/MassCrowd/Source/MassCrowd/Private/Tasks/MassCrowdClaimWaitSlotTask.cpp
new file mode 100644
--- /dev/null
+++ b/Engine/Plugins/AI/MassCrowd/Source/MassCrowd/Private/Tasks/MassCrowdClaimWaitSlotTask.cpp
@@ -0,0 +1,81 @@
+﻿// Copyright Epic Games, Inc. All Rights Reserved.
+
+#include "Tasks/MassCrowdClaimWaitSlotTask.h"
+#include "StateTreeExecutionContext.h"
+#include "MassCrowdSubsystem.h"
+#include "MassZoneGraphMovementFragments.h"
+#include "MassAIMovementFragments.h"
+#include "MassStateTreeExecutionContext.h"
+
+EStateTreeRunStatus FMassCrowdClaimWaitSlotTask::EnterState(FStateTreeExecutionContext& Context, const EStateTreeStateChangeType ChangeType, const FStateTreeTransitionResult& Transition)
+{
+	if (ChangeType != EStateTreeStateChangeType::Changed)
+	{
+		return EStateTreeRunStatus::Running;
+	}
+
+	WaitSlotLocation = nullptr;
+
+	const FMassStateTreeExecutionContext& MassContext = static_cast<FMassStateTreeExecutionContext&>(Context);
+	const FLWEntity Entity = MassContext.GetEntity();
+	
+	const FMassZoneGraphLaneLocationFragment& LaneLocation = Context.GetExternalItem(LocationHandle).Get<FMassZoneGraphLaneLocationFragment>();
+	const FMassMoveTargetFragment& MoveTarget = Context.GetExternalItem(MoveTargetHandle).Get<FMassMoveTargetFragment>();
+	UMassCrowdSubsystem& CrowdSubsystem = Context.GetExternalItem(CrowdSubsystemHandle).GetMutable<UMassCrowdSubsystem>();
+
+	FVector SlotPosition = FVector::ZeroVector;
+	FVector SlotDirection = FVector::ForwardVector;
+	WaitingSlotIndex = CrowdSubsystem.AcquireWaitingSlot(Entity, MoveTarget.Center, LaneLocation.LaneHandle, SlotPosition, SlotDirection);
+	if (WaitingSlotIndex == INDEX_NONE)
+	{
+		// Failed to acquire slot
+		return EStateTreeRunStatus::Failed;
+	}
+	
+	AcquiredLaneHandle = LaneLocation.LaneHandle;
+
+	TargetLocation.LaneHandle = LaneLocation.LaneHandle;
+	TargetLocation.NextExitLinkType = EZoneLaneLinkType::None;
+	TargetLocation.NextLaneHandle.Reset();
+	TargetLocation.bMoveReverse = false;
+	TargetLocation.EndOfPathIntent = EMassMovementAction::Stand;
+	TargetLocation.EndOfPathPosition = SlotPosition;
+	TargetLocation.EndOfPathDirection = SlotDirection;
+	TargetLocation.TargetDistance = LaneLocation.LaneLength; // Go to end of lane
+	// Let's start moving toward the interaction a bit before the entry point.
+	constexpr float AnticipationDistance = 100.f;
+	TargetLocation.AnticipationDistance = AnticipationDistance;
+
+	WaitSlotLocation = &TargetLocation;
+	
+	return EStateTreeRunStatus::Running;
+}
+
+void FMassCrowdClaimWaitSlotTask::ExitState(FStateTreeExecutionContext& Context, const EStateTreeStateChangeType ChangeType, const FStateTreeTransitionResult& Transition)
+{
+	WaitSlotLocation = &TargetLocation;
+	
+	if (ChangeType != EStateTreeStateChangeType::Changed)
+	{
+		return;
+	}
+
+	const FMassStateTreeExecutionContext& MassContext = static_cast<FMassStateTreeExecutionContext&>(Context);
+	const FLWEntity Entity = MassContext.GetEntity();
+	
+	UMassCrowdSubsystem& CrowdSubsystem = Context.GetExternalItem(CrowdSubsystemHandle).GetMutable<UMassCrowdSubsystem>();
+
+	if (WaitingSlotIndex != INDEX_NONE)
+	{
+		CrowdSubsystem.ReleaseWaitingSlot(Entity, AcquiredLaneHandle, WaitingSlotIndex);
+	}
+	
+	TargetLocation.Reset();
+}
+
+EStateTreeRunStatus FMassCrowdClaimWaitSlotTask::Tick(FStateTreeExecutionContext& Context, const float DeltaTime)
+{
+	WaitSlotLocation = &TargetLocation;
+
+	return EStateTreeRunStatus::Running;
+}
diff --git a/Engine/Plugins/AI/MassCrowd/Source/MassCrowd/Private/Tasks/MassZoneGraphFindWanderTarget.cpp b/Engine/Plugins/AI/MassCrowd/Source/MassCrowd/Private/Tasks/MassZoneGraphFindWanderTarget.cpp
new file mode 100644
--- /dev/null
+++ b/Engine/Plugins/AI/MassCrowd/Source/MassCrowd/Private/Tasks/MassZoneGraphFindWanderTarget.cpp
@@ -0,0 +1,156 @@
+// Copyright Epic Games, Inc. All Rights Reserved.
+
+#include "Tasks/MassZoneGraphFindWanderTarget.h"
+#include "Templates/Tuple.h"
+#include "StateTreeExecutionContext.h"
+#include "MassStateTreeSubsystem.h"
+#include "ZoneGraphAnnotationSubsystem.h"
+#include "ZoneGraphSubsystem.h"
+#include "ZoneGraphQuery.h"
+#include "MassZoneGraphMovementFragments.h"
+#include "MassAIBehaviorTypes.h"
+#include "MassCrowdSettings.h"
+#include "MassCrowdSubsystem.h"
+#include "MassStateTreeExecutionContext.h"
+
+FMassZoneGraphFindWanderTarget::FMassZoneGraphFindWanderTarget()
+{
+}
+
+EStateTreeRunStatus FMassZoneGraphFindWanderTarget::EnterState(FStateTreeExecutionContext& Context, const EStateTreeStateChangeType ChangeType, const FStateTreeTransitionResult& Transition)
+{
+	const FMassStateTreeExecutionContext& MassContext = static_cast<FMassStateTreeExecutionContext&>(Context);
+
+	const FMassZoneGraphLaneLocationFragment& LaneLocation = Context.GetExternalItem(LocationHandle).Get<FMassZoneGraphLaneLocationFragment>();
+	const UZoneGraphSubsystem& ZoneGraphSubsystem = Context.GetExternalItem(ZoneGraphSubsystemHandle).Get<UZoneGraphSubsystem>();
+	UZoneGraphAnnotationSubsystem& ZoneGraphAnnotationSubsystem = Context.GetExternalItem(ZoneGraphAnnotationSubsystemHandle).GetMutable<UZoneGraphAnnotationSubsystem>();
+	const UMassCrowdSubsystem& MassCrowdSubsystem = Context.GetExternalItem(MassCrowdSubsystemHandle).Get<UMassCrowdSubsystem>();
+
+	bool bDisplayDebug = false;
+#if WITH_MASS_DEBUG
+	bDisplayDebug = UE::MassDebug::IsDebuggingEntity(MassContext.GetEntity());
+#endif // WITH_MASS_DEBUG
+
+	if (!LaneLocation.LaneHandle.IsValid())
+	{
+		MASSBEHAVIOR_LOG(Error, TEXT("Invalid lane location."));
+		return EStateTreeRunStatus::Failed;
+	}
+			
+	const FZoneGraphStorage* ZoneGraphStorage = ZoneGraphSubsystem.GetZoneGraphStorage(LaneLocation.LaneHandle.DataHandle);
+	if (!ZoneGraphStorage)
+	{
+		MASSBEHAVIOR_LOG(Error, TEXT("Missing ZoneGraph Storage for current lane %s."), *LaneLocation.LaneHandle.ToString());
+		return EStateTreeRunStatus::Failed;
+	}
+
+	const float MoveDistance = GetDefault<UMassCrowdSettings>()->GetMoveDistance();
+
+	TargetLocation.LaneHandle = LaneLocation.LaneHandle;
+	TargetLocation.TargetDistance = LaneLocation.DistanceAlongLane + MoveDistance;
+	TargetLocation.NextExitLinkType = EZoneLaneLinkType::None;
+	TargetLocation.NextLaneHandle.Reset();
+	TargetLocation.bMoveReverse = false;
+	TargetLocation.EndOfPathIntent = EMassMovementAction::Move;
+
+	if (bDisplayDebug)
+	{
+		MASSBEHAVIOR_LOG(Log, TEXT("Find wander target."));
+	}
+
+	EStateTreeRunStatus Status = EStateTreeRunStatus::Running;
+	
+	// When close to end of a lane, choose next lane.
+	if (TargetLocation.TargetDistance > LaneLocation.LaneLength)
+	{
+		TargetLocation.TargetDistance = FMath::Min(TargetLocation.TargetDistance, LaneLocation.LaneLength);
+
+		typedef TTuple<const FZoneGraphLinkedLane, const float> FBranchingCandidate;
+		TArray<FBranchingCandidate, TInlineAllocator<8>> Candidates;
+		float CombinedWeight = 0.f;
+
+		auto FindCandidates = [this, &ZoneGraphAnnotationSubsystem, &MassCrowdSubsystem, ZoneGraphStorage, LaneLocation, &Candidates, &CombinedWeight](const EZoneLaneLinkType Type)-> bool
+		{
+			TArray<FZoneGraphLinkedLane> LinkedLanes;
+			UE::ZoneGraph::Query::GetLinkedLanes(*ZoneGraphStorage, LaneLocation.LaneHandle, Type, EZoneLaneLinkFlags::All, EZoneLaneLinkFlags::None, LinkedLanes);
+
+			for (const FZoneGraphLinkedLane& LinkedLane : LinkedLanes)
+			{
+				// Apply tag filter
+				const FZoneGraphTagMask BehaviorTags = ZoneGraphAnnotationSubsystem.GetAnnotationTags(LinkedLane.DestLane);
+				if (!AllowedBehaviorTags.Pass(BehaviorTags))
+				{
+					continue;
+				}
+
+				// Add new candidate with its selection weight based on density
+				const FZoneGraphTagMask& LaneTagMask = ZoneGraphStorage->Lanes[LinkedLane.DestLane.Index].Tags;
+				const float Weight = MassCrowdSubsystem.GetDensityWeight(LinkedLane.DestLane, LaneTagMask);
+				CombinedWeight += Weight;
+				Candidates.Add(MakeTuple(LinkedLane, CombinedWeight));
+			}
+
+			return !Candidates.IsEmpty();
+		};
+
+		if (FindCandidates(EZoneLaneLinkType::Outgoing))
+		{
+			TargetLocation.NextExitLinkType = EZoneLaneLinkType::Outgoing;
+		}
+		else
+		{
+			// Could not continue, try to switch to an adjacent lane.
+			// @todo: we could try to do something smarter here so that agents do not clump up. May need to have some heuristic,
+			//		  i.e. at intersections it looks better to switch lane immediately, with flee, it looks better to vary the location randomly.
+			TargetLocation.TargetDistance = LaneLocation.DistanceAlongLane;
+
+			// Try adjacent lanes
+			if (FindCandidates(EZoneLaneLinkType::Adjacent))
+			{
+				// Found adjacent lane, choose it once followed the short path. Keeping the random offset from above,
+				// so that all agents dont follow until the end of the path to turn.
+				TargetLocation.NextExitLinkType = EZoneLaneLinkType::Adjacent;
+			}
+		}
+
+		if (Candidates.IsEmpty())
+		{
+			// Could not find next lane, fail.
+			TargetLocation.Reset();
+			Status = EStateTreeRunStatus::Failed;
+		}
+		else
+		{
+			// Select new lane based on the weight of each candidates
+			const float Rand = FMath::RandRange(0.f, CombinedWeight);
+			for (const FBranchingCandidate& Candidate : Candidates)
+			{
+				const float CandidateCombinedWeight = Candidate.Get<1>();
+				if (Rand < CandidateCombinedWeight)
+				{
+					const FZoneGraphLinkedLane& LinkedLane = Candidate.Get<0>();
+					TargetLocation.NextLaneHandle = LinkedLane.DestLane;
+					break;
+				}
+			}
+		}
+	}
+
+	WanderTargetLocation = &TargetLocation;
+	
+	return Status;
+}
+
+void FMassZoneGraphFindWanderTarget::ExitState(FStateTreeExecutionContext& Context, const EStateTreeStateChangeType ChangeType, const FStateTreeTransitionResult& Transition)
+{
+	// Temp solution to make sure the target location is always up to date (will be replaced with automatic update).
+	WanderTargetLocation = &TargetLocation;
+}
+
+EStateTreeRunStatus FMassZoneGraphFindWanderTarget::Tick(FStateTreeExecutionContext& Context, const float DeltaTime)
+{
+	// Temp solution to make sure the target location is always up to date (will be replaced with automatic update).
+	WanderTargetLocation = &TargetLocation;
+	
+	return EStateTreeRunStatus::Running;
+}
diff --git a/Engine/Plugins/AI/MassCrowd/Source/MassCrowd/Public/MassCrowdFragments.h b/Engine/Plugins/AI/MassCrowd/Source/MassCrowd/Public/MassCrowdFragments.h
new file mode 100644
--- /dev/null
+++ b/Engine/Plugins/AI/MassCrowd/Source/MassCrowd/Public/MassCrowdFragments.h
@@ -0,0 +1,32 @@
+// Copyright Epic Games, Inc. All Rights Reserved.
+
+#pragma once
+
+#include "LWComponentTypes.h"
+#include "ZoneGraphTypes.h"
+#include "Containers/StaticArray.h"
+#include "MassLODManager.h"
+
+#include "MassCrowdFragments.generated.h"
+
+/**
+ * Special tag to differentiate the crowd from the rest of the other entities
+ * Should not contain any data, this is purely a tag
+ */
+USTRUCT()
+struct MASSCROWD_API FTagFragment_MassCrowd : public FComponentTag
+{
+	GENERATED_BODY()
+};
+
+/**
+ * Data fragment to store the last lane the agent was tracked on.
+ */
+USTRUCT()
+struct MASSCROWD_API FMassCrowdLaneTrackingFragment : public FLWComponentData
+{
+	GENERATED_BODY()
+	FZoneGraphLaneHandle TrackedLaneHandle;
+};
+
+
diff --git a/Engine/Plugins/AI/MassCrowd/Source/MassCrowd/Public/MassCrowdLaneDataRenderingComponent.h b/Engine/Plugins/AI/MassCrowd/Source/MassCrowd/Public/MassCrowdLaneDataRenderingComponent.h
new file mode 100644
--- /dev/null
+++ b/Engine/Plugins/AI/MassCrowd/Source/MassCrowd/Public/MassCrowdLaneDataRenderingComponent.h
@@ -0,0 +1,33 @@
+// Copyright Epic Games, Inc. All Rights Reserved.
+
+#pragma once
+
+#include "Components/PrimitiveComponent.h"
+#include "MassCrowdLaneDataRenderingComponent.generated.h"
+
+/**
+ * Primitive component that can be used to render runtime state of zone graph lanes (e.g. Opened|Closed, Density, etc.)
+ * The component must be added on a ZoneGraphData actor.
+ */
+UCLASS(editinlinenew, meta = (BlueprintSpawnableComponent), hidecategories = (Object, LOD, Lighting, VirtualTexture, Transform, HLOD, Collision, TextureStreaming, Mobile, Physics, Tags, AssetUserData, Activation, Cooking, Rendering, Navigation))
+class MASSCROWD_API UMassCrowdLaneDataRenderingComponent : public UPrimitiveComponent
+{
+	GENERATED_BODY()
+
+public:
+	UMassCrowdLaneDataRenderingComponent() = default;
+
+#if !UE_BUILD_SHIPPING && !UE_BUILD_TEST
+private:
+	virtual FPrimitiveSceneProxy* CreateSceneProxy() override;
+	virtual void OnRegister() override;
+	virtual void OnUnregister() override;
+	virtual FBoxSphereBounds CalcBounds(const FTransform& LocalToWorld) const override;
+	void DebugDrawOnCanvas(UCanvas* Canvas, APlayerController*) const;
+	FDelegateHandle DebugTextDrawingDelegateHandle;
+	FDelegateHandle OnLaneStateChangedDelegateHandle;
+#if WITH_EDITOR
+	FDelegateHandle OnLaneRenderSettingsChangedDelegateHandle;
+#endif // WITH_EDITOR
+#endif // !UE_BUILD_SHIPPING && !UE_BUILD_TEST
+};
\ No newline at end of file
diff --git a/Engine/Plugins/AI/MassCrowd/Source/MassCrowd/Public/MassCrowdMemberTrait.h b/Engine/Plugins/AI/MassCrowd/Source/MassCrowd/Public/MassCrowdMemberTrait.h
new file mode 100644
--- /dev/null
+++ b/Engine/Plugins/AI/MassCrowd/Source/MassCrowd/Public/MassCrowdMemberTrait.h
@@ -0,0 +1,18 @@
+﻿// Copyright Epic Games, Inc. All Rights Reserved.
+
+#pragma once
+
+#include "MassEntityTraitBase.h"
+#include "MassCrowdMemberTrait.generated.h"
+
+/**
+ * Trait that makes the agent aware of traffic.
+ */
+UCLASS(BlueprintType, EditInlineNew, CollapseCategories, meta = (DisplayName = "CrowdMember"))
+class MASSCROWD_API UMassCrowdMemberTrait : public UMassEntityTraitBase
+{
+	GENERATED_BODY()
+
+protected:
+	virtual void BuildTemplate(FMassEntityTemplateBuildContext& BuildContext, UWorld& World) const override;
+};
diff --git a/Engine/Plugins/AI/MassCrowd/Source/MassCrowd/Public/MassCrowdNavigationProcessor.h b/Engine/Plugins/AI/MassCrowd/Source/MassCrowd/Public/MassCrowdNavigationProcessor.h
new file mode 100644
--- /dev/null
+++ b/Engine/Plugins/AI/MassCrowd/Source/MassCrowd/Public/MassCrowdNavigationProcessor.h
@@ -0,0 +1,74 @@
+// Copyright Epic Games, Inc. All Rights Reserved.
+
+#pragma once
+
+#include "MassSignalProcessorBase.h"
+#include "MassTranslator.h"
+#include "NavigationProcessor.h"
+#include "MassCrowdNavigationProcessor.generated.h"
+
+class UZoneGraphSubsystem;
+class UMassCrowdSettings;
+class UZoneGraphAnnotationSubsystem;
+class UMassCrowdSubsystem;
+
+/** Processor that monitors when entities change lane and notify the MassCrowd subsystem. */
+UCLASS()
+class MASSCROWD_API UMassCrowdLaneTrackingSignalProcessor : public UMassSignalProcessorBase
+{
+	GENERATED_BODY()
+public:
+	UMassCrowdLaneTrackingSignalProcessor();
+
+protected:
+	virtual void Initialize(UObject& Owner) override;
+	virtual void ConfigureQueries() override;
+	virtual void SignalEntities(UEntitySubsystem& EntitySubsystem, FLWComponentSystemExecutionContext& Context, FMassSignalNameLookup& EntitySignals) override;
+
+	UPROPERTY(Transient)
+	UMassCrowdSubsystem* MassCrowdSubsystem;
+};
+
+/** Processors that cleans up the lane tracking on entity destruction. */
+UCLASS()
+class MASSCROWD_API UMassCrowdLaneTrackingDestructor : public UMassFragmentDestructor
+{
+	GENERATED_BODY()
+
+public:
+	UMassCrowdLaneTrackingDestructor();
+
+protected:
+	virtual void Initialize(UObject& Owner) override;
+	virtual void ConfigureQueries() override;
+	virtual void Execute(UEntitySubsystem& EntitySubsystem, FLWComponentSystemExecutionContext& Context) override;
+
+	UPROPERTY(Transient)
+	UMassCrowdSubsystem* MassCrowdSubsystem;
+
+	FLWComponentQuery EntityQuery;
+};
+
+UCLASS()
+class MASSCROWD_API UMassCrowdDynamicObstacleProcessor : public UMassDynamicObstacleProcessor
+{
+	GENERATED_BODY()
+
+public:
+	UMassCrowdDynamicObstacleProcessor();
+
+protected:
+	virtual void Initialize(UObject& Owner) override;
+
+	virtual void OnStop(FMassDynamicObstacleFragment& OutObstacle, const float Radius) override;
+	virtual void OnMove(FMassDynamicObstacleFragment& OutObstacle) override;
+
+	UPROPERTY(Transient)
+	UZoneGraphSubsystem* ZoneGraphSubsystem;
+	
+	UPROPERTY(Transient)
+	UZoneGraphAnnotationSubsystem* ZoneGraphAnnotationSubsystem;
+
+	UPROPERTY(Transient)
+	const UMassCrowdSettings* CrowdSettings;
+};
diff --git a/Engine/Plugins/AI/MassCrowd/Source/MassCrowd/Public/MassCrowdReplicationProcessor.h b/Engine/Plugins/AI/MassCrowd/Source/MassCrowd/Public/MassCrowdReplicationProcessor.h
new file mode 100644
--- /dev/null
+++ b/Engine/Plugins/AI/MassCrowd/Source/MassCrowd/Public/MassCrowdReplicationProcessor.h
@@ -0,0 +1,28 @@
+// Copyright Epic Games, Inc. All Rights Reserved.
+
+#pragma once
+
+#include "MassReplicationProcessorBase.h"
+
+#include "MassCrowdReplicationProcessor.generated.h"
+
+class UMassReplicationManager;
+class MassLODManager;
+class AMassCrowdClientBubbleInfo;
+class UWorld;
+
+/** Processor that handles replication and only runs on the server. It queries Mass entity fragments and sets those values when appropriate using the MassClientBubbleHandler. */
+UCLASS()
+class MASSCROWD_API UMassCrowdReplicationProcessor : public UMassReplicationProcessorBase
+{
+	GENERATED_BODY()
+public:
+	UMassCrowdReplicationProcessor();
+
+protected:
+	virtual void ConfigureQueries() override;
+	virtual void Initialize(UObject& Owner) override;
+	virtual void Execute(UEntitySubsystem& EntitySubsystem, FLWComponentSystemExecutionContext& Context) override;
+
+	void ProcessClientReplication(UEntitySubsystem& EntitySubsystem, FLWComponentSystemExecutionContext& Context);
+};
\ No newline at end of file
diff --git a/Engine/Plugins/AI/MassCrowd/Source/MassCrowd/Public/MassCrowdRepresentationProcessor.h b/Engine/Plugins/AI/MassCrowd/Source/MassCrowd/Public/MassCrowdRepresentationProcessor.h
new file mode 100644
--- /dev/null
+++ b/Engine/Plugins/AI/MassCrowd/Source/MassCrowd/Public/MassCrowdRepresentationProcessor.h
@@ -0,0 +1,74 @@
+// Copyright Epic Games, Inc. All Rights Reserved.
+
+#pragma once
+
+#include "MassRepresentationProcessor.h"
+
+#include "MassCrowdRepresentationProcessor.generated.h"
+
+/**
+ * Overridden representation processor to make it tied to the crowd via the requirements.
+ * It is also the base class for all the different type of crowd representation (Visualization & ServerSideRepresentation)
+ */
+UCLASS(meta = (DisplayName = "Mass Crowd Visualization"))
+class MASSCROWD_API UMassCrowdRepresentationProcessor : public UMassRepresentationProcessor
+{
+	GENERATED_BODY()
+
+public:
+	UMassCrowdRepresentationProcessor();
+
+protected:
+
+	/**
+	 * Initialization of this processor
+	 * @param Owner of this processor
+	 */
+	virtual void Initialize(UObject& Owner) override;
+
+	/** Configure the owned FLWComponentQuery instances to express processor's requirements */
+	virtual void ConfigureQueries() override;
+
+	/**
+	 * Enable/disable a spawned actor
+	 * @param EnabledType is the type of enabling to do on this actor
+	 * @param Actor is the actual actor to perform enabling type on
+	 * @param EntityIdx is the entity index currently processing
+	 * @param Context is the current light weight component execution context
+	 */
+	virtual void SetActorEnabled(const EActorEnabledType EnabledType, AActor& Actor, const int32 EntityIdx, FLWComponentSystemExecutionContext& Context) override;
+
+	/** 
+	 * Returns an actor of the template type and setup fragments values from it
+	 * @param MassAgent is the handle to the associated mass agent
+	 * @param ActorInfo is the fragment where we are going to store the actor pointer
+	 * @param Transform is the spatial information about where to spawn the actor 
+	 * @param TemplateActorIndex is the index of the type fetched with UMassRepresentationSubsystem::FindOrAddTemplateActor()
+	 * @param SpawnRequestHandle (in/out) In: previously requested spawn Out: newly requested spawn
+	 * @param Priority of this spawn request in comparison with the others, lower value means higher priority
+	 * @param Context of the execution from the entity sub system
+	 * @return the actor spawned
+	 */
+	virtual AActor* GetOrSpawnActor(const FMassHandle MassAgent, FDataFragment_Actor& ActorInfo, const FTransform& Transform, const int16 TemplateActorIndex, FMassHandle_ActorSpawnRequest& SpawnRequestHandle, const float Priority) override;
+
+	/**
+	 * Teleports the actor at the specified transform by preserving its velocity and without collision.
+	 * The destination will be adjusted to fit an existing capsule.
+	 * @param Transform is the new actor's transform 
+	 * @param Actor is the actual actor to teleport
+	 * @param Context is the current light weight component execution context
+	 */
+	virtual void TeleportActor(const FTransform& Transform, AActor& Actor, FLWComponentSystemExecutionContext& Context) override;
+
+	/**
+	 * Initializes velocity for each entity when just switch to use movement component, must be called within a ForEachEntityChunk
+	 * @param EntitySubsystem The subsystem in which all entities must be initialized
+	 * @param Context of the execution from the entity sub system
+	 */
+	void InitializeVelocity(UEntitySubsystem& EntitySubsystem, FLWComponentSystemExecutionContext& Context);
+
+	/** 
+	 * A dedicated query for processing entities owning a FDataFragment_CharacterMovementComponentWrapper
+	 */
+	FLWComponentQuery CharacterMovementEntitiesQuery_Conditional;
+};
\ No newline at end of file
diff --git a/Engine/Plugins/AI/MassCrowd/Source/MassCrowd/Public/MassCrowdRepresentationSubsystem.h b/Engine/Plugins/AI/MassCrowd/Source/MassCrowd/Public/MassCrowdRepresentationSubsystem.h
new file mode 100644
--- /dev/null
+++ b/Engine/Plugins/AI/MassCrowd/Source/MassCrowd/Public/MassCrowdRepresentationSubsystem.h
@@ -0,0 +1,21 @@
+// Copyright Epic Games, Inc. All Rights Reserved.
+
+#pragma once
+
+#include "MassRepresentationSubsystem.h"
+
+#include "MassCrowdRepresentationSubsystem.generated.h"
+
+/**
+ * Subsystem responsible for all visual of mass crowd agents, will handle actors spawning and static mesh instances
+ */
+UCLASS()
+class MASSCROWD_API UMassCrowdRepresentationSubsystem : public UMassRepresentationSubsystem
+{
+	GENERATED_BODY()
+
+protected:
+	// USubsystem BEGIN
+	virtual void Initialize(FSubsystemCollectionBase& Collection) override;
+	// USubsystem END
+};
\ No newline at end of file
diff --git a/Engine/Plugins/AI/MassCrowd/Source/MassCrowd/Public/MassCrowdServerRepresentationLODProcessor.h b/Engine/Plugins/AI/MassCrowd/Source/MassCrowd/Public/MassCrowdServerRepresentationLODProcessor.h
new file mode 100644
--- /dev/null
+++ b/Engine/Plugins/AI/MassCrowd/Source/MassCrowd/Public/MassCrowdServerRepresentationLODProcessor.h
@@ -0,0 +1,42 @@
+// Copyright Epic Games, Inc. All Rights Reserved.
+
+#pragma once
+
+#include "MassCrowdTypes.h"
+#include "MassCrowdFragments.h"
+#include "MassLODManager.h"
+#include "MassLODCalculator.h"
+#include "MassLODTickRateController.h"
+
+#include "MassCrowdServerRepresentationLODProcessor.generated.h"
+
+UCLASS(meta=(DisplayName="Crowd Simulation LOD"))
+class MASSCROWD_API UMassCrowdServerRepresentationLODProcessor : public UMassProcessor_LODBase
+{
+	GENERATED_BODY()
+
+public:
+	UMassCrowdServerRepresentationLODProcessor();
+
+protected:
+
+	virtual void ConfigureQueries() override;
+	virtual void Initialize(UObject& InOwner) override;
+	virtual void Execute(UEntitySubsystem& EntitySubsystem, FLWComponentSystemExecutionContext& Context) override;
+
+	/** Distance where each LOD becomes relevant */
+	UPROPERTY(EditAnywhere, Category = "Mass|LOD", config)
+	float LODDistance[EMassLOD::Max];
+
+	/** Hysteresis percentage on delta between the LOD distances */
+	UPROPERTY(EditAnywhere, Category = "Mass|LOD", meta = (ClampMin = "0.0", UIMin = "0.0"), config)
+	float BufferHysteresisOnDistancePercentage = 10.0f;
+
+	/** Maximum limit of entity per LOD */
+	UPROPERTY(EditAnywhere, Category = "Mass|LOD", config)
+	int32 LODMaxCount[EMassLOD::Max];
+
+	TMassLODCalculator<FLODDefaultLogic> LODCalculator;
+
+	FLWComponentQuery EntityQuery;
+};
diff --git a/Engine/Plugins/AI/MassCrowd/Source/MassCrowd/Public/MassCrowdServerRepresentationProcessor.h b/Engine/Plugins/AI/MassCrowd/Source/MassCrowd/Public/MassCrowdServerRepresentationProcessor.h
new file mode 100644
--- /dev/null
+++ b/Engine/Plugins/AI/MassCrowd/Source/MassCrowd/Public/MassCrowdServerRepresentationProcessor.h
@@ -0,0 +1,29 @@
+// Copyright Epic Games, Inc. All Rights Reserved.
+
+#pragma once
+
+#include "MassCrowdRepresentationProcessor.h"
+
+#include "MassCrowdServerRepresentationProcessor.generated.h"
+
+/**
+ * Overridden representation processor to make it tied to the crowd on the server via the requirements
+ * It is the counter part of the crowd visualization processor on the client.
+ */
+UCLASS(meta = (DisplayName = "Mass Crowd Server Representation"))
+class MASSCROWD_API UMassCrowdServerRepresentationProcessor : public UMassCrowdRepresentationProcessor
+{
+	GENERATED_BODY()
+
+public:
+	UMassCrowdServerRepresentationProcessor();
+
+protected:
+
+	/**
+	 * Execution method for this processor
+	 * @param EntitySubsystem is the system to execute the lambdas on each entity chunk
+	 * @param Context is the execution context to be passed when executing the lambdas
+	 */
+	virtual void Execute(UEntitySubsystem& EntitySubsystem, FLWComponentSystemExecutionContext& Context) override;
+};
\ No newline at end of file
diff --git a/Engine/Plugins/AI/MassCrowd/Source/MassCrowd/Public/MassCrowdServerRepresentationTrait.h b/Engine/Plugins/AI/MassCrowd/Source/MassCrowd/Public/MassCrowdServerRepresentationTrait.h
new file mode 100644
--- /dev/null
+++ b/Engine/Plugins/AI/MassCrowd/Source/MassCrowd/Public/MassCrowdServerRepresentationTrait.h
@@ -0,0 +1,23 @@
+// Copyright Epic Games, Inc. All Rights Reserved.
+
+#pragma once 
+
+#include "MassEntityTraitBase.h"
+#include "MassRepresentationTypes.h"
+#include "GameFramework/Actor.h"
+#include "MassCrowdServerRepresentationTrait.generated.h"
+
+
+UCLASS(BlueprintType, EditInlineNew, CollapseCategories, meta=(DisplayName="Crowd Server Representation"))
+class MASSCROWD_API UMassCrowdServerRepresentationTrait : public UMassEntityTraitBase
+{
+	GENERATED_BODY()
+
+protected:
+
+	virtual void BuildTemplate(FMassEntityTemplateBuildContext& BuildContext, UWorld& World) const override;
+
+	/** Actor class of this agent when spawned on server */
+	UPROPERTY(EditAnywhere, Category = "Mass|Visual")
+	TSubclassOf<AActor> TemplateActor;
+};
diff --git a/Engine/Plugins/AI/MassCrowd/Source/MassCrowd/Public/MassCrowdVisualizationLODProcessor.h b/Engine/Plugins/AI/MassCrowd/Source/MassCrowd/Public/MassCrowdVisualizationLODProcessor.h
new file mode 100644
--- /dev/null
+++ b/Engine/Plugins/AI/MassCrowd/Source/MassCrowd/Public/MassCrowdVisualizationLODProcessor.h
@@ -0,0 +1,22 @@
+// Copyright Epic Games, Inc. All Rights Reserved.
+
+#pragma once
+
+#include "MassVisualizationLODProcessor.h"
+#include "MassCrowdVisualizationLODProcessor.generated.h"
+
+/*
+ * Created a crowd version for visualization of the UMassVisualizationLODProcessor as we want to use a custom version of the LOD info
+ * that combines the visualization and simulation LOD.
+ */
+UCLASS(meta=(DisplayName="Crowd visualization LOD"))
+class MASSCROWD_API UMassCrowdVisualizationLODProcessor : public UMassVisualizationLODProcessor
+{
+	GENERATED_BODY()
+public:
+	UMassCrowdVisualizationLODProcessor();
+
+protected:
+	virtual void ConfigureQueries() override;
+	virtual void Execute(UEntitySubsystem& EntitySubsystem, FLWComponentSystemExecutionContext& Context) override;
+};
diff --git a/Engine/Plugins/AI/MassCrowd/Source/MassCrowd/Public/MassCrowdVisualizationProcessor.h b/Engine/Plugins/AI/MassCrowd/Source/MassCrowd/Public/MassCrowdVisualizationProcessor.h
new file mode 100644
--- /dev/null
+++ b/Engine/Plugins/AI/MassCrowd/Source/MassCrowd/Public/MassCrowdVisualizationProcessor.h
@@ -0,0 +1,38 @@
+// Copyright Epic Games, Inc. All Rights Reserved.
+
+#pragma once
+
+#include "MassCrowdRepresentationProcessor.h"
+
+#include "MassCrowdVisualizationProcessor.generated.h"
+
+/**
+ * Overridden visualization processor to make it tied to the crowd via the requirements
+ */
+UCLASS(meta = (DisplayName = "Mass Crowd Visualization"))
+class MASSCROWD_API UMassCrowdVisualizationProcessor : public UMassCrowdRepresentationProcessor
+{
+	GENERATED_BODY()
+public:
+	UMassCrowdVisualizationProcessor();
+
+protected:
+
+	/** Configure the owned FLWComponentQuery instances to express processor's requirements */
+	virtual void ConfigureQueries() override;
+
+	/**
+	 * Execution method for this processor
+	 * @param EntitySubsystem is the system to execute the lambdas on each entity chunk
+	 * @param Context is the execution context to be passed when executing the lambdas
+	 */
+	virtual void Execute(UEntitySubsystem& EntitySubsystem, FLWComponentSystemExecutionContext& Context) override;
+};
+
+UCLASS()
+class MASSCROWD_API UMassCrowdRepresentationFragmentDestructor : public UMassRepresentationFragmentDestructor
+{
+	GENERATED_BODY()
+protected:
+	virtual void Initialize(UObject& Owner) override;
+};
diff --git a/Engine/Plugins/AI/MassCrowd/Source/MassCrowd/Public/MassCrowdVisualizationTrait.h b/Engine/Plugins/AI/MassCrowd/Source/MassCrowd/Public/MassCrowdVisualizationTrait.h
new file mode 100644
--- /dev/null
+++ b/Engine/Plugins/AI/MassCrowd/Source/MassCrowd/Public/MassCrowdVisualizationTrait.h
@@ -0,0 +1,15 @@
+// Copyright Epic Games, Inc. All Rights Reserved.
+
+#pragma once 
+
+#include "MassVisualizationTrait.h"
+#include "MassCrowdVisualizationTrait.generated.h"
+
+UCLASS(BlueprintType, EditInlineNew, CollapseCategories, meta=(DisplayName="Crowd Visualization"))
+class MASSCROWD_API UMassCrowdVisualizationTrait : public UMassVisualizationTrait
+{
+	GENERATED_BODY()
+public:
+	UMassCrowdVisualizationTrait();
+
+};
diff --git a/Engine/Plugins/AI/MassCrowd/Source/MassCrowd/Public/Tasks/MassCrowdClaimWaitSlotTask.h b/Engine/Plugins/AI/MassCrowd/Source/MassCrowd/Public/Tasks/MassCrowdClaimWaitSlotTask.h
new file mode 100644
--- /dev/null
+++ b/Engine/Plugins/AI/MassCrowd/Source/MassCrowd/Public/Tasks/MassCrowdClaimWaitSlotTask.h
@@ -0,0 +1,42 @@
+﻿// Copyright Epic Games, Inc. All Rights Reserved.
+
+#pragma once
+
+#include "MassStateTreeTypes.h"
+#include "ZoneGraphTypes.h"
+#include "Tasks/MassZoneGraphPathFollowTask.h"
+#include "MassCrowdClaimWaitSlotTask.generated.h"
+
+struct FStateTreeExecutionContext;
+
+/**
+* Claim wait slot and expose slot position for path follow.
+*/
+USTRUCT(meta = (DisplayName = "Crowd Claim Wait Slot"))
+struct MASSCROWD_API FMassCrowdClaimWaitSlotTask : public FMassStateTreeTaskBase
+{
+	GENERATED_BODY()
+
+protected:
+	virtual EStateTreeRunStatus EnterState(FStateTreeExecutionContext& Context, const EStateTreeStateChangeType ChangeType, const FStateTreeTransitionResult& Transition) override;
+	virtual void ExitState(FStateTreeExecutionContext& Context, const EStateTreeStateChangeType ChangeType, const FStateTreeTransitionResult& Transition) override;
+	virtual EStateTreeRunStatus Tick(FStateTreeExecutionContext& Context, const float DeltaTime) override;
+
+	UPROPERTY(meta=(BaseStruct="MassZoneGraphLaneLocationFragment"))
+	FStateTreeExternalItemHandle LocationHandle;
+
+	UPROPERTY(meta=(BaseStruct="MassMoveTargetFragment"))
+	FStateTreeExternalItemHandle MoveTargetHandle;
+
+	UPROPERTY(meta=(BaseClass="MassCrowdSubsystem"))
+	FStateTreeExternalItemHandle CrowdSubsystemHandle;
+
+	UPROPERTY(EditAnywhere, Category = Parameters, meta=(Struct="MassZoneGraphTargetLocation"))
+	FStateTreeResultRef WaitSlotLocation;
+
+	FMassZoneGraphTargetLocation TargetLocation;
+
+	int32 WaitingSlotIndex = INDEX_NONE;
+	
+	FZoneGraphLaneHandle AcquiredLaneHandle;
+};
diff --git a/Engine/Plugins/AI/MassCrowd/Source/MassCrowd/Public/Tasks/MassZoneGraphFindWanderTarget.h b/Engine/Plugins/AI/MassCrowd/Source/MassCrowd/Public/Tasks/MassZoneGraphFindWanderTarget.h
new file mode 100644
--- /dev/null
+++ b/Engine/Plugins/AI/MassCrowd/Source/MassCrowd/Public/Tasks/MassZoneGraphFindWanderTarget.h
@@ -0,0 +1,45 @@
+﻿// Copyright Epic Games, Inc. All Rights Reserved.
+
+#pragma once
+
+#include "MassStateTreeTypes.h"
+#include "Tasks/MassZoneGraphPathFollowTask.h"
+#include "MassZoneGraphFindWanderTarget.generated.h"
+
+struct FStateTreeExecutionContext;
+
+/**
+ * Updates TargetLocation to a wander target based on the agents current location on ZoneGraph.
+ */
+USTRUCT(meta = (DisplayName = "ZG Find Wander Target"))
+struct MASSCROWD_API FMassZoneGraphFindWanderTarget : public FMassStateTreeTaskBase
+{
+	GENERATED_BODY()
+
+	FMassZoneGraphFindWanderTarget();
+
+protected:
+	virtual EStateTreeRunStatus EnterState(FStateTreeExecutionContext& Context, const EStateTreeStateChangeType ChangeType, const FStateTreeTransitionResult& Transition) override;
+	virtual void ExitState(FStateTreeExecutionContext& Context, const EStateTreeStateChangeType ChangeType, const FStateTreeTransitionResult& Transition) override;
+	virtual EStateTreeRunStatus Tick(FStateTreeExecutionContext& Context, const float DeltaTime) override;
+
+	UPROPERTY(meta=(BaseStruct="MassZoneGraphLaneLocationFragment"))
+	FStateTreeExternalItemHandle LocationHandle;
+
+	UPROPERTY(meta=(BaseClass="ZoneGraphSubsystem"))
+	FStateTreeExternalItemHandle ZoneGraphSubsystemHandle;
+
+	UPROPERTY(meta=(BaseClass="ZoneGraphAnnotationSubsystem"))
+	FStateTreeExternalItemHandle ZoneGraphAnnotationSubsystemHandle;
+
+	UPROPERTY(meta=(BaseClass="MassCrowdSubsystem"))
+	FStateTreeExternalItemHandle MassCrowdSubsystemHandle;
+
+	UPROPERTY(EditAnywhere, Category = Parameters)
+	FZoneGraphTagFilter AllowedBehaviorTags;
+
+	FMassZoneGraphTargetLocation TargetLocation;
+
+	UPROPERTY(EditAnywhere, Category = Parameters, meta=(Struct="MassZoneGraphTargetLocation"))
+	FStateTreeResultRef WanderTargetLocation;
+};
diff --git a/Engine/Plugins/AI/MassCrowd/Source/MassCrowd/Public/ZoneGraphCrowdLaneAnnotations.h b/Engine/Plugins/AI/MassCrowd/Source/MassCrowd/Public/ZoneGraphCrowdLaneAnnotations.h
new file mode 100644
--- /dev/null
+++ b/Engine/Plugins/AI/MassCrowd/Source/MassCrowd/Public/ZoneGraphCrowdLaneAnnotations.h
@@ -0,0 +1,63 @@
+// Copyright Epic Games, Inc. All Rights Reserved.
+
+#pragma once
+
+#include "MassCrowdSubsystem.h"
+#include "ZoneGraphAnnotationComponent.h"
+#include "ZoneGraphAnnotationTypes.h"
+
+#include "ZoneGraphCrowdLaneAnnotations.generated.h"
+
+/** Event indicating the new state of a lane. */
+USTRUCT()
+struct MASSCROWD_API FZoneGraphCrowdLaneStateChangeEvent : public FZoneGraphAnnotationEventBase
+{
+	GENERATED_BODY()
+
+	FZoneGraphCrowdLaneStateChangeEvent() = default;
+	FZoneGraphCrowdLaneStateChangeEvent(const FZoneGraphLaneHandle LaneHandle, const ECrowdLaneState NewState)
+		: Lane(LaneHandle), State(NewState)	{}
+
+	/** Affected lane. */
+	UPROPERTY()
+	FZoneGraphLaneHandle Lane;
+
+	/** New state. */
+	UPROPERTY()
+	ECrowdLaneState State = ECrowdLaneState::Opened;
+};
+
+/**
+ * Zone graph blocking behavior
+ */
+UCLASS(ClassGroup = AI, BlueprintType, meta = (BlueprintSpawnableComponent))
+class MASSCROWD_API UZoneGraphCrowdLaneAnnotations : public UZoneGraphAnnotationComponent
+{
+	GENERATED_BODY()
+
+protected:
+	virtual void PostSubsystemsInitialized() override;
+	virtual FZoneGraphTagMask GetAnnotationTags() const override;
+	virtual void HandleEvents(TConstArrayView<const UScriptStruct*> AllEventStructs, const FInstancedStructStream& Events) override;
+	virtual void TickAnnotation(const float DeltaTime, FZoneGraphAnnotationTagContainer& AnnotationTagContainer) override;
+
+#if !UE_BUILD_SHIPPING && !UE_BUILD_TEST
+	virtual void DebugDraw(FZoneGraphAnnotationSceneProxy* DebugProxy) override;
+	virtual void DebugDrawCanvas(UCanvas* Canvas, APlayerController*) override;
+#endif
+
+	/** Annotation Tag to mark a closed lane. */
+	UPROPERTY(EditAnywhere, Category = CrowdLane)
+	FZoneGraphTag CloseLaneTag;
+
+	/** Annotation Tag to mark a waiting lane. */
+	UPROPERTY(EditAnywhere, Category = CrowdLane)
+	FZoneGraphTag WaitingLaneTag;
+
+	/** Array of queued events. */
+	TArray<FZoneGraphCrowdLaneStateChangeEvent> StateChangeEvents;
+
+	/** Cached ZoneGraphSubsystem */
+	UPROPERTY(Transient)
+	UMassCrowdSubsystem* CrowdSubsystem = nullptr;
+};
diff --git a/Engine/Plugins/Runtime/MassEntity/Config/BaseMassEntity.ini b/Engine/Plugins/Runtime/MassEntity/Config/BaseMassEntity.ini
new file mode 100644
--- /dev/null
+++ b/Engine/Plugins/Runtime/MassEntity/Config/BaseMassEntity.ini
@@ -0,0 +1,3 @@
+[CoreRedirects]
++PackageRedirects=(OldName="/Script/AggregateTicking",NewName="/Script/MassEntity")
++PackageRedirects=(OldName="/Script/Pipe",NewName="/Script/MassEntity")
diff --git a/Engine/Plugins/Runtime/MassEntity/MassEntity.uplugin b/Engine/Plugins/Runtime/MassEntity/MassEntity.uplugin
new file mode 100644
--- /dev/null
+++ b/Engine/Plugins/Runtime/MassEntity/MassEntity.uplugin
@@ -0,0 +1,40 @@
+{
+	"FileVersion": 3,
+	"Version": 1,
+	"VersionName": "1.0",
+	"FriendlyName": "MassEntity",
+	"Description": "",
+	"Category": "Gameplay",
+	"CreatedBy": "",
+	"CreatedByURL": "",
+	"DocsURL": "",
+	"MarketplaceURL": "",
+	"SupportURL": "",
+	"CanContainContent": true,
+	"IsBetaVersion": false,
+	"IsExperimentalVersion": true,
+	"Installed": false,
+	"EnabledByDefault": false,
+	"Modules": [
+		{
+			"Name": "MassEntity",
+			"Type": "Runtime",
+			"LoadingPhase": "PreLoadingScreen"
+		},
+		{
+			"Name": "MassEntityTestSuite",
+			"Type": "Runtime",
+			"LoadingPhase": "Default"
+		},
+		{
+			"Name": "MassEntityEditor",
+			"Type": "Editor"
+		}
+	],
+	"Plugins": [
+		{
+			"Name": "StructUtils",
+			"Enabled": true
+		}
+	]
+}
\ No newline at end of file
