Non-unity and cooked Mass fixes

diff --git a/Engine/Plugins/AI/MassAI/MassAI.uplugin b/Engine/Plugins/AI/MassAI/MassAI.uplugin
--- a/Engine/Plugins/AI/MassAI/MassAI.uplugin
+++ b/Engine/Plugins/AI/MassAI/MassAI.uplugin
@@ -35,7 +35,7 @@
 		},
 		{
 			"Name": "MassAIDebug",
-			"Type": "DeveloperTool"
+			"Type": "UncookedOnly"
 		},
 		{
 			"Name": "MassAITestSuite",
diff --git a/Engine/Plugins/AI/MassAI/Source/MassAIDebug/Private/GameplayDebuggerCategory_Mass.cpp b/Engine/Plugins/AI/MassAI/Source/MassAIDebug/Private/GameplayDebuggerCategory_Mass.cpp
--- a/Engine/Plugins/AI/MassAI/Source/MassAIDebug/Private/GameplayDebuggerCategory_Mass.cpp
+++ b/Engine/Plugins/AI/MassAI/Source/MassAIDebug/Private/GameplayDebuggerCategory_Mass.cpp
@@ -19,6 +19,7 @@
 #include "Util/ColorConstants.h"
 #include "MassSimulationLOD.h"
 #include "CanvasItem.h"
+#include "Engine/World.h"
 
 //----------------------------------------------------------------------//
 //  FGameplayDebuggerCategory_Mass
diff --git a/Engine/Plugins/AI/MassAI/Source/MassAIDebug/Private/MassDebugStateTreeProcessor.cpp b/Engine/Plugins/AI/MassAI/Source/MassAIDebug/Private/MassDebugStateTreeProcessor.cpp
--- a/Engine/Plugins/AI/MassAI/Source/MassAIDebug/Private/MassDebugStateTreeProcessor.cpp
+++ b/Engine/Plugins/AI/MassAI/Source/MassAIDebug/Private/MassDebugStateTreeProcessor.cpp
@@ -6,6 +6,7 @@
 #include "MassStateTreeFragments.h"
 #include "MassCommonFragments.h"
 #include "VisualLogger/VisualLogger.h"
+#include "Engine/World.h"
 
 //----------------------------------------------------------------------//
 // UMassDebugStateTreeProcessor
diff --git a/Engine/Plugins/AI/MassAI/Source/MassAIMovement/Private/MassMovementModule.cpp b/Engine/Plugins/AI/MassAI/Source/MassAIMovement/Private/MassMovementModule.cpp
--- a/Engine/Plugins/AI/MassAI/Source/MassAIMovement/Private/MassMovementModule.cpp
+++ b/Engine/Plugins/AI/MassAI/Source/MassAIMovement/Private/MassMovementModule.cpp
@@ -5,34 +5,33 @@
 #include "UObject/CoreRedirects.h"
 
 
-class FMassMovementModule : public IMassMovementModule
+class FMassAIMovementModule : public IMassAIMovementModule
 {
 	/** IModuleInterface implementation */
 	virtual void StartupModule() override;
 	virtual void ShutdownModule() override;
 };
 
-IMPLEMENT_MODULE(FMassMovementModule, MassMovement)
+IMPLEMENT_MODULE(FMassAIMovementModule, MassAIMovement)
 
 
 
-void FMassMovementModule::StartupModule()
+void FMassAIMovementModule::StartupModule()
 {
 	// This code will execute after your module is loaded into memory (but after global variables are initialized, of course.)
 
 	TArray<FCoreRedirect> Redirects;
-	Redirects.Emplace(ECoreRedirectFlags::Type_Struct, TEXT("VelocityComponent"), TEXT("MassVelocityFragment"));
 	Redirects.Emplace(ECoreRedirectFlags::Type_Struct, TEXT("AvoidanceComponent"), TEXT("MassAvoidanceFragment"));
 	Redirects.Emplace(ECoreRedirectFlags::Type_Struct, TEXT("DataFragment_GridCellLocation"), TEXT("MassGridCellLocationFragment"));
 	Redirects.Emplace(ECoreRedirectFlags::Type_Struct, TEXT("DataFragment_PathFollow"), TEXT("MassPathFollowFragment"));
 	Redirects.Emplace(ECoreRedirectFlags::Type_Struct, TEXT("DataFragment_Wander"), TEXT("MassWanderFragment"));
 	Redirects.Emplace(ECoreRedirectFlags::Type_Struct, TEXT("DataFragment_MoveTo"), TEXT("MassMoveToFragment"));
 	
-	FCoreRedirects::AddRedirectList(Redirects, TEXT("MassMovement"));
+	FCoreRedirects::AddRedirectList(Redirects, TEXT("MassAIMovement"));
 }
 
 
-void FMassMovementModule::ShutdownModule()
+void FMassAIMovementModule::ShutdownModule()
 {
 	// This function may be called during shutdown to clean up your module.  For modules that support dynamic reloading,
 	// we call this function before unloading the module.
diff --git a/Engine/Plugins/AI/MassAI/Source/MassAIMovement/Public/IMassMovementModule.h b/Engine/Plugins/AI/MassAI/Source/MassAIMovement/Public/IMassMovementModule.h
--- a/Engine/Plugins/AI/MassAI/Source/MassAIMovement/Public/IMassMovementModule.h
+++ b/Engine/Plugins/AI/MassAI/Source/MassAIMovement/Public/IMassMovementModule.h
@@ -11,7 +11,7 @@
  * The public interface to this module.  In most cases, this interface is only public to sibling modules 
  * within this plugin.
  */
-class IMassMovementModule : public IModuleInterface
+class IMassAIMovementModule : public IModuleInterface
 {
 
 public:
@@ -22,9 +22,9 @@ public:
 	 *
 	 * @return Returns singleton instance, loading the module on demand if needed
 	 */
-	static inline IMassMovementModule& Get()
+	static inline IMassAIMovementModule& Get()
 	{
-		return FModuleManager::LoadModuleChecked<IMassMovementModule>( "MassMovement" );
+		return FModuleManager::LoadModuleChecked<IMassAIMovementModule>( "MassAIMovement" );
 	}
 
 	/**
@@ -34,7 +34,7 @@ public:
 	 */
 	static inline bool IsAvailable()
 	{
-		return FModuleManager::Get().IsModuleLoaded( "MassMovement" );
+		return FModuleManager::Get().IsModuleLoaded( "MassAIMovement" );
 	}
 };
 
diff --git a/Engine/Plugins/AI/MassCrowd/MassCrowd.uplugin b/Engine/Plugins/AI/MassCrowd/MassCrowd.uplugin
--- a/Engine/Plugins/AI/MassCrowd/MassCrowd.uplugin
+++ b/Engine/Plugins/AI/MassCrowd/MassCrowd.uplugin
@@ -20,7 +20,7 @@
 			"Name": "MassCrowd",
 			"Type": "Runtime",
 			"LoadingPhase": "Default"
-		},
+		}
 	],
 	"Plugins": [
 		{
diff --git a/Engine/Plugins/AI/MassCrowd/Source/MassCrowd/Public/MassCrowdReplicatedAgent.h b/Engine/Plugins/AI/MassCrowd/Source/MassCrowd/Public/MassCrowdReplicatedAgent.h
--- a/Engine/Plugins/AI/MassCrowd/Source/MassCrowd/Public/MassCrowdReplicatedAgent.h
+++ b/Engine/Plugins/AI/MassCrowd/Source/MassCrowd/Public/MassCrowdReplicatedAgent.h
@@ -5,6 +5,7 @@
 #include "MassReplicationTransformHandlers.h"
 #include "MassReplicationTypes.h"
 #include "MassClientBubbleHandler.h"
+#include "MassReplicationPathHandlers.h"
 
 #include "MassCrowdReplicatedAgent.generated.h"
 
diff --git a/Engine/Plugins/Runtime/MassEntity/MassEntity.uplugin b/Engine/Plugins/Runtime/MassEntity/MassEntity.uplugin
--- a/Engine/Plugins/Runtime/MassEntity/MassEntity.uplugin
+++ b/Engine/Plugins/Runtime/MassEntity/MassEntity.uplugin
@@ -23,7 +23,7 @@
 		},
 		{
 			"Name": "MassEntityTestSuite",
-			"Type": "Runtime",
+			"Type": "UncookedOnly",
 			"LoadingPhase": "Default"
 		},
 		{
diff --git a/Engine/Plugins/Runtime/MassEntity/Source/MassEntity/Private/MassProcessorDependencySolver.cpp b/Engine/Plugins/Runtime/MassEntity/Source/MassEntity/Private/MassProcessorDependencySolver.cpp
--- a/Engine/Plugins/Runtime/MassEntity/Source/MassEntity/Private/MassProcessorDependencySolver.cpp
+++ b/Engine/Plugins/Runtime/MassEntity/Source/MassEntity/Private/MassProcessorDependencySolver.cpp
@@ -6,6 +6,7 @@
 #include "Logging/MessageLog.h"
 #include "HAL/FileManager.h"
 #include "Misc/Paths.h"
+#include "Engine/World.h"
 
 #define LOCTEXT_NAMESPACE "Pipe"
 
diff --git a/Engine/Plugins/Runtime/MassGameplay/Source/MassMovement/Private/MassMovementModule.cpp b/Engine/Plugins/Runtime/MassGameplay/Source/MassMovement/Private/MassMovementModule.cpp
--- a/Engine/Plugins/Runtime/MassGameplay/Source/MassMovement/Private/MassMovementModule.cpp
+++ b/Engine/Plugins/Runtime/MassGameplay/Source/MassMovement/Private/MassMovementModule.cpp
@@ -22,11 +22,6 @@ void FMassMovementModule::StartupModule()
 
 	TArray<FCoreRedirect> Redirects;
 	Redirects.Emplace(ECoreRedirectFlags::Type_Struct, TEXT("VelocityComponent"), TEXT("MassVelocityFragment"));
-	Redirects.Emplace(ECoreRedirectFlags::Type_Struct, TEXT("AvoidanceComponent"), TEXT("MassAvoidanceFragment"));
-	Redirects.Emplace(ECoreRedirectFlags::Type_Struct, TEXT("DataFragment_GridCellLocation"), TEXT("MassGridCellLocationFragment"));
-	Redirects.Emplace(ECoreRedirectFlags::Type_Struct, TEXT("DataFragment_PathFollow"), TEXT("MassPathFollowFragment"));
-	Redirects.Emplace(ECoreRedirectFlags::Type_Struct, TEXT("DataFragment_Wander"), TEXT("MassWanderFragment"));
-	Redirects.Emplace(ECoreRedirectFlags::Type_Struct, TEXT("DataFragment_MoveTo"), TEXT("MassMoveToFragment"));
 	
 	FCoreRedirects::AddRedirectList(Redirects, TEXT("MassMovement"));
 }
