Updated all the WITH_xxxx_DEBUG Mass' macros

diff --git a/Engine/Plugins/AI/MassAI/Source/MassAIBehavior/Private/MassLookAtProcessors.cpp b/Engine/Plugins/AI/MassAI/Source/MassAIBehavior/Private/MassLookAtProcessors.cpp
--- a/Engine/Plugins/AI/MassAI/Source/MassAIBehavior/Private/MassLookAtProcessors.cpp
+++ b/Engine/Plugins/AI/MassAI/Source/MassAIBehavior/Private/MassLookAtProcessors.cpp
@@ -143,10 +143,10 @@ void UMassLookAtProcessor::Execute(UMassEntitySubsystem& EntitySubsystem, FLWCom
 
 				bool bDisplayDebug = false;
 				const FLWEntity Entity = Context.GetEntity(i);
-	#if WITH_MASS_DEBUG
+	#if WITH_MASSGAMEPLAY_DEBUG
 				FColor EntityColor = FColor::White;
 				bDisplayDebug = UE::MassDebug::IsDebuggingEntity(Entity, &EntityColor);
-	#endif // WITH_MASS_DEBUG
+	#endif // WITH_MASSGAMEPLAY_DEBUG
 
 				// Update gaze target when current cycle is finished.
 				if (LookAt.RandomGazeMode != EMassLookAtGazeMode::None)
@@ -221,14 +221,14 @@ void UMassLookAtProcessor::Execute(UMassEntitySubsystem& EntitySubsystem, FLWCom
 				// Clamp
 				LookAt.Direction = UE::MassBehavior::ClampDirectionToXAxisCone(LookAt.Direction, FMath::DegreesToRadians(AngleThresholdInDegrees));
 
-	#if WITH_MASS_DEBUG && UNSAFE_FOR_MT
+	#if WITH_MASSGAMEPLAY_DEBUG && UNSAFE_FOR_MT
 				if (bDisplayDebug)
 				{
 					const FVector Origin = TransformFragment.GetTransform().GetLocation() + FVector(0.f,0.f,DebugZOffset);
 					const FVector Dest = Origin + 100.f*TransformFragment.GetTransform().TransformVector(LookAt.Direction);
 					UE_VLOG_ARROW(this, LogMassBehavior, Display, Origin, Dest, EntityColor, TEXT(""));
 				}
-	#endif // WITH_MASS_DEBUG
+	#endif // WITH_MASSGAMEPLAY_DEBUG
 			}
 		});
 }
@@ -339,7 +339,7 @@ void UMassLookAtProcessor::UpdateLookAtTrajectory(const FTransform& Transform, c
 		LookAt.Direction = Transform.InverseTransformVector(NewGlobalDirection);
 		LookAt.Direction.Z = 0.0f;
 						
-#if WITH_MASS_DEBUG && UNSAFE_FOR_MT
+#if WITH_MASSGAMEPLAY_DEBUG && UNSAFE_FOR_MT
 		if (bDisplayDebug)
 		{
 			const FVector ZOffset(0.f,0.f,DebugZOffset);
@@ -360,7 +360,7 @@ void UMassLookAtProcessor::UpdateLookAtTrackedEntity(const UMassEntitySubsystem&
 			const FVector NewGlobalDirection = (TargetTransform->GetTransform().GetLocation() - AgentPosition).GetSafeNormal();
 			LookAt.Direction = Transform.InverseTransformVector(NewGlobalDirection);
 
-#if WITH_MASS_DEBUG && UNSAFE_FOR_MT
+#if WITH_MASSGAMEPLAY_DEBUG && UNSAFE_FOR_MT
 			if (bDisplayDebug)
 			{
 				const FVector ZOffset(0.f,0.f,DebugZOffset);
@@ -386,7 +386,7 @@ bool UMassLookAtProcessor::UpdateGazeTrackedEntity(const UMassEntitySubsystem& E
 
 			bHasTarget = true;
 			
-#if WITH_MASS_DEBUG && UNSAFE_FOR_MT
+#if WITH_MASSGAMEPLAY_DEBUG && UNSAFE_FOR_MT
 			if (bDisplayDebug)
 			{
 				const FVector ZOffset(0.f,0.f,DebugZOffset);
@@ -476,7 +476,7 @@ void UMassLookAtProcessor::BuildTrajectory(const FMassZoneGraphLaneLocationFragm
 		}
 	}
 
-#if WITH_MASS_DEBUG && UNSAFE_FOR_MT
+#if WITH_MASSGAMEPLAY_DEBUG && UNSAFE_FOR_MT
 	if (bDisplayDebug)
 	{
 		const FVector ZOffset(0,0,35);
@@ -517,7 +517,7 @@ void UMassLookAtProcessor::BuildTrajectory(const FMassZoneGraphLaneLocationFragm
 				TEXT("D:%.1f"), CurrPoint.DistanceAlongLane.Get());
 		}
 	}
-#endif // WITH_MASS_DEBUG
+#endif // WITH_MASSGAMEPLAY_DEBUG
 }
 
 #undef UNSAFE_FOR_MT
\ No newline at end of file
diff --git a/Engine/Plugins/AI/MassAI/Source/MassAIBehavior/Private/Tasks/MassZoneGraphFindEscapeTarget.cpp b/Engine/Plugins/AI/MassAI/Source/MassAIBehavior/Private/Tasks/MassZoneGraphFindEscapeTarget.cpp
--- a/Engine/Plugins/AI/MassAI/Source/MassAIBehavior/Private/Tasks/MassZoneGraphFindEscapeTarget.cpp
+++ b/Engine/Plugins/AI/MassAI/Source/MassAIBehavior/Private/Tasks/MassZoneGraphFindEscapeTarget.cpp
@@ -22,9 +22,9 @@ EStateTreeRunStatus FMassZoneGraphFindEscapeTarget::EnterState(FStateTreeExecuti
 
 	EscapeTargetLocation = nullptr;
 
-#if WITH_MASS_DEBUG
+#if WITH_MASSGAMEPLAY_DEBUG
 	bDisplayDebug = UE::MassDebug::IsDebuggingEntity(MassContext.GetEntity());
-#endif // WITH_MASS_DEBUG
+#endif // WITH_MASSGAMEPLAY_DEBUG
 
 	const FMassZoneGraphLaneLocationFragment& LaneLocation = Context.GetExternalItem(LocationHandle).Get<FMassZoneGraphLaneLocationFragment>();
 	UZoneGraphSubsystem& ZoneGraphSubsystem = Context.GetExternalItem(ZoneGraphSubsystemHandle).GetMutable<UZoneGraphSubsystem>();
diff --git a/Engine/Plugins/AI/MassAI/Source/MassAIBehavior/Private/Tasks/MassZoneGraphPathFollowTask.cpp b/Engine/Plugins/AI/MassAI/Source/MassAIBehavior/Private/Tasks/MassZoneGraphPathFollowTask.cpp
--- a/Engine/Plugins/AI/MassAI/Source/MassAIBehavior/Private/Tasks/MassZoneGraphPathFollowTask.cpp
+++ b/Engine/Plugins/AI/MassAI/Source/MassAIBehavior/Private/Tasks/MassZoneGraphPathFollowTask.cpp
@@ -33,9 +33,9 @@ bool FMassZoneGraphPathFollowTask::RequestPath(FMassStateTreeExecutionContext& C
 	}
 
 	bool bDisplayDebug = false;
-#if WITH_MASS_DEBUG
+#if WITH_MASSGAMEPLAY_DEBUG
 	bDisplayDebug = UE::MassDebug::IsDebuggingEntity(Context.GetEntity());
-#endif // WITH_MASS_DEBUG
+#endif // WITH_MASSGAMEPLAY_DEBUG
 
 	if (RequestedTargetLocation.LaneHandle != LaneLocation.LaneHandle)
 	{
@@ -87,9 +87,9 @@ EStateTreeRunStatus FMassZoneGraphPathFollowTask::EnterState(FStateTreeExecution
 	FMassStateTreeExecutionContext& MassContext = static_cast<FMassStateTreeExecutionContext&>(Context);
 
 	bool bDisplayDebug = false;
-#if WITH_MASS_DEBUG
+#if WITH_MASSGAMEPLAY_DEBUG
 	bDisplayDebug = UE::MassDebug::IsDebuggingEntity(MassContext.GetEntity());
-#endif // WITH_MASS_DEBUG
+#endif // WITH_MASSGAMEPLAY_DEBUG
 	if (bDisplayDebug)
 	{
 		MASSBEHAVIOR_LOG(Verbose, TEXT("enterstate."));
@@ -129,9 +129,9 @@ EStateTreeRunStatus FMassZoneGraphPathFollowTask::Tick(FStateTreeExecutionContex
 	FMassStateTreeExecutionContext& MassContext = static_cast<FMassStateTreeExecutionContext&>(Context);
 
 	bool bDisplayDebug = false;
-#if WITH_MASS_DEBUG
+#if WITH_MASSGAMEPLAY_DEBUG
 	bDisplayDebug = UE::MassDebug::IsDebuggingEntity(MassContext.GetEntity());
-#endif // WITH_MASS_DEBUG
+#endif // WITH_MASSGAMEPLAY_DEBUG
 	if (bDisplayDebug)
 	{
 		MASSBEHAVIOR_LOG(Verbose, TEXT("tick"));
diff --git a/Engine/Plugins/AI/MassAI/Source/MassAIBehavior/Public/MassAIBehaviorTypes.h b/Engine/Plugins/AI/MassAI/Source/MassAIBehavior/Public/MassAIBehaviorTypes.h
--- a/Engine/Plugins/AI/MassAI/Source/MassAIBehavior/Public/MassAIBehaviorTypes.h
+++ b/Engine/Plugins/AI/MassAI/Source/MassAIBehavior/Public/MassAIBehaviorTypes.h
@@ -19,15 +19,15 @@ namespace UE::Mass::ProcessorGroupNames
  * These macros should be used to standardize the output format and simplify code at call site.
  * They could also easily be changed from UE_(C)VLOG to UE_(C)VLOG_UELOG (or any other implementation) in one go.
  *  e.g. the following:
- *		#if WITH_MASS_DEBUG
+ *		#if WITH_MASSGAMEPLAY_DEBUG
  *			const FMassStateTreeExecutionContext& MassContext = static_cast<FMassStateTreeExecutionContext&>(Context);
  *			UE_VLOG(MassContext.GetOwner(), LogMassBehavior, Log, TEXT("Entity [%s]: Starting action: %s"), *MassContext.GetEntity().DebugGetDescription(), *StaticEnum<ESomeActionEnum>()->GetValueAsString(SomeActionEnumValue));
- *		#endif // WITH_MASS_DEBUG
+ *		#endif // WITH_MASSGAMEPLAY_DEBUG
  *
  *	could be replaced by:
  *		MASSBEHAVIOR_CLOG(bDisplayDebug, Log, TEXT("Starting action: %s"), *StaticEnum<ESomeActionEnum>()->GetValueAsString(SomeActionEnumValue));
  */
-#if WITH_MASS_DEBUG
+#if WITH_MASSGAMEPLAY_DEBUG
 #define MASSBEHAVIOR_LOG(Verbosity, Format, ...) UE_VLOG(static_cast<FMassStateTreeExecutionContext&>(Context).GetOwner(), LogMassBehavior, Verbosity, \
 	TEXT("Entity [%s][%s] ") Format, *static_cast<FMassStateTreeExecutionContext&>(Context).GetEntity().DebugGetDescription(), *StaticStruct()->GetName(), ##__VA_ARGS__)
 #define MASSBEHAVIOR_CLOG(Condition, Verbosity, Format, ...) UE_CVLOG((Condition), static_cast<FMassStateTreeExecutionContext&>(Context).GetOwner(), LogMassBehavior, Verbosity, \
@@ -35,4 +35,4 @@ namespace UE::Mass::ProcessorGroupNames
 #else
 #define MASSBEHAVIOR_LOG(Verbosity, Format, ...)
 #define MASSBEHAVIOR_CLOG(Condition, Verbosity, Format, ...)
-#endif // WITH_MASS_DEBUG
+#endif // WITH_MASSGAMEPLAY_DEBUG
diff --git a/Engine/Plugins/AI/MassAI/Source/MassAIDebug/Private/GameplayDebuggerCategory_Mass.cpp b/Engine/Plugins/AI/MassAI/Source/MassAIDebug/Private/GameplayDebuggerCategory_Mass.cpp
--- a/Engine/Plugins/AI/MassAI/Source/MassAIDebug/Private/GameplayDebuggerCategory_Mass.cpp
+++ b/Engine/Plugins/AI/MassAI/Source/MassAIDebug/Private/GameplayDebuggerCategory_Mass.cpp
@@ -2,7 +2,7 @@
 
 #include "GameplayDebuggerCategory_Mass.h"
 
-#if WITH_GAMEPLAY_DEBUGGER && WITH_MASS_DEBUG
+#if WITH_GAMEPLAY_DEBUGGER && WITH_MASSGAMEPLAY_DEBUG
 #include "EntityView.h"
 #include "GameplayDebuggerCategoryReplicator.h"
 #include "GameplayDebuggerPlayerManager.h"
@@ -312,7 +312,7 @@ void FGameplayDebuggerCategory_Mass::CollectData(APlayerController* OwnerPC, AAc
 		const float CurrentTime = World->GetTimeSeconds();
 		
 		UMassStateTreeSubsystem* MassStateTreeSubsystem = World->GetSubsystem<UMassStateTreeSubsystem>();
-		if (MassStateTreeSubsystem != nullptr)
+		if (MassStateTreeSubsystem && EntitySystem)
 		{
 			FLWComponentSystemExecutionContext Context(0.0f);
 		
@@ -519,7 +519,7 @@ void FGameplayDebuggerCategory_Mass::CollectData(APlayerController* OwnerPC, AAc
 			});
 		}
 
-		if (bShowNearEntityAvoidance)
+		if (bShowNearEntityAvoidance && EntitySystem)
 		{
 			FLWComponentQuery EntityColliderQuery;
 			EntityColliderQuery.AddRequirement<FMassAvoidanceColliderFragment>(ELWComponentAccess::ReadOnly);
@@ -667,5 +667,5 @@ void FGameplayDebuggerCategory_Mass::DrawData(APlayerController* OwnerPC, FGamep
 
 	FGameplayDebuggerCategory::DrawData(OwnerPC, CanvasContext);
 }
-#endif // WITH_GAMEPLAY_DEBUGGER && WITH_MASS_DEBUG
+#endif // WITH_GAMEPLAY_DEBUGGER && WITH_MASSGAMEPLAY_DEBUG
 
diff --git a/Engine/Plugins/AI/MassAI/Source/MassAIDebug/Private/MassAIDebugModule.cpp b/Engine/Plugins/AI/MassAI/Source/MassAIDebug/Private/MassAIDebugModule.cpp
--- a/Engine/Plugins/AI/MassAI/Source/MassAIDebug/Private/MassAIDebugModule.cpp
+++ b/Engine/Plugins/AI/MassAI/Source/MassAIDebug/Private/MassAIDebugModule.cpp
@@ -20,7 +20,7 @@ IMPLEMENT_MODULE(FMassAIDebug, MassAIDebug)
 
 void FMassAIDebug::StartupModule()
 {
-#if WITH_GAMEPLAY_DEBUGGER && WITH_MASS_DEBUG
+#if WITH_GAMEPLAY_DEBUGGER && WITH_MASSGAMEPLAY_DEBUG
 	IGameplayDebugger& GameplayDebuggerModule = IGameplayDebugger::Get();
 	GameplayDebuggerModule.RegisterCategory("Mass", IGameplayDebugger::FOnGetCategory::CreateStatic(&FGameplayDebuggerCategory_Mass::MakeInstance), EGameplayDebuggerCategoryState::EnabledInGameAndSimulate);
 	GameplayDebuggerModule.NotifyCategoriesChanged();
@@ -29,7 +29,7 @@ void FMassAIDebug::StartupModule()
 
 void FMassAIDebug::ShutdownModule()
 {
-#if WITH_GAMEPLAY_DEBUGGER && WITH_MASS_DEBUG
+#if WITH_GAMEPLAY_DEBUGGER && WITH_MASSGAMEPLAY_DEBUG
 	if (IGameplayDebugger::IsAvailable())
 	{
 		IGameplayDebugger& GameplayDebuggerModule = IGameplayDebugger::Get();
diff --git a/Engine/Plugins/AI/MassAI/Source/MassAIDebug/Private/MassDebugStateTreeProcessor.cpp b/Engine/Plugins/AI/MassAI/Source/MassAIDebug/Private/MassDebugStateTreeProcessor.cpp
--- a/Engine/Plugins/AI/MassAI/Source/MassAIDebug/Private/MassDebugStateTreeProcessor.cpp
+++ b/Engine/Plugins/AI/MassAI/Source/MassAIDebug/Private/MassDebugStateTreeProcessor.cpp
@@ -25,7 +25,7 @@ void UMassDebugStateTreeProcessor::ConfigureQueries()
 
 void UMassDebugStateTreeProcessor::Execute(UMassEntitySubsystem& EntitySubsystem, FLWComponentSystemExecutionContext& Context)
 {
-#if WITH_MASS_DEBUG
+#if WITH_MASSGAMEPLAY_DEBUG
 	UWorld* World = GetWorld();
 	if (World == nullptr)
 	{
@@ -100,5 +100,5 @@ void UMassDebugStateTreeProcessor::Execute(UMassEntitySubsystem& EntitySubsystem
 				}
 			}
 		});
-	#endif // WITH_MASS_DEBUG
+	#endif // WITH_MASSGAMEPLAY_DEBUG
 }
diff --git a/Engine/Plugins/AI/MassAI/Source/MassAIDebug/Public/GameplayDebuggerCategory_Mass.h b/Engine/Plugins/AI/MassAI/Source/MassAIDebug/Public/GameplayDebuggerCategory_Mass.h
--- a/Engine/Plugins/AI/MassAI/Source/MassAIDebug/Public/GameplayDebuggerCategory_Mass.h
+++ b/Engine/Plugins/AI/MassAI/Source/MassAIDebug/Public/GameplayDebuggerCategory_Mass.h
@@ -2,10 +2,10 @@
 
 #pragma once
 
-// Required first for WITH_MASS_DEBUG
+// Required first for WITH_MASSGAMEPLAY_DEBUG
 #include "MassCommonTypes.h"
 
-#if WITH_GAMEPLAY_DEBUGGER && WITH_MASS_DEBUG
+#if WITH_GAMEPLAY_DEBUGGER && WITH_MASSGAMEPLAY_DEBUG
 
 #include "GameplayDebuggerCategory.h"
 
@@ -62,4 +62,4 @@ protected:
 
 };
 
-#endif // WITH_GAMEPLAY_DEBUGGER && WITH_MASS_DEBUG
+#endif // WITH_GAMEPLAY_DEBUGGER && WITH_MASSGAMEPLAY_DEBUG
diff --git a/Engine/Plugins/AI/MassAI/Source/MassAIMovement/Private/MassAvoidanceProcessors.cpp b/Engine/Plugins/AI/MassAI/Source/MassAIMovement/Private/MassAvoidanceProcessors.cpp
--- a/Engine/Plugins/AI/MassAI/Source/MassAIMovement/Private/MassAvoidanceProcessors.cpp
+++ b/Engine/Plugins/AI/MassAI/Source/MassAIMovement/Private/MassAvoidanceProcessors.cpp
@@ -204,7 +204,7 @@ namespace UE::MassAvoidance
 		return Tweakables::bUseDrawDebugHelpers;
 	}
 
-#if WITH_MASS_DEBUG	
+#if WITH_MASSGAMEPLAY_DEBUG	
 	
 	//----------------------------------------------------------------------//
 	// Begin MassDebugUtils
@@ -379,7 +379,7 @@ namespace UE::MassAvoidance
 		}
 	}
 
-#endif // WITH_MASS_DEBUG
+#endif // WITH_MASSGAMEPLAY_DEBUG
 
 } // namespace UE::MassAvoidance
 
@@ -567,7 +567,7 @@ void UMassAvoidanceProcessor::Execute(UMassEntitySubsystem& EntitySubsystem, FLW
 			}
 			if (!CurrentMovementConfig)
 			{
-#if WITH_MASS_DEBUG && UNSAFE_FOR_MT
+#if WITH_MASSGAMEPLAY_DEBUG && UNSAFE_FOR_MT
 				UE_VLOG(this, LogMassNavigation, Log, TEXT("%s Invalid movement config."), *Entity.DebugGetDescription());
 #endif
 				continue;
@@ -595,7 +595,7 @@ void UMassAvoidanceProcessor::Execute(UMassEntitySubsystem& EntitySubsystem, FLW
 			const float NearEndFade = bFadeAvoidingAtDestination ? FMath::Clamp(MoveTargetList[EntityIndex].DistanceToGoal / NearTargetLocDist, 0.f, 1.f) : 1.0f;
 			const bool bAgentIsMoving = MoveTarget.GetCurrentAction() == EMassMovementAction::Move;
 			
-#if WITH_MASS_DEBUG && UNSAFE_FOR_MT
+#if WITH_MASSGAMEPLAY_DEBUG && UNSAFE_FOR_MT
 			const UE::MassAvoidance::FDebugContext BaseDebugContext(this, LogAvoidance, World, Entity);
 			const UE::MassAvoidance::FDebugContext VelocitiesDebugContext(this, LogAvoidanceVelocities, World, Entity);
 			const UE::MassAvoidance::FDebugContext ObstacleDebugContext(this, LogAvoidanceObstacles, World, Entity);
@@ -622,7 +622,7 @@ void UMassAvoidanceProcessor::Execute(UMassEntitySubsystem& EntitySubsystem, FLW
 				// Draw center
 				DebugDrawSphere(BaseDebugContext, AgentLocation, /*Radius*/2.f, CurrentAgentColor);
 			}
-#endif // WITH_MASS_DEBUG
+#endif // WITH_MASSGAMEPLAY_DEBUG
 
 			FVector OldSteeringForce = FVector::ZeroVector;
 			
@@ -631,10 +631,10 @@ void UMassAvoidanceProcessor::Execute(UMassEntitySubsystem& EntitySubsystem, FLW
 				const FVector DesiredAcceleration = UE::MassAvoidance::Clamp(SteeringForce, MaxSteerAccel);
 				const FVector DesiredVelocity = UE::MassAvoidance::Clamp(AgentVelocity + DesiredAcceleration * TimeDelta, MaximumSpeed);
 
-#if WITH_MASS_DEBUG && UNSAFE_FOR_MT
+#if WITH_MASSGAMEPLAY_DEBUG && UNSAFE_FOR_MT
 				// Draw desired velocity (yellow)
 				UE::MassAvoidance::DebugDrawVelocity(VelocitiesDebugContext, AgentLocation + DebugAgentHeightOffset, AgentLocation + DebugAgentHeightOffset + DesiredVelocity, DesiredVelocityColor);
-#endif // WITH_MASS_DEBUG
+#endif // WITH_MASSGAMEPLAY_DEBUG
 
 				OldSteeringForce = SteeringForce;
 				Contacts.Reset();
@@ -755,7 +755,7 @@ void UMassAvoidanceProcessor::Execute(UMassEntitySubsystem& EntitySubsystem, FLW
 
 						SteeringForce += AvoidForce;
 
-#if WITH_MASS_DEBUG && UNSAFE_FOR_MT
+#if WITH_MASSGAMEPLAY_DEBUG && UNSAFE_FOR_MT
 						// Draw contact normal
 						UE::MassAvoidance::DebugDrawArrow(ObstacleDebugContext, ConPos, ConPos + 50.f * ConNorm, ObstacleContactNormalColor, /*HeadSize=*/ 5.f);
 						UE::MassAvoidance::DebugDrawSphere(ObstacleDebugContext, ConPos, 2.5f, ObstacleContactNormalColor);
@@ -766,18 +766,18 @@ void UMassAvoidanceProcessor::Execute(UMassEntitySubsystem& EntitySubsystem, FLW
 
 						// Draw avoid obstacle force
 						UE::MassAvoidance::DebugDrawForce(ObstacleDebugContext, HitObPos, HitObPos + AvoidForce, ObstacleAvoidForceColor);
-#endif // WITH_MASS_DEBUG
+#endif // WITH_MASSGAMEPLAY_DEBUG
 					}
 				} // edge loop
 
-#if WITH_MASS_DEBUG && UNSAFE_FOR_MT
+#if WITH_MASSGAMEPLAY_DEBUG && UNSAFE_FOR_MT
 				// Draw total steering force to avoid obstacles
 				const FVector EnvironmentAvoidSteeringForce = SteeringForce - OldSteeringForce;
 				UE::MassAvoidance::DebugDrawSummedForce(ObstacleDebugContext,
 					AgentLocation + DebugAgentHeightOffset,
 					AgentLocation + DebugAgentHeightOffset + EnvironmentAvoidSteeringForce,
 					ObstacleAvoidForceColor);
-#endif // WITH_MASS_DEBUG
+#endif // WITH_MASSGAMEPLAY_DEBUG
 
 				// Process contacts to add edge separation force
 				const FVector SteeringForceBeforeSeparation = SteeringForce;
@@ -793,13 +793,13 @@ void UMassAvoidanceProcessor::Execute(UMassEntitySubsystem& EntitySubsystem, FLW
 
 					SteeringForce += SeparationForce;
 
-#if WITH_MASS_DEBUG && UNSAFE_FOR_MT
+#if WITH_MASSGAMEPLAY_DEBUG && UNSAFE_FOR_MT
 					// Draw individual contact forces
 					DebugDrawForce(ObstacleDebugContext, Contacts[ContactIndex].Position + DebugAgentHeightOffset, Contacts[ContactIndex].Position + SeparationForce + DebugAgentHeightOffset, ObstacleSeparationForceColor);
-#endif // WITH_MASS_DEBUG
+#endif // WITH_MASSGAMEPLAY_DEBUG
 				}
 				
-#if WITH_MASS_DEBUG && UNSAFE_FOR_MT
+#if WITH_MASSGAMEPLAY_DEBUG && UNSAFE_FOR_MT
 				// Draw total steering force to separate from close edges
 				const FVector TotalSeparationForce = SteeringForce - SteeringForceBeforeSeparation;
 				DebugDrawSummedForce(ObstacleDebugContext,
@@ -817,7 +817,7 @@ void UMassAvoidanceProcessor::Execute(UMassEntitySubsystem& EntitySubsystem, FLW
 						DebugDrawArrow(ObstacleDebugContext, Middle, Middle + 10.f * FVector::CrossProduct((Edge.End - Edge.Start), FVector::UpVector).GetSafeNormal(), ObstacleColor, /*HeadSize=*/2.f);
 					}
 				}
-#endif // WITH_MASS_DEBUG
+#endif // WITH_MASSGAMEPLAY_DEBUG
 			}
 
 			//////////////////////////////////////////////////////////////////////////
@@ -1056,7 +1056,7 @@ void UMassAvoidanceProcessor::Execute(UMassEntitySubsystem& EntitySubsystem, FLW
 
 				SteeringForce += AvoidForce;
 
-#if WITH_MASS_DEBUG && UNSAFE_FOR_MT
+#if WITH_MASSGAMEPLAY_DEBUG && UNSAFE_FOR_MT
 				// Display close agent
 				UE::MassAvoidance::DebugDrawCylinder(AgentDebugContext, Collider.Location, Collider.Location + DebugLowCylinderOffset, Collider.Radius, AgentsColor);
 
@@ -1095,7 +1095,7 @@ void UMassAvoidanceProcessor::Execute(UMassEntitySubsystem& EntitySubsystem, FLW
 						OtherHitPosition + DebugAgentHeightOffset + AvoidForce,
 						AgentAvoidForceColor);
 				}
-#endif // WITH_MASS_DEBUG
+#endif // WITH_MASSGAMEPLAY_DEBUG
 			} // close entities loop
 
 			if (MoveTarget.GetPreviousAction() != EMassMovementAction::Move)
@@ -1110,7 +1110,7 @@ void UMassAvoidanceProcessor::Execute(UMassEntitySubsystem& EntitySubsystem, FLW
 			
 			Steering.SteeringForce = UE::MassAvoidance::Clamp(SteeringForce, MaxSteerAccel); // Assume unit mass
 
-#if WITH_MASS_DEBUG && UNSAFE_FOR_MT
+#if WITH_MASSGAMEPLAY_DEBUG && UNSAFE_FOR_MT
 			const FVector AgentAvoidSteeringForce = SteeringForce - OldSteeringForce;
 
 			// Draw total steering force to separate agents
@@ -1130,7 +1130,7 @@ void UMassAvoidanceProcessor::Execute(UMassEntitySubsystem& EntitySubsystem, FLW
 				AgentLocation + AgentVelocity + DebugAgentHeightOffset,
 				AgentLocation + AgentVelocity + DebugAgentHeightOffset + SteeringList[EntityIndex].SteeringForce,
 				FinalSteeringForceColor, BigArrowHeadSize, BigArrowThickness);
-#endif // WITH_MASS_DEBUG
+#endif // WITH_MASSGAMEPLAY_DEBUG
 		}
 	});
 }
@@ -1488,7 +1488,7 @@ void UMassAvoidanceObstacleProcessor::Execute(UMassEntitySubsystem& EntitySubsys
 				const FBox NewBounds(NewPos - FVector(Radius, Radius, 0.f), NewPos + FVector(Radius, Radius, 0.f));
 				AvoidanceObstacleCellLocationList[EntityIndex].CellLoc = MovementSubsystem->GetGridMutable().Move(ObstacleItem, AvoidanceObstacleCellLocationList[EntityIndex].CellLoc, NewBounds);
 
-#if WITH_MASS_DEBUG && 0
+#if WITH_MASSGAMEPLAY_DEBUG && 0
 				const FDebugContext BaseDebugContext(this, LogAvoidance, nullptr, ObstacleItem.Entity);
 				if (DebugIsSelected(ObstacleItem.Entity))
 				{
@@ -1496,7 +1496,7 @@ void UMassAvoidanceObstacleProcessor::Execute(UMassEntitySubsystem& EntitySubsys
 					Box.Max.Z += 200.f;					
 					DebugDrawBox(BaseDebugContext, Box, FColor::Yellow);
 				}
-#endif // WITH_MASS_DEBUG
+#endif // WITH_MASSGAMEPLAY_DEBUG
 			}
 		});
 
@@ -1684,7 +1684,7 @@ void UMassLaneBoundaryProcessor::Execute(UMassEntitySubsystem& EntitySubsystem,
 				FMassNavigationEdgesFragment& EdgesFragment = EdgesList[EntityIndex];
 				EdgesFragment.AvoidanceEdges.Reset();
 
-#if WITH_MASS_DEBUG && UNSAFE_FOR_MT
+#if WITH_MASSGAMEPLAY_DEBUG && UNSAFE_FOR_MT
 				const FLWEntity Entity = Context.GetEntity(EntityIndex);
 				const FDebugContext BaseDebugContext(this, LogAvoidance, World, Entity);
 				const FDebugContext ObstacleDebugContext(this, LogAvoidanceObstacles, World, Entity);
@@ -1696,7 +1696,7 @@ void UMassLaneBoundaryProcessor::Execute(UMassEntitySubsystem& EntitySubsystem,
 						DebugDrawSphere(ObstacleDebugContext, Location, /*Radius=*/100.f, FColor(128,128,128));
 					}
 				}
-#endif // WITH_MASS_DEBUG
+#endif // WITH_MASSGAMEPLAY_DEBUG
 
 				// @todo: Fix transition between lanes is not smooth. WanderFragment.CurrentLaneLocation is in front of the actual
 				//		  position so when there is a lane switch picking the closest position on the lane jumps forward.
@@ -1719,7 +1719,7 @@ void UMassLaneBoundaryProcessor::Execute(UMassEntitySubsystem& EntitySubsystem,
 
 					if (LaneLocation.IsValid())
 					{
-#if WITH_MASS_DEBUG && UNSAFE_FOR_MT
+#if WITH_MASSGAMEPLAY_DEBUG && UNSAFE_FOR_MT
 						if (DebugIsSelected(Entity))
 						{
 							// Draw the location found and the lane from that position to the end.
@@ -1735,7 +1735,7 @@ void UMassLaneBoundaryProcessor::Execute(UMassEntitySubsystem& EntitySubsystem,
 								LaneSegment++;
 							}
 						}
-#endif // WITH_MASS_DEBUG
+#endif // WITH_MASSGAMEPLAY_DEBUG
 
 						// Get width of adjacent lanes.
 						float AdjacentLeftWidth = 0.f;
@@ -1819,7 +1819,7 @@ void UMassLaneBoundaryProcessor::Execute(UMassEntitySubsystem& EntitySubsystem,
 						int32 NumRightPositions = NumPoints;
 
 
-#if 0 && WITH_MASS_DEBUG // Detailed debug disabled
+#if 0 && WITH_MASSGAMEPLAY_DEBUG // Detailed debug disabled
 						if (DebugIsSelected(Entity))
 						{
 							float Radius = 2.f;
@@ -1835,7 +1835,7 @@ void UMassLaneBoundaryProcessor::Execute(UMassEntitySubsystem& EntitySubsystem,
 								Radius += 4.f;
 							}
 						}
-#endif //WITH_MASS_DEBUG
+#endif //WITH_MASSGAMEPLAY_DEBUG
 
 						// Remove edges crossing when there are 3 eges.
 						if (NumPoints == 4)
@@ -1942,7 +1942,7 @@ void UMassLaneCacheBoundaryProcessor::Execute(UMassEntitySubsystem& EntitySubsys
 			const float DeltaDistSquared = FVector::DistSquared(MovementTarget.Center, LaneCacheBoundary.LastUpdatePosition);
 			const float UpdateDistanceThresholdSquared = FMath::Square(50.f);
 
-#if WITH_MASS_DEBUG
+#if WITH_MASSGAMEPLAY_DEBUG
 			const FDebugContext ObstacleDebugContext(this, LogAvoidanceObstacles, World, Entity);
 			if (DebugIsSelected(Entity))
 			{
@@ -1969,12 +1969,12 @@ void UMassLaneCacheBoundaryProcessor::Execute(UMassEntitySubsystem& EntitySubsys
 			}
 			
 
-#if WITH_MASS_DEBUG
+#if WITH_MASSGAMEPLAY_DEBUG
 			if (DebugIsSelected(Entity))
 			{
 				DebugDrawSphere(ObstacleDebugContext, MovementTarget.Center, /*Radius=*/100.f, FColor(128,128,128));
 			}
-#endif // WITH_MASS_DEBUG
+#endif // WITH_MASSGAMEPLAY_DEBUG
 
 			const float HalfWidth = 0.5 * CachedLane.LaneWidth.Get();
 
@@ -2033,7 +2033,7 @@ void UMassLaneCacheBoundaryProcessor::Execute(UMassEntitySubsystem& EntitySubsys
 			int32 NumRightPositions = NumPoints;
 
 
-#if 0 && WITH_MASS_DEBUG // Detailed debug disabled
+#if 0 && WITH_MASSGAMEPLAY_DEBUG // Detailed debug disabled
 			if (DebugIsSelected(Entity))
 			{
 				float Radius = 2.f;
@@ -2049,7 +2049,7 @@ void UMassLaneCacheBoundaryProcessor::Execute(UMassEntitySubsystem& EntitySubsys
 					Radius += 4.f;
 				}
 			}
-#endif //WITH_MASS_DEBUG
+#endif //WITH_MASSGAMEPLAY_DEBUG
 
 			// Remove edges crossing when there are 3 edges.
 			if (NumPoints == 4)
diff --git a/Engine/Plugins/AI/MassAI/Source/MassAIMovement/Private/MassZoneGraphLaneObstacleAnnotations.cpp b/Engine/Plugins/AI/MassAI/Source/MassAIMovement/Private/MassZoneGraphLaneObstacleAnnotations.cpp
--- a/Engine/Plugins/AI/MassAI/Source/MassAIMovement/Private/MassZoneGraphLaneObstacleAnnotations.cpp
+++ b/Engine/Plugins/AI/MassAI/Source/MassAIMovement/Private/MassZoneGraphLaneObstacleAnnotations.cpp
@@ -87,7 +87,7 @@ void UZoneGraphLaneObstacleAnnotations::TickAnnotation(const float DeltaTime, FZ
 	}
 	LaneObstacleChangeEvents.Reset();
 
-#if WITH_MASS_DEBUG && 0
+#if WITH_MASSGAMEPLAY_DEBUG && 0
 	for (const FMassRegisteredMovementLaneData& LaneData : RegisteredLaneData)
 	{
 		const FMassLaneObstacleContainer& LaneObstaclesContainer = LaneData.LaneObstacles;
@@ -107,7 +107,7 @@ void UZoneGraphLaneObstacleAnnotations::TickAnnotation(const float DeltaTime, FZ
 			}
 		}
 	}
-#endif // WITH_MASS_DEBUG 
+#endif // WITH_MASSGAMEPLAY_DEBUG 
 }
 
 void UZoneGraphLaneObstacleAnnotations::PostZoneGraphDataAdded(const AZoneGraphData& ZoneGraphData)
diff --git a/Engine/Plugins/AI/MassAI/Source/MassAIMovement/Private/MassZoneGraphMovementFragments.cpp b/Engine/Plugins/AI/MassAI/Source/MassAIMovement/Private/MassZoneGraphMovementFragments.cpp
--- a/Engine/Plugins/AI/MassAI/Source/MassAIMovement/Private/MassZoneGraphMovementFragments.cpp
+++ b/Engine/Plugins/AI/MassAI/Source/MassAIMovement/Private/MassZoneGraphMovementFragments.cpp
@@ -169,7 +169,7 @@ bool FMassZoneGraphShortPathFragment::RequestPath(const FMassZoneGraphCachedLane
 	check(CurrentDistanceAlongLane >= -Epsilon && CurrentDistanceAlongLane <= (CachedLane.LaneLength + Epsilon));
 	
 	// Set common lane parameters
-#if WITH_MASS_DEBUG
+#if WITH_MASSGAMEPLAY_DEBUG
 	DebugLaneHandle = CachedLane.LaneHandle;
 #endif
 
@@ -443,7 +443,7 @@ bool FMassZoneGraphShortPathFragment::RequestStand(const FMassZoneGraphCachedLan
 	CachedLane.GetPointAndTangentAtDistance(CurrentDistanceAlongLane, CurrentLanePosition, CurrentLaneTangent);
 	
 	// Set common lane parameters
-#if WITH_MASS_DEBUG
+#if WITH_MASSGAMEPLAY_DEBUG
 	DebugLaneHandle = CachedLane.LaneHandle;
 #endif
 	bMoveReverse = false;
diff --git a/Engine/Plugins/AI/MassAI/Source/MassAIMovement/Private/MassZoneGraphMovementProcessors.cpp b/Engine/Plugins/AI/MassAI/Source/MassAIMovement/Private/MassZoneGraphMovementProcessors.cpp
--- a/Engine/Plugins/AI/MassAI/Source/MassAIMovement/Private/MassZoneGraphMovementProcessors.cpp
+++ b/Engine/Plugins/AI/MassAI/Source/MassAIMovement/Private/MassZoneGraphMovementProcessors.cpp
@@ -14,7 +14,7 @@
 
 #define UNSAFE_FOR_MT 1
 
-#if WITH_MASS_DEBUG
+#if WITH_MASSGAMEPLAY_DEBUG
 
 namespace UE::MassMovement::Debug
 {
@@ -28,7 +28,7 @@ namespace UE::MassMovement::Debug
 	}
 }
 
-#endif // WITH_MASS_DEBUG
+#endif // WITH_MASSGAMEPLAY_DEBUG
 
 namespace UE::MassMovement
 {
@@ -125,7 +125,7 @@ void UMassZoneGraphLocationInitializer::Execute(UMassEntitySubsystem& EntitySubs
 			}
 			if (!CurrentMovementConfig)
 			{
-#if WITH_MASS_DEBUG && UNSAFE_FOR_MT
+#if WITH_MASSGAMEPLAY_DEBUG && UNSAFE_FOR_MT
 				const FLWEntity Entity = Context.GetEntity(EntityIndex);
 				UE_VLOG(this, LogMassNavigation, Log, TEXT("Entity [%s] Invalid movement config."), *Entity.DebugGetDescription());
 #endif
@@ -223,13 +223,13 @@ void UMassZoneGraphPathFollowProcessor::Execute(UMassEntitySubsystem& EntitySubs
 			const float DeltaTime = SimLOD.DeltaTime;
 
 			bool bDisplayDebug = false;
-#if WITH_MASS_DEBUG && UNSAFE_FOR_MT // this will result in bDisplayDebug == false and disabling of all the vlogs below
+#if WITH_MASSGAMEPLAY_DEBUG && UNSAFE_FOR_MT // this will result in bDisplayDebug == false and disabling of all the vlogs below
 			bDisplayDebug = UE::MassDebug::IsDebuggingEntity(Entity);
 			if (bDisplayDebug)
 			{
 				UE_VLOG(this, LogMassNavigation, Log, TEXT("Entity [%s] Updating path following"), *Entity.DebugGetDescription());
 			}
-#endif // WITH_MASS_DEBUG
+#endif // WITH_MASSGAMEPLAY_DEBUG
 
 			// Must have at least two points to interpolate.
 			if (MoveTarget.GetCurrentAction() == EMassMovementAction::Move && ShortPath.NumPoints >= 2)
@@ -251,9 +251,9 @@ void UMassZoneGraphPathFollowProcessor::Execute(UMassEntitySubsystem& EntitySubs
 				if (!bWasDone)
 				{
 					const uint8 LastPointIndex = ShortPath.NumPoints - 1;
-#if WITH_MASS_DEBUG
+#if WITH_MASSGAMEPLAY_DEBUG
 					ensureMsgf(LaneLocation.LaneHandle == ShortPath.DebugLaneHandle, TEXT("Short path lane should match current lane location."));
-#endif // WITH_MASS_DEBUG
+#endif // WITH_MASSGAMEPLAY_DEBUG
 
 					if (ShortPath.ProgressDistance <= 0.0f)
 					{
@@ -408,7 +408,7 @@ void UMassZoneGraphPathFollowProcessor::Execute(UMassEntitySubsystem& EntitySubs
 					EntitiesToSignalPathDone.Add(Entity);
 				}
 
-#if WITH_MASS_DEBUG && UNSAFE_FOR_MT
+#if WITH_MASSGAMEPLAY_DEBUG && UNSAFE_FOR_MT
 				if (bDisplayDebug)
 				{
 					const FColor EntityColor = UE::MassDebug::GetEntityDebugColor(Entity);
@@ -440,7 +440,7 @@ void UMassZoneGraphPathFollowProcessor::Execute(UMassEntitySubsystem& EntitySubs
 						UE_VLOG_SEGMENT_THICK(this, LogMassNavigation, Display, CurrBase, CurrBase + FVector(0,0,100), FColor::Red, /*Thickness*/3.0f, TEXT("Next: %s"), *ShortPath.NextLaneHandle.ToString());
 					}
 				}
-#endif // WITH_MASS_DEBUG
+#endif // WITH_MASSGAMEPLAY_DEBUG
 			}
 		}
 	});
@@ -532,7 +532,7 @@ void UMassZoneGraphSteeringProcessor::Execute(UMassEntitySubsystem& EntitySubsys
 			}
 			if (!CurrentMovementConfig)
 			{
-#if WITH_MASS_DEBUG && UNSAFE_FOR_MT
+#if WITH_MASSGAMEPLAY_DEBUG && UNSAFE_FOR_MT
 				UE_VLOG(this, LogMassNavigation, Log, TEXT("Entity [%s] Invalid movement config."), *Entity.DebugGetDescription());
 #endif
 				continue;
@@ -707,7 +707,7 @@ void UMassZoneGraphSteeringProcessor::Execute(UMassEntitySubsystem& EntitySubsys
 				MoveTarget.bSteeringFallingBehind = false;
 			}
 
-#if WITH_MASS_DEBUG && UNSAFE_FOR_MT
+#if WITH_MASSGAMEPLAY_DEBUG && UNSAFE_FOR_MT
 			FColor EntityColor = FColor::White;
 			const bool bDisplayDebug = UE::MassDebug::IsDebuggingEntity(Entity, &EntityColor);
 			if (bDisplayDebug)
@@ -740,7 +740,7 @@ void UMassZoneGraphSteeringProcessor::Execute(UMassEntitySubsystem& EntitySubsys
 				UE_VLOG_SEGMENT_THICK(this, LogMassNavigation, Log, CurrentLocation + ZOffset, CurrentLocation + Steering.DesiredVelocity + ZOffset, LightEntityColor, 2.0f,
 					TEXT("%s Steer %.1f"), *Entity.DebugGetDescription(), Steering.DesiredVelocity.Length());
 			}
-#endif // WITH_MASS_DEBUG
+#endif // WITH_MASSGAMEPLAY_DEBUG
 			
 		}
 	});
diff --git a/Engine/Plugins/AI/MassAI/Source/MassAIMovement/Private/MassZoneGraphMovementUtils.cpp b/Engine/Plugins/AI/MassAI/Source/MassAIMovement/Private/MassZoneGraphMovementUtils.cpp
--- a/Engine/Plugins/AI/MassAI/Source/MassAIMovement/Private/MassZoneGraphMovementUtils.cpp
+++ b/Engine/Plugins/AI/MassAI/Source/MassAIMovement/Private/MassZoneGraphMovementUtils.cpp
@@ -1,7 +1,7 @@
 ﻿// Copyright Epic Games, Inc. All Rights Reserved.
 
 #include "MassZoneGraphMovementUtils.h"
-
+#include "MassCommonTypes.h"
 #include "MassAIMovementFragments.h"
 #include "MassZoneGraphMovementFragments.h"
 #include "ZoneGraphSubsystem.h"
@@ -50,7 +50,7 @@ namespace UE::MassMovement
 		{
 			MoveTarget.IntentAtGoal = ShortPath.EndOfPathIntent;
 			MoveTarget.DistanceToGoal = (ShortPath.NumPoints > 0) ? ShortPath.Points[ShortPath.NumPoints - 1].DistanceAlongLane.Get() : 0.0f;
-
+#if WITH_MASSGAMEPLAY_DEBUG
 			UE_CVLOG(UE::MassDebug::IsDebuggingEntity(Entity),
 				Requester,
 				LogMassNavigation,
@@ -61,6 +61,7 @@ namespace UE::MassMovement
 				LaneLocation.DistanceAlongLane,
 				PathRequest.TargetDistance,
 				*PathRequest.NextLaneHandle.ToString());
+#endif // WITH_MASSGAMEPLAY_DEBUG
 		}
 		else
 		{
diff --git a/Engine/Plugins/AI/MassAI/Source/MassAIMovement/Private/MovementProcessor.cpp b/Engine/Plugins/AI/MassAI/Source/MassAIMovement/Private/MovementProcessor.cpp
--- a/Engine/Plugins/AI/MassAI/Source/MassAIMovement/Private/MovementProcessor.cpp
+++ b/Engine/Plugins/AI/MassAI/Source/MassAIMovement/Private/MovementProcessor.cpp
@@ -60,47 +60,47 @@ namespace UE::MassMovement
 
  	static bool DebugIsSelected(const FLWEntity& Entity)
  	{
- #if WITH_MASS_DEBUG
+ #if WITH_MASSGAMEPLAY_DEBUG
  		return UE::MassDebug::IsDebuggingEntity(Entity);
  #else
  	return false;
- #endif // WITH_MASS_DEBUG
+ #endif // WITH_MASSGAMEPLAY_DEBUG
  	}
 
  	static void DebugDrawLine(const UObject* LogOwner, const UWorld* World, const FVector& Start, const FVector& End, const FColor& Color, const bool bPersistent = false)
  	{
- #if WITH_MASS_DEBUG
+ #if WITH_MASSGAMEPLAY_DEBUG
  		UE_VLOG_SEGMENT(LogOwner, LogNavigation, Log, Start, End, Color, TEXT(""));
 
  		if (World)
  		{
  			DrawDebugLine(World, Start, End, Color, bPersistent);
  		}
- #endif // WITH_MASS_DEBUG
+ #endif // WITH_MASSGAMEPLAY_DEBUG
  	}
 
  	static void DebugDrawArrow(const UObject* LogOwner, const UWorld* World, const FVector& Start, const FVector& End, const FColor& Color)
  	{
- #if WITH_MASS_DEBUG
+ #if WITH_MASSGAMEPLAY_DEBUG
  		UE_VLOG_ARROW(LogOwner, LogNavigation, Log, Start, End, Color, TEXT(""));
 
  		if (World)
  		{
  			DrawDebugDirectionalArrow(World, Start, End, /*arrow size = */20.f, Color);
  		}
- #endif // WITH_MASS_DEBUG
+ #endif // WITH_MASSGAMEPLAY_DEBUG
  	}
 
  	static void DebugDrawCylinder(const UObject* LogOwner, const UWorld* World, const FVector& Start, const FVector& End, const float Radius, const FColor& Color)
  	{
- #if WITH_MASS_DEBUG
+ #if WITH_MASSGAMEPLAY_DEBUG
  		UE_VLOG_CYLINDER(LogOwner, LogNavigation, Log, Start, End, Radius, Color, TEXT(""));
 
  		if (World)
  		{
  			DrawDebugCylinder(World, Start, End, Radius, /*segments = */16, Color);
  		}
- #endif // WITH_MASS_DEBUG
+ #endif // WITH_MASSGAMEPLAY_DEBUG
 	}
 } // namespace UE::Movement
 
diff --git a/Engine/Plugins/AI/MassAI/Source/MassAIMovement/Public/MassZoneGraphMovementFragments.h b/Engine/Plugins/AI/MassAI/Source/MassAIMovement/Public/MassZoneGraphMovementFragments.h
--- a/Engine/Plugins/AI/MassAI/Source/MassAIMovement/Public/MassZoneGraphMovementFragments.h
+++ b/Engine/Plugins/AI/MassAI/Source/MassAIMovement/Public/MassZoneGraphMovementFragments.h
@@ -178,7 +178,7 @@ struct MASSAIMOVEMENT_API FMassZoneGraphShortPathFragment : public FLWComponentD
 
 	void Reset()
 	{
-#if WITH_MASS_DEBUG
+#if WITH_MASSGAMEPLAY_DEBUG
 		DebugLaneHandle.Reset();
 #endif
 		NextLaneHandle.Reset();
@@ -203,7 +203,7 @@ struct MASSAIMOVEMENT_API FMassZoneGraphShortPathFragment : public FLWComponentD
 		return NumPoints == 0 || bDone;
 	}
 
-#if WITH_MASS_DEBUG
+#if WITH_MASSGAMEPLAY_DEBUG
 	/** Current lane handle, for debug */
 	FZoneGraphLaneHandle DebugLaneHandle;
 #endif
diff --git a/Engine/Plugins/AI/MassCrowd/Source/MassCrowd/Private/MassCrowdBubble.cpp b/Engine/Plugins/AI/MassCrowd/Source/MassCrowd/Private/MassCrowdBubble.cpp
--- a/Engine/Plugins/AI/MassCrowd/Source/MassCrowd/Private/MassCrowdBubble.cpp
+++ b/Engine/Plugins/AI/MassCrowd/Source/MassCrowd/Private/MassCrowdBubble.cpp
@@ -13,7 +13,7 @@ namespace UE::Mass::Crowd
 	bool bDebugReplicationPositions = false;
 	FAutoConsoleVariableRef CVarbDebugReplication(TEXT("ai.debug.CrowdReplicationPositions"), bDebugReplicationPositions, TEXT("Crowd debug replication positions on server and client"), ECVF_Cheat);
 
-#if WITH_MASS_DEBUG && WITH_EDITOR
+#if WITH_MASSGAMEPLAY_DEBUG && WITH_EDITOR
 	// @todo provide a better way of selecting agents to debug
 	constexpr int32 MaxAgentsDraw = 300;
 
@@ -60,10 +60,10 @@ namespace UE::Mass::Crowd
 			DrawDebugCylinder(World, Pos + 0.5f * DebugCylinderHeight, Pos + DebugCylinderHeight, DebugCylinderRadius, /*segments = */24, DebugCylinderColor);
 		}
 	}
-#endif // WITH_MASS_DEBUG && WITH_EDITOR
+#endif // WITH_MASSGAMEPLAY_DEBUG && WITH_EDITOR
 }
 
-#if WITH_MASS_DEBUG && WITH_EDITOR
+#if WITH_MASSGAMEPLAY_DEBUG && WITH_EDITOR
 void FMassCrowdClientBubbleHandler::DebugValidateBubbleOnServer()
 {
 	Super::DebugValidateBubbleOnServer();
@@ -88,9 +88,9 @@ void FMassCrowdClientBubbleHandler::DebugValidateBubbleOnServer()
 		}
 	}
 }
-#endif // WITH_MASS_DEBUG && WITH_EDITOR
+#endif // WITH_MASSGAMEPLAY_DEBUG && WITH_EDITOR
 
-#if WITH_MASS_DEBUG && WITH_EDITOR
+#if WITH_MASSGAMEPLAY_DEBUG && WITH_EDITOR
 void FMassCrowdClientBubbleHandler::DebugValidateBubbleOnClient()
 {
 	Super::DebugValidateBubbleOnClient();
@@ -118,7 +118,7 @@ void FMassCrowdClientBubbleHandler::DebugValidateBubbleOnClient()
 		}
 	}
 }
-#endif // WITH_MASS_DEBUG && WITH_EDITOR
+#endif // WITH_MASSGAMEPLAY_DEBUG && WITH_EDITOR
 
 #if UE_REPLICATION_COMPILE_CLIENT_CODE
 void FMassCrowdClientBubbleHandler::PostReplicatedAdd(const TArrayView<int32> AddedIndices, int32 FinalSize)
diff --git a/Engine/Plugins/AI/MassCrowd/Source/MassCrowd/Private/MassCrowdReplicationProcessor.cpp b/Engine/Plugins/AI/MassCrowd/Source/MassCrowd/Private/MassCrowdReplicationProcessor.cpp
--- a/Engine/Plugins/AI/MassCrowd/Source/MassCrowd/Private/MassCrowdReplicationProcessor.cpp
+++ b/Engine/Plugins/AI/MassCrowd/Source/MassCrowd/Private/MassCrowdReplicationProcessor.cpp
@@ -54,7 +54,7 @@ void UMassCrowdReplicationProcessor::Execute(UMassEntitySubsystem& EntitySubsyst
 
 	ProcessClientReplication(EntitySubsystem, Context);
 
-#if WITH_MASS_DEBUG
+#if WITH_MASSGAMEPLAY_DEBUG
 	// Optional debug display
 	if (UE::Mass::Crowd::bDebugReplicationViewerLOD)
 	{
@@ -65,7 +65,7 @@ void UMassCrowdReplicationProcessor::Execute(UMassEntitySubsystem& EntitySubsyst
 				LODCalculator.DebugDisplayLOD(Context, ViewerLODList, TransformList, World);
 			});
 	}
-#endif // WITH_MASS_DEBUG
+#endif // WITH_MASSGAMEPLAY_DEBUG
 }
 
 void UMassCrowdReplicationProcessor::ProcessClientReplication(UMassEntitySubsystem& EntitySubsystem, FLWComponentSystemExecutionContext& Context)
diff --git a/Engine/Plugins/AI/MassCrowd/Source/MassCrowd/Private/MassCrowdSubsystem.cpp b/Engine/Plugins/AI/MassCrowd/Source/MassCrowd/Private/MassCrowdSubsystem.cpp
--- a/Engine/Plugins/AI/MassCrowd/Source/MassCrowd/Private/MassCrowdSubsystem.cpp
+++ b/Engine/Plugins/AI/MassCrowd/Source/MassCrowd/Private/MassCrowdSubsystem.cpp
@@ -15,7 +15,7 @@
 #include "ZoneGraphSubsystem.h"
 #include "VisualLogger/VisualLogger.h"
 
-#if WITH_MASS_DEBUG
+#if WITH_MASSGAMEPLAY_DEBUG
 namespace UE::MassCrowdDebug
 {
 bool bForceFillIntersectionLanes = false;
@@ -39,7 +39,7 @@ static FAutoConsoleCommandWithWorld RebuildCmd(
 	}));
 #endif // WITH_EDITOR
 }// namespace UE::MassCrowdDebug
-#endif // WITH_MASS_DEBUG
+#endif // WITH_MASSGAMEPLAY_DEBUG
 
 //----------------------------------------------------------------------//
 // UMassCrowdSubsystem
diff --git a/Engine/Plugins/AI/MassCrowd/Source/MassCrowd/Private/Tasks/MassZoneGraphFindWanderTarget.cpp b/Engine/Plugins/AI/MassCrowd/Source/MassCrowd/Private/Tasks/MassZoneGraphFindWanderTarget.cpp
--- a/Engine/Plugins/AI/MassCrowd/Source/MassCrowd/Private/Tasks/MassZoneGraphFindWanderTarget.cpp
+++ b/Engine/Plugins/AI/MassCrowd/Source/MassCrowd/Private/Tasks/MassZoneGraphFindWanderTarget.cpp
@@ -27,9 +27,9 @@ EStateTreeRunStatus FMassZoneGraphFindWanderTarget::EnterState(FStateTreeExecuti
 	const UMassCrowdSubsystem& MassCrowdSubsystem = Context.GetExternalItem(MassCrowdSubsystemHandle).Get<UMassCrowdSubsystem>();
 
 	bool bDisplayDebug = false;
-#if WITH_MASS_DEBUG
+#if WITH_MASSGAMEPLAY_DEBUG
 	bDisplayDebug = UE::MassDebug::IsDebuggingEntity(MassContext.GetEntity());
-#endif // WITH_MASS_DEBUG
+#endif // WITH_MASSGAMEPLAY_DEBUG
 
 	if (!LaneLocation.LaneHandle.IsValid())
 	{
diff --git a/Engine/Plugins/AI/MassCrowd/Source/MassCrowd/Public/MassCrowdBubble.h b/Engine/Plugins/AI/MassCrowd/Source/MassCrowd/Public/MassCrowdBubble.h
--- a/Engine/Plugins/AI/MassCrowd/Source/MassCrowd/Public/MassCrowdBubble.h
+++ b/Engine/Plugins/AI/MassCrowd/Source/MassCrowd/Public/MassCrowdBubble.h
@@ -42,10 +42,10 @@ protected:
 	void PostReplicatedChangeEntity(const FEntityView& EntityView, const FReplicatedCrowdAgent& Item) const;
 #endif //UE_REPLICATION_COMPILE_CLIENT_CODE
 
-#if WITH_MASS_DEBUG && WITH_EDITOR
+#if WITH_MASSGAMEPLAY_DEBUG && WITH_EDITOR
 	virtual void DebugValidateBubbleOnServer() override;
 	virtual void DebugValidateBubbleOnClient() override;
-#endif // WITH_MASS_DEBUG
+#endif // WITH_MASSGAMEPLAY_DEBUG
 
 	FMassClientBubblePathHandler PathHandler;
 	FMassClientBubbleTransformHandler TransformHandler;
diff --git a/Engine/Plugins/Runtime/MassEntity/Source/MassEntity/Private/ArchetypeData.cpp b/Engine/Plugins/Runtime/MassEntity/Source/MassEntity/Private/ArchetypeData.cpp
--- a/Engine/Plugins/Runtime/MassEntity/Source/MassEntity/Private/ArchetypeData.cpp
+++ b/Engine/Plugins/Runtime/MassEntity/Source/MassEntity/Private/ArchetypeData.cpp
@@ -769,7 +769,7 @@ SIZE_T FArchetypeData::GetAllocatedSize() const
 
 FString FArchetypeData::DebugGetDescription() const
 {
-#if WITH_AGGREGATETICKING_DEBUG
+#if WITH_MASSENTITY_DEBUG
 	TStringBuilder<256> ArchetypeDebugName;
 	
 	ArchetypeDebugName.Append(TEXT("<"));
@@ -792,7 +792,7 @@ FString FArchetypeData::DebugGetDescription() const
 #endif
 }
 
-#if WITH_AGGREGATETICKING_DEBUG
+#if WITH_MASSENTITY_DEBUG
 void FArchetypeData::DebugPrintArchetype(FOutputDevice& Ar)
 {
 	FStringOutputDevice TagsDecription;
@@ -874,7 +874,7 @@ void FArchetypeData::DebugPrintEntity(FLWEntity Entity, FOutputDevice& Ar, const
 	}
 }
 
-#endif // WITH_AGGREGATETICKING_DEBUG
+#endif // WITH_MASSENTITY_DEBUG
 
 void FArchetypeData::REMOVEME_GetArrayViewForComponentInChunk(int32 ChunkIndex, const UScriptStruct* ComponentType, void*& OutChunkBase, int32& OutNumEntities)
 {
diff --git a/Engine/Plugins/Runtime/MassEntity/Source/MassEntity/Private/ArchetypeData.h b/Engine/Plugins/Runtime/MassEntity/Source/MassEntity/Private/ArchetypeData.h
--- a/Engine/Plugins/Runtime/MassEntity/Source/MassEntity/Private/ArchetypeData.h
+++ b/Engine/Plugins/Runtime/MassEntity/Source/MassEntity/Private/ArchetypeData.h
@@ -92,9 +92,9 @@ public:
 		ChunkComponentData = InChunkComponentsTemplate;
 	}
 
-#if WITH_AGGREGATETICKING_DEBUG
+#if WITH_MASSENTITY_DEBUG
 	int32 DebugGetChunkComponentCount() const { return ChunkComponentData.Num(); }
-#endif // WITH_AGGREGATETICKING_DEBUG
+#endif // WITH_MASSENTITY_DEBUG
 };
 
 // Information for a single component type in an archetype
@@ -224,7 +224,7 @@ public:
 	// Converts the list of components into a user-readable debug string
 	FString DebugGetDescription() const;
 
-#if WITH_AGGREGATETICKING_DEBUG
+#if WITH_MASSENTITY_DEBUG
 	/**
 	 * Prints out debug information about the archetype
 	 */
@@ -237,7 +237,7 @@ public:
 	 * @param InPrefix Optional prefix to remove from component names
 	 */
 	void DebugPrintEntity(FLWEntity Entity, FOutputDevice& Ar, const TCHAR* InPrefix = TEXT("")) const;
-#endif // WITH_AGGREGATETICKING_DEBUG
+#endif // WITH_MASSENTITY_DEBUG
 
 	void REMOVEME_GetArrayViewForComponentInChunk(int32 ChunkIndex, const UScriptStruct* ComponentType, void*& OutChunkBase, int32& OutNumEntities);
 
diff --git a/Engine/Plugins/Runtime/MassEntity/Source/MassEntity/Private/EntityQuery.cpp b/Engine/Plugins/Runtime/MassEntity/Source/MassEntity/Private/EntityQuery.cpp
--- a/Engine/Plugins/Runtime/MassEntity/Source/MassEntity/Private/EntityQuery.cpp
+++ b/Engine/Plugins/Runtime/MassEntity/Source/MassEntity/Private/EntityQuery.cpp
@@ -127,7 +127,7 @@ void FLWComponentQuery::ForEachEntityChunk(const FArchetypeChunkCollection& Chun
 
 void FLWComponentQuery::ForEachEntityChunk(UMassEntitySubsystem& EntitySubsystem, FLWComponentSystemExecutionContext& ExecutionContext, const FLWComponentSystemExecuteFunction& ExecuteFunction)
 {
-#if WITH_AGGREGATETICKING_DEBUG
+#if WITH_MASSENTITY_DEBUG
 	int32 NumEntitiesToProcess = 0;
 #endif
 
@@ -143,7 +143,7 @@ void FLWComponentQuery::ForEachEntityChunk(UMassEntitySubsystem& EntitySubsystem
 		}
 		ExecutionContext.SetRequirements(Requirements, ChunkRequirements);
 		ExecutionContext.GetChunkCollection().GetArchetype().DataPtr->ExecuteFunction(ExecutionContext, ExecuteFunction, {}, ExecutionContext.GetChunkCollection());
-#if WITH_AGGREGATETICKING_DEBUG
+#if WITH_MASSENTITY_DEBUG
 		NumEntitiesToProcess = ExecutionContext.GetChunkCollection().GetArchetype().DataPtr->GetNumEntities();
 #endif
 	}
@@ -159,13 +159,13 @@ void FLWComponentQuery::ForEachEntityChunk(UMassEntitySubsystem& EntitySubsystem
 			check(Archetype.IsValid());
 			Archetype.DataPtr->ExecuteFunction(ExecutionContext, ExecuteFunction, ArchetypeComponentMapping[i], ChunkCondition);
 			ExecutionContext.ClearComponentViews();
-#if WITH_AGGREGATETICKING_DEBUG
+#if WITH_MASSENTITY_DEBUG
 			NumEntitiesToProcess += Archetype.DataPtr->GetNumEntities();
 #endif
 		}
 	}
 
-#if WITH_AGGREGATETICKING_DEBUG
+#if WITH_MASSENTITY_DEBUG
 	// Not using VLOG to be thread safe
 	UE_CLOG(!ExecutionContext.DebugGetExecutionDesc().IsEmpty(), LogAggregateTicking, VeryVerbose,
 		TEXT("%s: %d entities sent for processing"), *ExecutionContext.DebugGetExecutionDesc(), NumEntitiesToProcess);
@@ -272,7 +272,7 @@ bool FLWComponentQuery::HasMatchingEntities(UMassEntitySubsystem& InEntitySubsys
 
 FString FLWComponentQuery::DebugGetDescription() const
 {
-#if WITH_AGGREGATETICKING_DEBUG
+#if WITH_MASSENTITY_DEBUG
 	TStringBuilder<256> StringBuilder;
 	StringBuilder.Append(TEXT("<"));
 
@@ -299,7 +299,7 @@ FString FLWComponentQuery::DebugGetDescription() const
 
 FString FLWComponentRequirement::DebugGetDescription() const
 {
-#if WITH_AGGREGATETICKING_DEBUG
+#if WITH_MASSENTITY_DEBUG
 	return FString::Printf(TEXT("%s%s[%s]"), IsOptional() ? TEXT("?") : (Presence == ELWComponentPresence::None ? TEXT("-") : TEXT("+"))
 		, *GetNameSafe(StructType), *UE::AggregateTicking::DebugGetComponentAccessString(AccessMode));
 #else
diff --git a/Engine/Plugins/Runtime/MassEntity/Source/MassEntity/Private/MassEntityDebug.cpp b/Engine/Plugins/Runtime/MassEntity/Source/MassEntity/Private/MassEntityDebug.cpp
--- a/Engine/Plugins/Runtime/MassEntity/Source/MassEntity/Private/MassEntityDebug.cpp
+++ b/Engine/Plugins/Runtime/MassEntity/Source/MassEntity/Private/MassEntityDebug.cpp
@@ -13,7 +13,7 @@ namespace UE::Pipe::Debug
 
 void DebugOutputDescription(TConstArrayView<UPipeProcessor*> Processors, FOutputDevice& Ar)
 {
-#if WITH_PIPE_DEBUG
+#if WITH_MASSENTITY_DEBUG
 	const bool bAutoLineEnd = Ar.GetAutoEmitLineTerminator();
 	Ar.SetAutoEmitLineTerminator(false);
 	for (const UPipeProcessor* Proc : Processors)
@@ -29,10 +29,10 @@ void DebugOutputDescription(TConstArrayView<UPipeProcessor*> Processors, FOutput
 		}
 	}
 	Ar.SetAutoEmitLineTerminator(bAutoLineEnd);
-#endif // WITH_PIPE_DEBUG
+#endif // WITH_MASSENTITY_DEBUG
 }
 
-#if WITH_PIPE_DEBUG && WITH_AGGREGATETICKING_DEBUG
+#if WITH_MASSENTITY_DEBUG && WITH_MASSENTITY_DEBUG
 FAutoConsoleCommandWithWorldArgsAndOutputDevice PrintEntityFragmentsCmd(
 	TEXT("pipe.PrintEntityFragments"),
 	TEXT("Prints all fragment types and values (uproperties) for the specified Entity index"),
@@ -127,7 +127,7 @@ FAutoConsoleCommandWithWorldArgsAndOutputDevice LogFragmentSizes(
 			}
 		})
 	);
-#endif // WITH_PIPE_DEBUG && WITH_AGGREGATETICKING_DEBUG
+#endif // WITH_MASSENTITY_DEBUG && WITH_MASSENTITY_DEBUG
 
 } // namespace UE::Pipe::Debug
 
diff --git a/Engine/Plugins/Runtime/MassEntity/Source/MassEntity/Private/MassEntitySubsystem.cpp b/Engine/Plugins/Runtime/MassEntity/Source/MassEntity/Private/MassEntitySubsystem.cpp
--- a/Engine/Plugins/Runtime/MassEntity/Source/MassEntity/Private/MassEntitySubsystem.cpp
+++ b/Engine/Plugins/Runtime/MassEntity/Source/MassEntity/Private/MassEntitySubsystem.cpp
@@ -15,7 +15,7 @@ namespace UE { namespace AggregateTicking {
 
 FString DebugGetComponentAccessString(ELWComponentAccess Access)
 {
-#if WITH_AGGREGATETICKING_DEBUG
+#if WITH_MASSENTITY_DEBUG
 	switch (Access)
 	{
 	case ELWComponentAccess::None:	return TEXT("--");
@@ -25,7 +25,7 @@ FString DebugGetComponentAccessString(ELWComponentAccess Access)
 		ensureMsgf(false, TEXT("Missing string conversion for ELWComponentAccess=%d"), Access);
 		break;
 	}
-#endif // WITH_AGGREGATETICKING_DEBUG
+#endif // WITH_MASSENTITY_DEBUG
 	return TEXT("Missing string conversion");
 }
 
@@ -774,102 +774,102 @@ void UMassEntitySubsystem::GetValidArchetypes(const FLWComponentQuery& Query, TA
 		if (Archetype.GetTagBitSet().HasAll(Query.GetRequiredAllTags()) == false)
 		{
 			// missing some required tags, skip.
-#if WITH_AGGREGATETICKING_DEBUG
+#if WITH_MASSENTITY_DEBUG
 			const FLWTagBitSet UnsatisfiedTags = Query.GetRequiredAllTags() - Archetype.GetTagBitSet();
 			FStringOutputDevice Description;
 			UnsatisfiedTags.DebugGetStringDesc(Description);
 			UE_LOG(LogAggregateTicking, VeryVerbose, TEXT("Archetype did not match due to missing tags: %s")
 				, *Description);
-#endif // WITH_AGGREGATETICKING_DEBUG
+#endif // WITH_MASSENTITY_DEBUG
 			continue;
 		}
 
 		if (Archetype.GetTagBitSet().HasNone(Query.GetRequiredNoneTags()) == false)
 		{
 			// has some tags required to be absent
-#if WITH_AGGREGATETICKING_DEBUG
+#if WITH_MASSENTITY_DEBUG
 			const FLWTagBitSet UnwantedTags = Query.GetRequiredAllTags().GetOverlap(Archetype.GetTagBitSet());
 			FStringOutputDevice Description;
 			UnwantedTags.DebugGetStringDesc(Description);
 			UE_LOG(LogAggregateTicking, VeryVerbose, TEXT("Archetype has tags required absent: %s")
 				, *Description);
-#endif // WITH_AGGREGATETICKING_DEBUG
+#endif // WITH_MASSENTITY_DEBUG
 			continue;
 		}
 
 		if (Query.GetRequiredAnyTags().IsEmpty() == false 
 			&& Archetype.GetTagBitSet().HasAny(Query.GetRequiredAnyTags()) == false)
 		{
-#if WITH_AGGREGATETICKING_DEBUG
+#if WITH_MASSENTITY_DEBUG
 			FStringOutputDevice Description;
 			Query.GetRequiredAnyTags().DebugGetStringDesc(Description);
 			UE_LOG(LogAggregateTicking, VeryVerbose, TEXT("Archetype did not match due to missing \'any\' tags: %s")
 				, *Description);
-#endif // WITH_AGGREGATETICKING_DEBUG
+#endif // WITH_MASSENTITY_DEBUG
 			continue;
 		}
 		
 		if (Archetype.GetComponentBitSet().HasAll(Query.GetRequiredAllComponents()) == false)
 		{
 			// missing some required components, skip.
-#if WITH_AGGREGATETICKING_DEBUG
+#if WITH_MASSENTITY_DEBUG
 			const FLWComponentBitSet UnsatisfiedComponents = Query.GetRequiredAllComponents() - Archetype.GetComponentBitSet();
 			FStringOutputDevice Description;
 			UnsatisfiedComponents.DebugGetStringDesc(Description);
 			UE_LOG(LogAggregateTicking, VeryVerbose, TEXT("Archetype did not match due to missing Components: %s")
 				, *Description);
-#endif // WITH_AGGREGATETICKING_DEBUG
+#endif // WITH_MASSENTITY_DEBUG
 			continue;
 		}
 
 		if (Archetype.GetComponentBitSet().HasNone(Query.GetRequiredNoneComponents()) == false)
 		{
 			// has some Components required to be absent
-#if WITH_AGGREGATETICKING_DEBUG
+#if WITH_MASSENTITY_DEBUG
 			const FLWComponentBitSet UnwantedComponents = Query.GetRequiredAllComponents().GetOverlap(Archetype.GetComponentBitSet());
 			FStringOutputDevice Description;
 			UnwantedComponents.DebugGetStringDesc(Description);
 			UE_LOG(LogAggregateTicking, VeryVerbose, TEXT("Archetype has Components required absent: %s")
 				, *Description);
-#endif // WITH_AGGREGATETICKING_DEBUG
+#endif // WITH_MASSENTITY_DEBUG
 			continue;
 		}
 
 		if (Query.GetRequiredAnyComponents().IsEmpty() == false 
 			&& Archetype.GetComponentBitSet().HasAny(Query.GetRequiredAnyComponents()) == false)
 		{
-#if WITH_AGGREGATETICKING_DEBUG
+#if WITH_MASSENTITY_DEBUG
 			FStringOutputDevice Description;
 			Query.GetRequiredAnyComponents().DebugGetStringDesc(Description);
 			UE_LOG(LogAggregateTicking, VeryVerbose, TEXT("Archetype did not match due to missing \'any\' components: %s")
 				, *Description);
-#endif // WITH_AGGREGATETICKING_DEBUG
+#endif // WITH_MASSENTITY_DEBUG
 			continue;
 		}
 
 		if (Archetype.GetChunkComponentBitSet().HasAll(Query.GetRequiredAllChunkComponents()) == false)
 		{
 			// missing some required components, skip.
-#if WITH_AGGREGATETICKING_DEBUG
+#if WITH_MASSENTITY_DEBUG
 			const FLWChunkComponentBitSet UnsatisfiedComponents = Query.GetRequiredAllChunkComponents() - Archetype.GetChunkComponentBitSet();
 			FStringOutputDevice Description;
 			UnsatisfiedComponents.DebugGetStringDesc(Description);
 			UE_LOG(LogAggregateTicking, VeryVerbose, TEXT("Archetype did not match due to missing Chunk Components: %s")
 				, *Description);
-#endif // WITH_AGGREGATETICKING_DEBUG
+#endif // WITH_MASSENTITY_DEBUG
 			continue;
 		}
 
 		if (Archetype.GetChunkComponentBitSet().HasNone(Query.GetRequiredNoneChunkComponents()) == false)
 		{
 			// has some Components required to be absent
-#if WITH_AGGREGATETICKING_DEBUG
+#if WITH_MASSENTITY_DEBUG
 			const FLWChunkComponentBitSet UnwantedComponents = Query.GetRequiredNoneChunkComponents().GetOverlap(Archetype.GetChunkComponentBitSet());
 			FStringOutputDevice Description;
 			UnwantedComponents.DebugGetStringDesc(Description);
 			UE_LOG(LogAggregateTicking, VeryVerbose, TEXT("Archetype has Chunk Components required absent: %s")
 				, *Description);
-#endif // WITH_AGGREGATETICKING_DEBUG
+#endif // WITH_MASSENTITY_DEBUG
 			continue;
 		}
 
@@ -884,7 +884,7 @@ FLWComponentSystemExecutionContext UMassEntitySubsystem::CreateExecutionContext(
 	return MoveTemp(ExecutionContext);
 }
 
-#if WITH_AGGREGATETICKING_DEBUG
+#if WITH_MASSENTITY_DEBUG
 void UMassEntitySubsystem::DebugPrintArchetypes(FOutputDevice& Ar) const
 {
 	Ar.Logf(ELogVerbosity::Log, TEXT("Listing archetypes contained in %s"), *GetPathNameSafe(this));
@@ -957,7 +957,7 @@ void UMassEntitySubsystem::DebugGetStringDesc(const FArchetypeHandle& Archetype,
 
 void UMassEntitySubsystem::DebugGetArchetypesStringDetails(FOutputDevice& Ar, const bool bIncludeEmpty)
 {
-#if WITH_AGGREGATETICKING_DEBUG
+#if WITH_MASSENTITY_DEBUG
 	Ar.SetAutoEmitLineTerminator(true);
 	for (auto Pair : ComponentHashToArchetypeMap)
 	{
@@ -971,7 +971,7 @@ void UMassEntitySubsystem::DebugGetArchetypesStringDetails(FOutputDevice& Ar, co
 			}
 		}
 	}
-#endif // WITH_AGGREGATETICKING_DEBUG
+#endif // WITH_MASSENTITY_DEBUG
 }
 
 void UMassEntitySubsystem::DebugGetArchetypeComponentTypes(const FArchetypeHandle& Archetype, TArray<const UScriptStruct*>& InOutComponentList) const
@@ -1045,7 +1045,7 @@ FAutoConsoleCommandWithWorldArgsAndOutputDevice GPrintArchetypesCmd(
 		Ar.Logf(ELogVerbosity::Error, TEXT("Failed to find Entity Subsystem for world %s"), *GetPathNameSafe(World));
 	}
 }));
-#endif // WITH_AGGREGATETICKING_DEBUG
+#endif // WITH_MASSENTITY_DEBUG
 
 //////////////////////////////////////////////////////////////////////
 // FLWComponentSystemExecutionContext
diff --git a/Engine/Plugins/Runtime/MassEntity/Source/MassEntity/Private/MassEntityTypes.cpp b/Engine/Plugins/Runtime/MassEntity/Source/MassEntity/Private/MassEntityTypes.cpp
--- a/Engine/Plugins/Runtime/MassEntity/Source/MassEntity/Private/MassEntityTypes.cpp
+++ b/Engine/Plugins/Runtime/MassEntity/Source/MassEntity/Private/MassEntityTypes.cpp
@@ -126,7 +126,7 @@ void FRuntimePipeline::AppendUniqueRuntimeProcessorCopies(TConstArrayView<const
 			UPipeProcessor* ProcCopy = NewObject<UPipeProcessor>(&InOwner, Proc->GetClass(), FName(), RF_NoFlags, const_cast<UPipeProcessor*>(Proc));
 			Processors.Add(ProcCopy);
 		}
-#if WITH_PIPE_DEBUG
+#if WITH_MASSENTITY_DEBUG
 		else if (Proc)
 		{
 			if (Proc->ShouldExecute(WorldExecutionFlags) == false)
@@ -138,7 +138,7 @@ void FRuntimePipeline::AppendUniqueRuntimeProcessorCopies(TConstArrayView<const
 				UE_VLOG(&InOwner, LogPipe, Log, TEXT("Skipping %s due to it being a duplicate"), *Proc->GetName());
 			}
 		}
-#endif // WITH_PIPE_DEBUG
+#endif // WITH_MASSENTITY_DEBUG
 	}
 
 	for (int32 NewProcIndex = StartingCount; NewProcIndex < Processors.Num(); ++NewProcIndex)
diff --git a/Engine/Plugins/Runtime/MassEntity/Source/MassEntity/Private/MassExecutor.cpp b/Engine/Plugins/Runtime/MassEntity/Source/MassEntity/Private/MassExecutor.cpp
--- a/Engine/Plugins/Runtime/MassEntity/Source/MassEntity/Private/MassExecutor.cpp
+++ b/Engine/Plugins/Runtime/MassEntity/Source/MassEntity/Private/MassExecutor.cpp
@@ -74,13 +74,13 @@ void RunProcessorsView(TArrayView<UPipeProcessor*> Processors, FPipeContext& Pip
 		UE_LOG(LogPipe, Error, TEXT("%s PipeContext.EntitySubsystem is null. Baling out."), ANSI_TO_TCHAR(__FUNCTION__));
 		return;
 	}
-#if WITH_PIPE_DEBUG
+#if WITH_MASSENTITY_DEBUG
 	if (Processors.Find(nullptr) != INDEX_NONE)
 	{
 		UE_LOG(LogPipe, Error, TEXT("%s input Processors contains nullptr. Baling out."), ANSI_TO_TCHAR(__FUNCTION__));
 		return;
 	}
-#endif // WITH_PIPE_DEBUG
+#endif // WITH_MASSENTITY_DEBUG
 
 	TRACE_CPUPROFILER_EVENT_SCOPE_STR("PipeExecutor RunProcessorsView")
 
diff --git a/Engine/Plugins/Runtime/MassEntity/Source/MassEntity/Private/MassProcessingPhase.cpp b/Engine/Plugins/Runtime/MassEntity/Source/MassEntity/Private/MassProcessingPhase.cpp
--- a/Engine/Plugins/Runtime/MassEntity/Source/MassEntity/Private/MassProcessingPhase.cpp
+++ b/Engine/Plugins/Runtime/MassEntity/Source/MassEntity/Private/MassProcessingPhase.cpp
@@ -153,7 +153,7 @@ void UPipeProcessingPhaseManager::InitializePhases(UObject& InProcessorOwner)
 		PhaseProcessor->Initialize(InProcessorOwner);
 	}
 
-#if WITH_PIPE_DEBUG
+#if WITH_MASSENTITY_DEBUG
 	// print it all out to vislog
 	UE_VLOG_UELOG(this, LogPipe, Verbose, TEXT("Phases initialization done. Current composition:"));
 
@@ -167,7 +167,7 @@ void UPipeProcessingPhaseManager::InitializePhases(UObject& InProcessorOwner)
 			OutDescription.Reset();
 		}
 	}	
-#endif // WITH_PIPE_DEBUG
+#endif // WITH_MASSENTITY_DEBUG
 }
 
 void UPipeProcessingPhaseManager::Start(UWorld& World)
diff --git a/Engine/Plugins/Runtime/MassEntity/Source/MassEntity/Private/MassProcessor.cpp b/Engine/Plugins/Runtime/MassEntity/Source/MassEntity/Private/MassProcessor.cpp
--- a/Engine/Plugins/Runtime/MassEntity/Source/MassEntity/Private/MassProcessor.cpp
+++ b/Engine/Plugins/Runtime/MassEntity/Source/MassEntity/Private/MassProcessor.cpp
@@ -19,21 +19,21 @@ namespace UE::MassTraffic
 }
 #endif // PARALLELIZED_TRAFFIC_HACK
 
-#if WITH_PIPE_DEBUG
+#if WITH_MASSENTITY_DEBUG
 namespace UE::Mass::Debug
 {
 	bool bLogProcessingGraph = false;
 	FAutoConsoleVariableRef CVarLogProcessingGraph(TEXT("pipe.LogProcessingGraph"), bLogProcessingGraph
 		, TEXT("When enabled will log task graph tasks created while dispatching processors to other threads, along with their dependencies"), ECVF_Cheat);
 }
-#endif // WITH_PIPE_DEBUG
+#endif // WITH_MASSENTITY_DEBUG
 
 // change to && 1 to enable more detailed processing tasks logging
-#if WITH_PIPE_DEBUG && 0
+#if WITH_MASSENTITY_DEBUG && 0
 #define PROCESSOR_LOG(Fmt, ...) UE_LOG(LogPipe, Log, Fmt, ##__VA_ARGS__)
-#else // WITH_PIPE_DEBUG
+#else // WITH_MASSENTITY_DEBUG
 #define PROCESSOR_LOG(...) 
-#endif // WITH_PIPE_DEBUG
+#endif // WITH_MASSENTITY_DEBUG
 
 namespace FPipeTweakables
 {
@@ -151,7 +151,7 @@ void UPipeProcessor::PostInitProperties()
 void UPipeProcessor::CallExecute(UMassEntitySubsystem& EntitySubsystem, FLWComponentSystemExecutionContext& Context)
 {
 	TRACE_CPUPROFILER_EVENT_SCOPE_TEXT(*StatId);
-#if WITH_PIPE_DEBUG
+#if WITH_MASSENTITY_DEBUG
 	Context.DebugSetExecutionDesc(FString::Printf(TEXT("%s (%s)"), *GetProcessorName(), *ToString(EntitySubsystem.GetWorld()->GetNetMode())));
 #endif
 	Execute(EntitySubsystem, Context);
@@ -173,9 +173,9 @@ FGraphEventRef UPipeProcessor::DispatchProcessorTasks(UMassEntitySubsystem& Enti
 
 void UPipeProcessor::DebugOutputDescription(FOutputDevice& Ar, int32 Indent) const
 {
-#if WITH_PIPE_DEBUG
+#if WITH_MASSENTITY_DEBUG
 	Ar.Logf(TEXT("%*s%s"), Indent, TEXT(""), *GetProcessorName());
-#endif // WITH_PIPE_DEBUG
+#endif // WITH_MASSENTITY_DEBUG
 }
 
 #if WITH_EDITOR
@@ -241,7 +241,7 @@ FGraphEventRef UPipeCompositeProcessor::DispatchProcessorTasks(UMassEntitySubsys
 	}
 
 
-#if WITH_PIPE_DEBUG
+#if WITH_MASSENTITY_DEBUG
 	if (UE::Mass::Debug::bLogProcessingGraph)
 	{
 		for (int i = 0; i < ProcessingFlatGraph.Num(); ++i)
@@ -265,7 +265,7 @@ FGraphEventRef UPipeCompositeProcessor::DispatchProcessorTasks(UMassEntitySubsys
 			}
 		}
 	}
-#endif // WITH_PIPE_DEBUG
+#endif // WITH_MASSENTITY_DEBUG
 
 	FGraphEventRef CompletionEvent = FFunctionGraphTask::CreateAndDispatchWhenReady([this](){}
 		, GET_STATID(Pipe_GroupCompletedTask), &Events, ENamedThreads::AnyHiPriThreadHiPriTask);
@@ -553,7 +553,7 @@ int32 UPipeCompositeProcessor::Populate(TArray<FProcessorDependencySolver::FOrde
 
 void UPipeCompositeProcessor::DebugOutputDescription(FOutputDevice& Ar, int32 Indent) const
 {
-#if WITH_PIPE_DEBUG
+#if WITH_MASSENTITY_DEBUG
 	if (ChildPipeline.Processors.Num() == 0)
 	{
 		Ar.Logf(TEXT("%*sGroup %s: []"), Indent, TEXT(""), *GroupName.ToString());
@@ -568,7 +568,7 @@ void UPipeCompositeProcessor::DebugOutputDescription(FOutputDevice& Ar, int32 In
 			Proc->DebugOutputDescription(Ar, Indent + 3);
 		}
 	}
-#endif // WITH_PIPE_DEBUG
+#endif // WITH_MASSENTITY_DEBUG
 }
 
 void UPipeCompositeProcessor::SetProcessingPhase(EPipeProcessingPhase Phase)
diff --git a/Engine/Plugins/Runtime/MassEntity/Source/MassEntity/Public/LWComponentTypes.h b/Engine/Plugins/Runtime/MassEntity/Source/MassEntity/Public/LWComponentTypes.h
--- a/Engine/Plugins/Runtime/MassEntity/Source/MassEntity/Public/LWComponentTypes.h
+++ b/Engine/Plugins/Runtime/MassEntity/Source/MassEntity/Public/LWComponentTypes.h
@@ -5,9 +5,6 @@
 #include "ScriptStructTypeBitSet.h"
 #include "LWComponentTypes.generated.h"
 
-// the name is temporary. To be changed when we arrive at the final name for the plugin
-#define WITH_AGGREGATETICKING_DEBUG (!(UE_BUILD_SHIPPING || UE_BUILD_SHIPPING_WITH_EDITOR || UE_BUILD_TEST) && WITH_STRUCTUTILS_DEBUG && 1)
-
 // This is the base class for all lightweight components
 USTRUCT()
 struct FLWComponentData
diff --git a/Engine/Plugins/Runtime/MassEntity/Source/MassEntity/Public/MassEntitySubsystem.h b/Engine/Plugins/Runtime/MassEntity/Source/MassEntity/Public/MassEntitySubsystem.h
--- a/Engine/Plugins/Runtime/MassEntity/Source/MassEntity/Public/MassEntitySubsystem.h
+++ b/Engine/Plugins/Runtime/MassEntity/Source/MassEntity/Public/MassEntitySubsystem.h
@@ -4,6 +4,7 @@
 
 #include "Subsystems/WorldSubsystem.h"
 #include "LWComponentTypes.h"
+#include "MassEntityTypes.h"
 #include "InstancedStruct.h"
 #include "EntityQuery.h"
 #include "MassEntitySubsystem.generated.h"
@@ -256,7 +257,7 @@ public:
 	bool IsProcessing() const { return ProcessingScopeCount > 0; }
 	FLWCCommandBuffer& Defer() { return *DeferredCommandBuffer.Get(); }
 
-#if WITH_AGGREGATETICKING_DEBUG
+#if WITH_MASSENTITY_DEBUG
 	void DebugPrintEntity(int32 Index, FOutputDevice& Ar, const TCHAR* InPrefix = TEXT("")) const;
 	void DebugPrintEntity(FLWEntity Entity, FOutputDevice& Ar, const TCHAR* InPrefix = TEXT("")) const;
 	void DebugPrintArchetypes(FOutputDevice& Ar) const;
@@ -269,7 +270,7 @@ public:
 	void DebugRemoveAllEntities();
 	void DebugForceArchetypeDataVersionBump() { ++ArchetypeDataVersion; }
 	void DebugGetArchetypeStrings(const FArchetypeHandle& Archetype, TArray<FName>& OutComponentNames, TArray<FName>& OutTagNames);
-#endif // WITH_AGGREGATETICKING_DEBUG
+#endif // WITH_MASSENTITY_DEBUG
 
 protected:
 	void GetValidArchetypes(const FLWComponentQuery& Query, TArray<FArchetypeHandle>& OutValidArchetypes);
@@ -376,7 +377,7 @@ private:
 	int32 ChunkSerialModificationNumber = -1;
 	FLWTagBitSet CurrentArchetypesTagBitSet;
 
-#if WITH_AGGREGATETICKING_DEBUG
+#if WITH_MASSENTITY_DEBUG
 	FString DebugExecutionDescription;
 #endif
 	
@@ -398,7 +399,7 @@ public:
 		, bFlushDeferredCommands(bInFlushDeferredCommands)
 	{}
 
-#if WITH_AGGREGATETICKING_DEBUG
+#if WITH_MASSENTITY_DEBUG
 	const FString& DebugGetExecutionDesc() const { return DebugExecutionDescription; }
 	void DebugSetExecutionDesc(const FString& Description) { DebugExecutionDescription = Description; }
 #endif
diff --git a/Engine/Plugins/Runtime/MassEntity/Source/MassEntity/Public/MassEntityTypes.h b/Engine/Plugins/Runtime/MassEntity/Source/MassEntity/Public/MassEntityTypes.h
--- a/Engine/Plugins/Runtime/MassEntity/Source/MassEntity/Public/MassEntityTypes.h
+++ b/Engine/Plugins/Runtime/MassEntity/Source/MassEntity/Public/MassEntityTypes.h
@@ -2,11 +2,11 @@
 
 #pragma once
 
+#include "StructUtilsTypes.h"
 #include "InstancedStruct.h"
-#include "LWComponentTypes.h"
 #include "MassEntityTypes.generated.h"
 
-#define WITH_PIPE_DEBUG (!(UE_BUILD_SHIPPING || UE_BUILD_SHIPPING_WITH_EDITOR || UE_BUILD_TEST) && WITH_AGGREGATETICKING_DEBUG && 1)
+#define WITH_MASSENTITY_DEBUG (!(UE_BUILD_SHIPPING || UE_BUILD_SHIPPING_WITH_EDITOR || UE_BUILD_TEST) && WITH_STRUCTUTILS_DEBUG && 1)
 
 MASSENTITY_API DECLARE_LOG_CATEGORY_EXTERN(LogPipe, Warning, All);
 
diff --git a/Engine/Plugins/Runtime/MassEntity/Source/MassEntityTestSuite/Private/MassEntityTest.cpp b/Engine/Plugins/Runtime/MassEntity/Source/MassEntityTestSuite/Private/MassEntityTest.cpp
--- a/Engine/Plugins/Runtime/MassEntity/Source/MassEntityTestSuite/Private/MassEntityTest.cpp
+++ b/Engine/Plugins/Runtime/MassEntity/Source/MassEntityTestSuite/Private/MassEntityTest.cpp
@@ -15,7 +15,7 @@ PRAGMA_DISABLE_OPTIMIZATION
 
 namespace FPipeEntityTest
 {
-#if WITH_AGGREGATETICKING_DEBUG
+#if WITH_MASSENTITY_DEBUG
 struct FEntityTest_ArchetypeCreation : FEntityTestBase
 {
 	virtual bool InstantTest() override
@@ -355,7 +355,7 @@ struct FEntityTest_ReserveAPreviouslyBuiltEntity : FEntityTestBase
 };
 IMPLEMENT_AI_INSTANT_TEST(FEntityTest_ReserveAPreviouslyBuiltEntity, "System.Pipe.Entity.ReserveAPreviouslyBuiltEntity");
 
-#endif // WITH_AGGREGATETICKING_DEBUG
+#endif // WITH_MASSENTITY_DEBUG
 
 } // FPipeEntityTestTest
 
diff --git a/Engine/Plugins/Runtime/MassEntity/Source/MassEntityTestSuite/Private/MassExecutionTest.cpp b/Engine/Plugins/Runtime/MassEntity/Source/MassEntityTestSuite/Private/MassExecutionTest.cpp
--- a/Engine/Plugins/Runtime/MassEntity/Source/MassEntityTestSuite/Private/MassExecutionTest.cpp
+++ b/Engine/Plugins/Runtime/MassEntity/Source/MassEntityTestSuite/Private/MassExecutionTest.cpp
@@ -81,7 +81,7 @@ struct FExecution_InvalidPipeContext : FExecutionTestBase
 IMPLEMENT_AI_INSTANT_TEST(FExecution_InvalidPipeContext, "System.Pipe.Execution.InvalidPipeContext");
 
 
-#if WITH_PIPE_DEBUG
+#if WITH_MASSENTITY_DEBUG
 struct FExecution_SingleNullProcessor : FExecutionTestBase
 {
 	virtual bool InstantTest() override
@@ -141,7 +141,7 @@ struct FExecution_MultipleNullProcessors : FExecutionTestBase
 	}
 };
 IMPLEMENT_AI_INSTANT_TEST(FExecution_MultipleNullProcessors, "System.Pipe.Execution.MultipleNullProcessors");
-#endif // WITH_PIPE_DEBUG
+#endif // WITH_MASSENTITY_DEBUG
 
 
 struct FExecution_Sparse : FEntityTestBase
diff --git a/Engine/Plugins/Runtime/MassGameplay/Source/MassCommon/Private/MassCommonTypes.cpp b/Engine/Plugins/Runtime/MassGameplay/Source/MassCommon/Private/MassCommonTypes.cpp
--- a/Engine/Plugins/Runtime/MassGameplay/Source/MassCommon/Private/MassCommonTypes.cpp
+++ b/Engine/Plugins/Runtime/MassGameplay/Source/MassCommon/Private/MassCommonTypes.cpp
@@ -4,7 +4,7 @@
 #include "MassEntitySubsystem.h"
 #include "Math/ColorList.h"
 
-#if WITH_MASS_DEBUG
+#if WITH_MASSGAMEPLAY_DEBUG
 namespace UE { namespace MassDebug {
 
 // First Id of a range of lightweight entity for which we want to activate debug information
@@ -101,7 +101,7 @@ FColor GetEntityDebugColor(FLWEntity Entity)
 }
 
 }} // namespace UE::MassDebug
-#endif // WITH_MASS_DEBUG
+#endif // WITH_MASSGAMEPLAY_DEBUG
 
 const FMassHandle FMassHandle::InvalidHandle(UMassEntitySubsystem::InvalidEntity);
 
@@ -137,7 +137,7 @@ void FMassUniqueFragmentCollection::Append(const FMassUniqueFragmentCollection&
 
 void FMassUniqueFragmentCollection::DebugOutputDescription(FOutputDevice& Ar) const
 {
-#if WITH_MASS_DEBUG
+#if WITH_MASSGAMEPLAY_DEBUG
 	const bool bAutoLineEnd = Ar.GetAutoEmitLineTerminator();
 	Ar.SetAutoEmitLineTerminator(false);
 
@@ -147,5 +147,5 @@ void FMassUniqueFragmentCollection::DebugOutputDescription(FOutputDevice& Ar) co
 		Ar.Logf(TEXT("%s, "), *GetNameSafe(StructType));
 	}
 	Ar.SetAutoEmitLineTerminator(bAutoLineEnd);
-#endif // WITH_MASS_DEBUG
+#endif // WITH_MASSGAMEPLAY_DEBUG
 }
diff --git a/Engine/Plugins/Runtime/MassGameplay/Source/MassCommon/Public/MassCommonTypes.h b/Engine/Plugins/Runtime/MassGameplay/Source/MassCommon/Public/MassCommonTypes.h
--- a/Engine/Plugins/Runtime/MassGameplay/Source/MassCommon/Public/MassCommonTypes.h
+++ b/Engine/Plugins/Runtime/MassGameplay/Source/MassCommon/Public/MassCommonTypes.h
@@ -3,25 +3,26 @@
 #pragma once
 
 #include "LWComponentTypes.h"
+#include "MassEntityTypes.h"
 #include "InstancedStruct.h"
 #include "SequentialID.h"
 #include "MassCommonTypes.generated.h"
 
-#define WITH_MASS_DEBUG (!(UE_BUILD_SHIPPING || UE_BUILD_SHIPPING_WITH_EDITOR || UE_BUILD_TEST) && WITH_AGGREGATETICKING_DEBUG && 1)
+#define WITH_MASSGAMEPLAY_DEBUG (!(UE_BUILD_SHIPPING || UE_BUILD_SHIPPING_WITH_EDITOR || UE_BUILD_TEST) && WITH_MASSENTITY_DEBUG && 1)
 
 class UStaticMesh;
 class UMaterialInterface;
 
 MASSCOMMON_API DECLARE_LOG_CATEGORY_EXTERN(LogMass, Warning, All);
 
-#if WITH_MASS_DEBUG
+#if WITH_MASSGAMEPLAY_DEBUG
 namespace UE::MassDebug
 {
 	MASSCOMMON_API extern bool HasDebugEntities();
 	MASSCOMMON_API extern bool IsDebuggingEntity(FLWEntity Entity, FColor* OutEntityColor = nullptr);
 	MASSCOMMON_API extern FColor GetEntityDebugColor(FLWEntity Entity);
 } // namespace UE::MassDebug
-#endif // WITH_MASS_DEBUG
+#endif // WITH_MASSGAMEPLAY_DEBUG
 
 namespace UE::Mass::ProcessorGroupNames
 {
@@ -55,11 +56,11 @@ public:
 
 	FString DebugGetDescription() const 
 	{
-#if WITH_MASS_DEBUG
+#if WITH_MASSGAMEPLAY_DEBUG
 		return LWEntity.DebugGetDescription();
 #else
 		return FString();
-#endif // WITH_MASS_DEBUG
+#endif // WITH_MASSGAMEPLAY_DEBUG
 	}
 
 	const static FMassHandle InvalidHandle;
diff --git a/Engine/Plugins/Runtime/MassGameplay/Source/MassLOD/Private/MassLODUtils.cpp b/Engine/Plugins/Runtime/MassGameplay/Source/MassLOD/Private/MassLODUtils.cpp
--- a/Engine/Plugins/Runtime/MassGameplay/Source/MassLOD/Private/MassLODUtils.cpp
+++ b/Engine/Plugins/Runtime/MassGameplay/Source/MassLOD/Private/MassLODUtils.cpp
@@ -3,11 +3,11 @@
 #include "MassLODUtils.h"
 
 DEFINE_LOG_CATEGORY(LogMassLOD);
-#if WITH_MASS_DEBUG
+#if WITH_MASSGAMEPLAY_DEBUG
 namespace UE::MassLOD::Debug
 {
 	bool bLODCalculationsPaused = false;
 	
 	FAutoConsoleVariableRef CVarLODPause(TEXT("mass.lod.pause"), bLODCalculationsPaused, TEXT("If non zero will pause all LOD calculations"));
 } // UE::MassLOD::Debug
-#endif // WITH_MASS_DEBUG
\ No newline at end of file
+#endif // WITH_MASSGAMEPLAY_DEBUG
\ No newline at end of file
diff --git a/Engine/Plugins/Runtime/MassGameplay/Source/MassLOD/Public/MassLODCalculator.h b/Engine/Plugins/Runtime/MassGameplay/Source/MassLOD/Public/MassLODCalculator.h
--- a/Engine/Plugins/Runtime/MassGameplay/Source/MassLOD/Public/MassLODCalculator.h
+++ b/Engine/Plugins/Runtime/MassGameplay/Source/MassLOD/Public/MassLODCalculator.h
@@ -263,12 +263,12 @@ template <typename FLODLogic>
 template< typename FMassLODResultInfo, typename FMassLODSourceInfo >
 void TMassLODCalculator<FLODLogic>::CalculateLOD(FLWComponentSystemExecutionContext& Context, TArrayView<FMassLODResultInfo> LODList, TConstArrayView<FMassLODSourceInfo> ViewersInfoList)
 {
-#if WITH_MASS_DEBUG
+#if WITH_MASSGAMEPLAY_DEBUG
 	if (UE::MassLOD::Debug::bLODCalculationsPaused)
 	{
 		return;
 	}
-#endif // WITH_MASS_DEBUG
+#endif // WITH_MASSGAMEPLAY_DEBUG
 
 	const int32 NumEntities = Context.GetEntitiesNum();
 	for (int EntityIdx = 0; EntityIdx < NumEntities; EntityIdx++)
diff --git a/Engine/Plugins/Runtime/MassGameplay/Source/MassLOD/Public/MassLODUtils.h b/Engine/Plugins/Runtime/MassGameplay/Source/MassLOD/Public/MassLODUtils.h
--- a/Engine/Plugins/Runtime/MassGameplay/Source/MassLOD/Public/MassLODUtils.h
+++ b/Engine/Plugins/Runtime/MassGameplay/Source/MassLOD/Public/MassLODUtils.h
@@ -72,10 +72,10 @@ inline bool IsLODTagSet(const FLWComponentSystemExecutionContext& Context, EMass
 	}
 }
 
-#if WITH_MASS_DEBUG
+#if WITH_MASSGAMEPLAY_DEBUG
 namespace Debug 
 {
 	MASSLOD_API extern bool bLODCalculationsPaused;
 } // UE::MassLOD::Debug
-#endif // WITH_MASS_DEBUG
+#endif // WITH_MASSGAMEPLAY_DEBUG
 } // UE::MassLOD
\ No newline at end of file
diff --git a/Engine/Plugins/Runtime/MassGameplay/Source/MassReplication/Private/MassClientBubbleInfoBase.cpp b/Engine/Plugins/Runtime/MassGameplay/Source/MassReplication/Private/MassClientBubbleInfoBase.cpp
--- a/Engine/Plugins/Runtime/MassGameplay/Source/MassReplication/Private/MassClientBubbleInfoBase.cpp
+++ b/Engine/Plugins/Runtime/MassGameplay/Source/MassReplication/Private/MassClientBubbleInfoBase.cpp
@@ -73,15 +73,15 @@ void AMassClientBubbleInfoBase::BeginPlay()
 	UWorld* World = GetWorld();
 	check(World);
 
-	// Allow tick function WITH_MASS_DEBUG as we have debug functionality there, otherwise only on the clients.
-#if WITH_MASS_DEBUG
+	// Allow tick function WITH_MASSGAMEPLAY_DEBUG as we have debug functionality there, otherwise only on the clients.
+#if WITH_MASSGAMEPLAY_DEBUG
 	PrimaryActorTick.SetTickFunctionEnable(true);
 #else
 	if (World && World->GetNetMode() == NM_Client)
 	{
 		PrimaryActorTick.SetTickFunctionEnable(true);
 	}
-#endif // WITH_MASS_DEBUG
+#endif // WITH_MASSGAMEPLAY_DEBUG
 }
 
 void AMassClientBubbleInfoBase::EndPlay(const EEndPlayReason::Type EndPlayReason)
diff --git a/Engine/Plugins/Runtime/MassGameplay/Source/MassSmartObjects/Private/MassSmartObjectHandler.cpp b/Engine/Plugins/Runtime/MassGameplay/Source/MassSmartObjects/Private/MassSmartObjectHandler.cpp
--- a/Engine/Plugins/Runtime/MassGameplay/Source/MassSmartObjects/Private/MassSmartObjectHandler.cpp
+++ b/Engine/Plugins/Runtime/MassGameplay/Source/MassSmartObjects/Private/MassSmartObjectHandler.cpp
@@ -127,7 +127,7 @@ EMassSmartObjectClaimResult FMassSmartObjectHandler::ClaimCandidate(const FLWEnt
 	{
 		if (ClaimSmartObject(Entity, User, Candidate.ID))
 		{
-#if WITH_MASS_DEBUG
+#if WITH_MASSGAMEPLAY_DEBUG
 			UE_CVLOG(UE::MassDebug::IsDebuggingEntity(Entity),
 				&SmartObjectSubsystem,
 				LogSmartObject,
@@ -135,7 +135,7 @@ EMassSmartObjectClaimResult FMassSmartObjectHandler::ClaimCandidate(const FLWEnt
 				TEXT("[%s] claimed [%s]"),
 				*Entity.DebugGetDescription(),
 				*Candidate.ID.Describe());
-#endif // WITH_MASS_DEBUG
+#endif // WITH_MASSGAMEPLAY_DEBUG
 			break;
 		}
 	}
@@ -164,7 +164,7 @@ bool FMassSmartObjectHandler::ClaimSmartObject(const FLWEntity Entity, FDataFrag
 		SmartObjectSubsystem.RegisterSlotInvalidationCallback(ClaimHandle, FOnSlotInvalidated::CreateStatic(&UE::Mass::SmartObject::OnSlotInvalidated, Payload));
 	}
 
-#if WITH_MASS_DEBUG
+#if WITH_MASSGAMEPLAY_DEBUG
 	UE_CVLOG(UE::MassDebug::IsDebuggingEntity(Entity),
 		&SmartObjectSubsystem,
 		LogSmartObject,
@@ -173,7 +173,7 @@ bool FMassSmartObjectHandler::ClaimSmartObject(const FLWEntity Entity, FDataFrag
 		*Entity.DebugGetDescription(),
 		*ObjectID.Describe(),
 		bSuccess ? TEXT("Succeeded") : TEXT("Failed"));
-#endif // WITH_MASS_DEBUG
+#endif // WITH_MASSGAMEPLAY_DEBUG
 
 	return bSuccess;
 }
@@ -183,7 +183,7 @@ bool FMassSmartObjectHandler::UseSmartObject(
 	FDataFragment_SmartObjectUser& User,
 	const FDataFragment_Transform& Transform) const
 {
-#if WITH_MASS_DEBUG
+#if WITH_MASSGAMEPLAY_DEBUG
 	UE_CVLOG(UE::MassDebug::IsDebuggingEntity(Entity),
 		&SmartObjectSubsystem,
 		LogSmartObject,
@@ -191,7 +191,7 @@ bool FMassSmartObjectHandler::UseSmartObject(
 		TEXT("[%s] starts using [%s]"),
 		*Entity.DebugGetDescription(),
 		*User.ClaimHandle.Describe());
-#endif // WITH_MASS_DEBUG
+#endif // WITH_MASSGAMEPLAY_DEBUG
 
 	const USmartObjectMassBehaviorConfig* Config = SmartObjectSubsystem.Use<USmartObjectMassBehaviorConfig>(User.ClaimHandle);
 	if (Config == nullptr)
@@ -209,7 +209,7 @@ bool FMassSmartObjectHandler::UseSmartObject(
 
 void FMassSmartObjectHandler::ReleaseSmartObject(const FLWEntity Entity, FDataFragment_SmartObjectUser& User, const EMassSmartObjectInteractionStatus NewStatus) const
 {
-#if WITH_MASS_DEBUG
+#if WITH_MASSGAMEPLAY_DEBUG
 	UE_CVLOG(UE::MassDebug::IsDebuggingEntity(Entity),
 		&SmartObjectSubsystem,
 		LogSmartObject,
@@ -217,7 +217,7 @@ void FMassSmartObjectHandler::ReleaseSmartObject(const FLWEntity Entity, FDataFr
 		TEXT("[%s] releases handle [%s]"),
 		*Entity.DebugGetDescription(),
 		*User.ClaimHandle.Describe());
-#endif // WITH_MASS_DEBUG
+#endif // WITH_MASSGAMEPLAY_DEBUG
 
 #if DO_ENSURE
 	const EMassSmartObjectInteractionStatus CurrentStatus = User.InteractionStatus;
diff --git a/Engine/Plugins/Runtime/MassGameplay/Source/MassSmartObjects/Private/MassSmartObjectProcessor.cpp b/Engine/Plugins/Runtime/MassGameplay/Source/MassSmartObjects/Private/MassSmartObjectProcessor.cpp
--- a/Engine/Plugins/Runtime/MassGameplay/Source/MassSmartObjects/Private/MassSmartObjectProcessor.cpp
+++ b/Engine/Plugins/Runtime/MassGameplay/Source/MassSmartObjects/Private/MassSmartObjectProcessor.cpp
@@ -85,9 +85,9 @@ void UMassProcessor_SmartObjectCandidatesFinder::Execute(UMassEntitySubsystem& E
 		}
 		Result.bProcessed = true;
 
-#if WITH_MASS_DEBUG
+#if WITH_MASSGAMEPLAY_DEBUG
 		UE_VLOG(LogOwner, LogSmartObject, Verbose, TEXT("[%s] search completed: found %d"), *Entity.DebugGetDescription(), Result.NumCandidates);
-#endif // WITH_MASS_DEBUG
+#endif // WITH_MASSGAMEPLAY_DEBUG
 	};
 
 	// Process world location based requests
@@ -114,9 +114,9 @@ void UMassProcessor_SmartObjectCandidatesFinder::Execute(UMassEntitySubsystem& E
 			bool bDisplayDebug = false;
 			FColor DebugColor(FColor::White);
 
-#if WITH_MASS_DEBUG
+#if WITH_MASSGAMEPLAY_DEBUG
 			bDisplayDebug = UE::MassDebug::IsDebuggingEntity(Entity, &DebugColor);
-#endif // WITH_MASS_DEBUG
+#endif // WITH_MASSGAMEPLAY_DEBUG
 
 			BeginRequestProcessing(Entity, Context, Result);
 			ON_SCOPE_EXIT{ EndRequestProcessing(SmartObjectSubsystem, Entity, Result);	};
@@ -133,14 +133,14 @@ void UMassProcessor_SmartObjectCandidatesFinder::Execute(UMassEntitySubsystem& E
 						const FVector ObjectLocation = Element.Bounds.Center;
 						Result.Candidates[Result.NumCandidates++] = FSmartObjectCandidate(ID, FVector::DistSquared(SearchOrigin, ObjectLocation));
 
-#if WITH_MASS_DEBUG
+#if WITH_MASSGAMEPLAY_DEBUG
 						if (bDisplayDebug)
 						{
 							constexpr float DebugRadius = 10.f;
 							UE_VLOG_LOCATION(SmartObjectSubsystem, LogSmartObject, Display, ObjectLocation, DebugRadius, DebugColor, TEXT("%s"), *ID.Describe());
 							UE_VLOG_SEGMENT(SmartObjectSubsystem, LogSmartObject, Display, SearchOrigin, ObjectLocation, DebugColor, TEXT(""));
 						}
-#endif // WITH_MASS_DEBUG
+#endif // WITH_MASSGAMEPLAY_DEBUG
 					}
 				}
 			});
@@ -175,10 +175,10 @@ void UMassProcessor_SmartObjectCandidatesFinder::Execute(UMassEntitySubsystem& E
 
 				const FLWEntity Entity = Context.GetEntity(i);
 				bool bDisplayDebug = false;
-#if WITH_MASS_DEBUG
+#if WITH_MASSGAMEPLAY_DEBUG
 				FColor DebugColor(FColor::White);
 				bDisplayDebug = UE::MassDebug::IsDebuggingEntity(Entity, &DebugColor);
-#endif // WITH_MASS_DEBUG
+#endif // WITH_MASSGAMEPLAY_DEBUG
 
 				BeginRequestProcessing(Entity, Context, Result);
 				ON_SCOPE_EXIT{ EndRequestProcessing(SmartObjectSubsystem, Entity, Result); };
@@ -239,7 +239,7 @@ void UMassProcessor_SmartObjectCandidatesFinder::Execute(UMassEntitySubsystem& E
 
 					Result.Candidates[Result.NumCandidates++] = FSmartObjectCandidate(ID, Cost);
 
-#if WITH_MASS_DEBUG
+#if WITH_MASSGAMEPLAY_DEBUG
 					if (bDisplayDebug)
 					{
 						FZoneGraphLaneLocation RequestLaneLocation, EntryPointLaneLocation;
@@ -250,7 +250,7 @@ void UMassProcessor_SmartObjectCandidatesFinder::Execute(UMassEntitySubsystem& E
 						UE_VLOG_LOCATION(SmartObjectSubsystem, LogSmartObject, Display, EntryPointLaneLocation.Position, DebugRadius, DebugColor, TEXT("%s"), *ID.Describe());
 						UE_VLOG_SEGMENT(SmartObjectSubsystem, LogSmartObject, Display, RequestLaneLocation.Position, EntryPointLaneLocation.Position, DebugColor, TEXT(""));
 					}
-#endif // WITH_MASS_DEBUG
+#endif // WITH_MASSGAMEPLAY_DEBUG
 				}
 			}
 		});
@@ -305,7 +305,7 @@ void UMassProcessor_SmartObjectTimedBehavior::Execute(UMassEntitySubsystem& Enti
 			SOUser.SetUseTime(FMath::Max(SOUser.GetUseTime() - DT, 0.0f));
 			const bool bMustRelease = SOUser.GetUseTime() <= 0.f;
 
-#if WITH_MASS_DEBUG
+#if WITH_MASSGAMEPLAY_DEBUG
 			const FLWEntity Entity = Context.GetEntity(i);
 			FColor DebugColor(FColor::White);
 			const bool bIsDebuggingEntity = UE::MassDebug::IsDebuggingEntity(Entity, &DebugColor);
@@ -325,7 +325,7 @@ void UMassProcessor_SmartObjectTimedBehavior::Execute(UMassEntitySubsystem& Enti
 					UE_VLOG_ARROW(SmartObjectSubsystem, LogSmartObject, Display, Pos, Pos + Dir * 2.0f * Radius, DebugColor, TEXT(""));
 				}
 			}
-#endif // WITH_MASS_DEBUG
+#endif // WITH_MASSGAMEPLAY_DEBUG
 
 			if (bMustRelease)
 			{
diff --git a/Engine/Plugins/Runtime/MassGameplay/Source/MassSpawner/Private/MassEntityTemplate.cpp b/Engine/Plugins/Runtime/MassGameplay/Source/MassSpawner/Private/MassEntityTemplate.cpp
--- a/Engine/Plugins/Runtime/MassGameplay/Source/MassSpawner/Private/MassEntityTemplate.cpp
+++ b/Engine/Plugins/Runtime/MassGameplay/Source/MassSpawner/Private/MassEntityTemplate.cpp
@@ -32,7 +32,7 @@ void FMassEntityTemplate::SetUpProcessors(const FMassObjectHandlers& ObjectHandl
 FString FMassEntityTemplate::DebugGetDescription(UMassEntitySubsystem* EntitySubsystem) const
 { 
 	FStringOutputDevice Ar;
-#if WITH_MASS_DEBUG
+#if WITH_MASSGAMEPLAY_DEBUG
 	Ar.SetAutoEmitLineTerminator(true);
 
 	Ar += TEXT("Fragments:\n");
@@ -46,15 +46,15 @@ FString FMassEntityTemplate::DebugGetDescription(UMassEntitySubsystem* EntitySub
 		Ar += TEXT("Archetype detais:\n");
 		Ar += DebugGetArchetypeDescription(*EntitySubsystem);
 	}
-#endif // WITH_MASS_DEBUG
+#endif // WITH_MASSGAMEPLAY_DEBUG
 	return MoveTemp(Ar);
 }
 
 FString FMassEntityTemplate::DebugGetArchetypeDescription(UMassEntitySubsystem& EntitySubsystem) const
 {
 	FStringOutputDevice OutDescription;
-#if WITH_MASS_DEBUG
+#if WITH_MASSGAMEPLAY_DEBUG
 	EntitySubsystem.DebugGetStringDesc(Archetype, OutDescription);
-#endif // WITH_MASS_DEBUG
+#endif // WITH_MASSGAMEPLAY_DEBUG
 	return MoveTemp(OutDescription);
 }
diff --git a/Engine/Plugins/Runtime/MassGameplay/Source/MassSpawner/Private/MassEntityTemplateRegistry.cpp b/Engine/Plugins/Runtime/MassGameplay/Source/MassSpawner/Private/MassEntityTemplateRegistry.cpp
--- a/Engine/Plugins/Runtime/MassGameplay/Source/MassSpawner/Private/MassEntityTemplateRegistry.cpp
+++ b/Engine/Plugins/Runtime/MassGameplay/Source/MassSpawner/Private/MassEntityTemplateRegistry.cpp
@@ -266,10 +266,10 @@ void UMassEntityTemplateRegistry::InitializeEntityTemplate(FMassEntityTemplate&
 
 void UMassEntityTemplateRegistry::DebugReset()
 {
-#if WITH_MASS_DEBUG
+#if WITH_MASSGAMEPLAY_DEBUG
 	LookupTemplateIDMap.Reset();
 	TemplateIDToTemplateMap.Reset();
-#endif // WITH_MASS_DEBUG
+#endif // WITH_MASSGAMEPLAY_DEBUG
 }
 
 const FMassEntityTemplate* UMassEntityTemplateRegistry::FindTemplateFromTemplateID(FMassEntityTemplateID TemplateID) const 
