AggregateTicking and Pipe type renames:

FLWComponentData -> FMassFragment
FComponentTag -> FMassTag
FLWChunkComponent -> FMassChunkFragment
FLWEntity -> FMassEntityHandle
FEntityView -> FMassEntityView
FLWComponentBitSet -> FMassFragmentBitSet
FLWTagBitSet -> FMassTagBitSet
FLWChunkComponentBitSet -> FMassChunkFragmentBitSet
FLWCompositionDescriptor -> FMassCompositionDescriptor
FArchetypeData
FArchetypeChunk
ELWComponentAccess -> EMassFragmentAccess
ELWComponentPresence -> EMassFragmentPresence
FLWComponentRequirement -> FMassFragmentRequirement
FLWComponentQuery -> FMassEntityQuery
FLWComponentSystemExecutionContext -> FMassExecutionContext
FComponentView -> FFragmentView
FChunkComponentView -> FChunkFragmentView
FLWCommandBuffer -> FMassCommandBuffer
FLWComponentSorterOperator -> FMassFragmentSorterOperator
FLWComponentIndicesMapping -> FMassFragmentIndicesMapping
FLWComponentIndicesMappingView -> FMassFragmentIndicesMappingView
FLWRequirementIndicesMapping -> FMassQueryRequirementIndicesMapping
UPipeProcessor -> UMassProcessor
UPipeCompositeProcessor -> UMassCompositeProcessor
FPipeContext -> FMassProcessingContext
FRuntimePipeline -> FMassRuntimePipeline
UPipeProcessingPhaseManager -> UMassProcessingPhaseManager
FPipeProcessingPhase -> FMassProcessingPhase
EPipeProcessingPhase -> EMassProcessingPhase
UPipeSettings -> UMassSettings
UPipeSchematic -> UMassSchematic
FPipeProcessorExecutionOrder -> FMassProcessorExecutionOrder
FPipeProcessorTask -> FMassProcessorTask
FPipeProcessorsTask_GameThread -> FMassProcessorsTask_GameThread

REMOVED:
FChunkIteratorByComponent

diff --git a/Engine/Plugins/AI/MassAI/Source/MassAIBehavior/Private/Evaluators/MassStateTreeSmartObjectEvaluator.cpp b/Engine/Plugins/AI/MassAI/Source/MassAIBehavior/Private/Evaluators/MassStateTreeSmartObjectEvaluator.cpp
--- a/Engine/Plugins/AI/MassAI/Source/MassAIBehavior/Private/Evaluators/MassStateTreeSmartObjectEvaluator.cpp
+++ b/Engine/Plugins/AI/MassAI/Source/MassAIBehavior/Private/Evaluators/MassStateTreeSmartObjectEvaluator.cpp
@@ -86,7 +86,7 @@ void FMassStateTreeSmartObjectEvaluator::Evaluate(FStateTreeExecutionContext& Co
 	if (!SearchRequestID.IsSet())
 	{
 		// Use lanes if possible for faster queries using zone graph annotations
-		const FLWEntity RequestingEntity = MassContext.GetEntity();
+		const FMassEntityHandle RequestingEntity = MassContext.GetEntity();
 		const FMassZoneGraphLaneLocationFragment* LaneLocation = Context.GetExternalItem(LocationHandle).GetPtr<FMassZoneGraphLaneLocationFragment>();
 		bUsingZoneGraphAnnotations = LaneLocation != nullptr;
 		if (bUsingZoneGraphAnnotations)
@@ -118,7 +118,7 @@ void FMassStateTreeSmartObjectEvaluator::Evaluate(FStateTreeExecutionContext& Co
 			// Update bindable flag to indicate to tasks and conditions if some candidates were found
 			bCandidatesFound = SearchRequestResult.NumCandidates > 0;
 
-			const FLWEntity RequestingEntity = MassContext.GetEntity();
+			const FMassEntityHandle RequestingEntity = MassContext.GetEntity();
 			MASSBEHAVIOR_CLOG(bCandidatesFound, Log, TEXT("Found %d smart object candidates for %s"), SearchRequestResult.NumCandidates, *RequestingEntity.DebugGetDescription());
 
 			// When using ZoneGraph annotations we don't need to schedule a new update since we only need the CurrentLaneChanged signal.
diff --git a/Engine/Plugins/AI/MassAI/Source/MassAIBehavior/Private/MassComponentHitSubsystem.cpp b/Engine/Plugins/AI/MassAI/Source/MassAIBehavior/Private/MassComponentHitSubsystem.cpp
--- a/Engine/Plugins/AI/MassAI/Source/MassAIBehavior/Private/MassComponentHitSubsystem.cpp
+++ b/Engine/Plugins/AI/MassAI/Source/MassAIBehavior/Private/MassComponentHitSubsystem.cpp
@@ -61,14 +61,14 @@ void UMassComponentHitSubsystem::Deinitialize()
 	Super::Deinitialize();
 }
 
-void UMassComponentHitSubsystem::RegisterForComponentHit(const FLWEntity Entity, UCapsuleComponent& CapsuleComponent)
+void UMassComponentHitSubsystem::RegisterForComponentHit(const FMassEntityHandle Entity, UCapsuleComponent& CapsuleComponent)
 {
 	EntityToComponentMap.Add(Entity, &CapsuleComponent);
 	ComponentToEntityMap.Add(&CapsuleComponent, Entity);
 	CapsuleComponent.OnComponentHit.AddDynamic(this, &UMassComponentHitSubsystem::OnHitCallback);
 }
 
-void UMassComponentHitSubsystem::UnregisterForComponentHit(const FLWEntity Entity, UCapsuleComponent& CapsuleComponent)
+void UMassComponentHitSubsystem::UnregisterForComponentHit(const FMassEntityHandle Entity, UCapsuleComponent& CapsuleComponent)
 {
 	EntityToComponentMap.Remove(Entity);
 	ComponentToEntityMap.Remove(&CapsuleComponent);
@@ -79,8 +79,8 @@ void UMassComponentHitSubsystem::OnHitCallback(UPrimitiveComponent* HitComp, AAc
 {
 	const UWorld* World = GetWorld();
 	check(World);
-	const FLWEntity& Entity = ComponentToEntityMap.FindChecked(HitComp);
-	FLWEntity* OtherEntity = ComponentToEntityMap.Find(OtherComp);
+	const FMassEntityHandle& Entity = ComponentToEntityMap.FindChecked(HitComp);
+	FMassEntityHandle* OtherEntity = ComponentToEntityMap.Find(OtherComp);
 
 	bool bProcessHit = (OtherEntity != nullptr && OtherEntity->IsSet());
 	if (bProcessHit && UE::MassComponentHit::bOnlyProcessHitsFromPlayers)
@@ -117,7 +117,7 @@ void UMassComponentHitSubsystem::OnHitCallback(UPrimitiveComponent* HitComp, AAc
 	}
 }
 
-const FMassHitResult* UMassComponentHitSubsystem::GetLastHit(const FLWEntity Entity) const
+const FMassHitResult* UMassComponentHitSubsystem::GetLastHit(const FMassEntityHandle Entity) const
 {
 	return HitResults.Find(Entity);
 }
diff --git a/Engine/Plugins/AI/MassAI/Source/MassAIBehavior/Private/MassLookAtProcessors.cpp b/Engine/Plugins/AI/MassAI/Source/MassAIBehavior/Private/MassLookAtProcessors.cpp
--- a/Engine/Plugins/AI/MassAI/Source/MassAIBehavior/Private/MassLookAtProcessors.cpp
+++ b/Engine/Plugins/AI/MassAI/Source/MassAIBehavior/Private/MassLookAtProcessors.cpp
@@ -106,24 +106,24 @@ void UMassLookAtProcessor::Initialize(UObject& Owner)
 
 void UMassLookAtProcessor::ConfigureQueries()
 {
-	EntityQuery_Conditional.AddRequirement<FMassLookAtFragment>(ELWComponentAccess::ReadWrite);
-	EntityQuery_Conditional.AddRequirement<FDataFragment_Transform>(ELWComponentAccess::ReadOnly);
-	EntityQuery_Conditional.AddRequirement<FMassMoveTargetFragment>(ELWComponentAccess::ReadOnly);
-	EntityQuery_Conditional.AddRequirement<FMassZoneGraphLaneLocationFragment>(ELWComponentAccess::ReadOnly, ELWComponentPresence::Optional);
-	EntityQuery_Conditional.AddRequirement<FMassLookAtTrajectoryFragment>(ELWComponentAccess::ReadWrite, ELWComponentPresence::Optional);
-	EntityQuery_Conditional.AddRequirement<FMassZoneGraphShortPathFragment>(ELWComponentAccess::ReadOnly, ELWComponentPresence::Optional);
-	EntityQuery_Conditional.AddTagRequirement<FMassHighLODTag>(ELWComponentPresence::All);
-	EntityQuery_Conditional.AddChunkRequirement<FMassVisualizationChunkFragment>(ELWComponentAccess::ReadOnly);
+	EntityQuery_Conditional.AddRequirement<FMassLookAtFragment>(EMassFragmentAccess::ReadWrite);
+	EntityQuery_Conditional.AddRequirement<FDataFragment_Transform>(EMassFragmentAccess::ReadOnly);
+	EntityQuery_Conditional.AddRequirement<FMassMoveTargetFragment>(EMassFragmentAccess::ReadOnly);
+	EntityQuery_Conditional.AddRequirement<FMassZoneGraphLaneLocationFragment>(EMassFragmentAccess::ReadOnly, EMassFragmentPresence::Optional);
+	EntityQuery_Conditional.AddRequirement<FMassLookAtTrajectoryFragment>(EMassFragmentAccess::ReadWrite, EMassFragmentPresence::Optional);
+	EntityQuery_Conditional.AddRequirement<FMassZoneGraphShortPathFragment>(EMassFragmentAccess::ReadOnly, EMassFragmentPresence::Optional);
+	EntityQuery_Conditional.AddTagRequirement<FMassHighLODTag>(EMassFragmentPresence::All);
+	EntityQuery_Conditional.AddChunkRequirement<FMassVisualizationChunkFragment>(EMassFragmentAccess::ReadOnly);
 	EntityQuery_Conditional.SetChunkFilter(&FMassVisualizationChunkFragment::AreAnyEntitiesVisibleInChunk);
 }
 
-void UMassLookAtProcessor::Execute(UMassEntitySubsystem& EntitySubsystem, FLWComponentSystemExecutionContext& Context)
+void UMassLookAtProcessor::Execute(UMassEntitySubsystem& EntitySubsystem, FMassExecutionContext& Context)
 {
 	QUICK_SCOPE_CYCLE_COUNTER(LookAtProcessor_Run);
 
 	const float CurrentTime = GetWorld()->GetTimeSeconds();
 
-	EntityQuery_Conditional.ForEachEntityChunk(EntitySubsystem, Context, [this, &EntitySubsystem, CurrentTime](FLWComponentSystemExecutionContext& Context)
+	EntityQuery_Conditional.ForEachEntityChunk(EntitySubsystem, Context, [this, &EntitySubsystem, CurrentTime](FMassExecutionContext& Context)
 		{
 			const int32 NumEntities = Context.GetEntitiesNum();
 			const TArrayView<FMassLookAtFragment> LookAtList = Context.GetMutableComponentView<FMassLookAtFragment>();
@@ -142,7 +142,7 @@ void UMassLookAtProcessor::Execute(UMassEntitySubsystem& EntitySubsystem, FLWCom
 				const bool bHasLookAtTrajectory = ZoneGraphLocationList.Num() > 0 && LookAtTrajectoryList.Num() > 0 && ShortPathList.Num() > 0;
 
 				bool bDisplayDebug = false;
-				const FLWEntity Entity = Context.GetEntity(i);
+				const FMassEntityHandle Entity = Context.GetEntity(i);
 	#if WITH_MASSGAMEPLAY_DEBUG
 				FColor EntityColor = FColor::White;
 				bDisplayDebug = UE::MassDebug::IsDebuggingEntity(Entity, &EntityColor);
@@ -235,7 +235,7 @@ void UMassLookAtProcessor::Execute(UMassEntitySubsystem& EntitySubsystem, FLWCom
 
 void UMassLookAtProcessor::FindNewGazeTarget(const UMassEntitySubsystem& EntitySubsystem, const float CurrentTime, const FTransform& Transform, FMassLookAtFragment& LookAt) const
 {
-	const FLWEntity LastTrackedEntity = LookAt.GazeTrackedEntity;
+	const FMassEntityHandle LastTrackedEntity = LookAt.GazeTrackedEntity;
 	
 	LookAt.GazeTrackedEntity.Reset();
 	LookAt.GazeDirection = FVector::ForwardVector;
@@ -273,7 +273,7 @@ void UMassLookAtProcessor::FindNewGazeTarget(const UMassEntitySubsystem& EntityS
 				continue;
 			}
 
-			FEntityView EntityView(EntitySubsystem, NearbyEntity.Entity);
+			FMassEntityView EntityView(EntitySubsystem, NearbyEntity.Entity);
 			if (!EntityView.HasTag<FMassLookAtTargetTag>())
 			{
 				continue;
@@ -400,7 +400,7 @@ bool UMassLookAtProcessor::UpdateGazeTrackedEntity(const UMassEntitySubsystem& E
 }
 
 void UMassLookAtProcessor::BuildTrajectory(const FMassZoneGraphLaneLocationFragment& LaneLocation, const FMassZoneGraphShortPathFragment& ShortPath,
-											const FLWEntity Entity, const bool bDisplayDebug, FMassLookAtTrajectoryFragment& LookAtTrajectory)
+											const FMassEntityHandle Entity, const bool bDisplayDebug, FMassLookAtTrajectoryFragment& LookAtTrajectory)
 {
 	LookAtTrajectory.Reset();
 
diff --git a/Engine/Plugins/AI/MassAI/Source/MassAIBehavior/Private/MassStateTreeExecutionContext.cpp b/Engine/Plugins/AI/MassAI/Source/MassAIBehavior/Private/MassStateTreeExecutionContext.cpp
--- a/Engine/Plugins/AI/MassAI/Source/MassAIBehavior/Private/MassStateTreeExecutionContext.cpp
+++ b/Engine/Plugins/AI/MassAI/Source/MassAIBehavior/Private/MassStateTreeExecutionContext.cpp
@@ -7,7 +7,7 @@
 #include "Engine/World.h"
 
 FMassStateTreeExecutionContext::FMassStateTreeExecutionContext(UMassEntitySubsystem& InEntitySubsystem,
-                                                               FLWComponentSystemExecutionContext& InContext):
+                                                               FMassExecutionContext& InContext):
 	EntitySubsystem(&InEntitySubsystem), EntitySubsystemExecutionContext(&InContext)
 {
 	World = InEntitySubsystem.GetWorld();
diff --git a/Engine/Plugins/AI/MassAI/Source/MassAIBehavior/Private/MassStateTreeProcessors.cpp b/Engine/Plugins/AI/MassAI/Source/MassAIBehavior/Private/MassStateTreeProcessors.cpp
--- a/Engine/Plugins/AI/MassAI/Source/MassAIBehavior/Private/MassStateTreeProcessors.cpp
+++ b/Engine/Plugins/AI/MassAI/Source/MassAIBehavior/Private/MassStateTreeProcessors.cpp
@@ -21,10 +21,10 @@ namespace UE::MassBehavior
 bool SetExternalFragments(FMassStateTreeExecutionContext& Context, const UMassEntitySubsystem& EntitySubsystem)
 {
 	bool bFoundAllFragments = true;
-	const FEntityView EntityView(EntitySubsystem, Context.GetEntity());
+	const FMassEntityView EntityView(EntitySubsystem, Context.GetEntity());
 	for (const FStateTreeExternalItemDesc& ItemDesc : Context.GetExternalItems())
 	{
-		if (ItemDesc.Struct && ItemDesc.Struct->IsChildOf(FLWComponentData::StaticStruct()))
+		if (ItemDesc.Struct && ItemDesc.Struct->IsChildOf(FMassFragment::StaticStruct()))
 		{
 			const UScriptStruct* ScriptStruct = Cast<const UScriptStruct>(ItemDesc.Struct);
 			FStructView Fragment = EntityView.GetComponentDataStruct(ScriptStruct);
@@ -82,7 +82,7 @@ void ProcessChunk(
 	UMassStateTreeSubsystem& MassStateTreeSubsystem,
 	const TFunctionRef<void(FMassStateTreeExecutionContext&, FStateTreeItemView)> ForEachEntityCallback)
 {
-	const FLWComponentSystemExecutionContext& Context = StateTreeContext.GetEntitySubsystemExecutionContext();
+	const FMassExecutionContext& Context = StateTreeContext.GetEntitySubsystemExecutionContext();
 	const TConstArrayView<FMassStateTreeFragment> StateTreeList = Context.GetComponentView<FMassStateTreeFragment>();
 
 	// Assuming that all the entities share same StateTree, because they all should have the same storage fragment.
@@ -108,7 +108,7 @@ void ProcessChunk(
 	const UScriptStruct* StorageScriptStruct = StateTree->GetRuntimeStorageStruct();
 	for (int32 i = 0; i < NumEntities; ++i)
 	{
-		const FLWEntity Entity = Context.GetEntity(i);
+		const FMassEntityHandle Entity = Context.GetEntity(i);
 		StateTreeContext.SetEntity(Entity);
 		StateTreeContext.SetEntityIndex(i);
 
@@ -149,14 +149,14 @@ UMassStateTreeFragmentInitializer::UMassStateTreeFragmentInitializer()
 
 void UMassStateTreeFragmentInitializer::ConfigureQueries()
 {
-	EntityQuery.AddRequirement<FMassStateTreeFragment>(ELWComponentAccess::ReadOnly);
+	EntityQuery.AddRequirement<FMassStateTreeFragment>(EMassFragmentAccess::ReadOnly);
 }
 
-void UMassStateTreeFragmentInitializer::Execute(UMassEntitySubsystem& EntitySubsystem, FLWComponentSystemExecutionContext& Context)
+void UMassStateTreeFragmentInitializer::Execute(UMassEntitySubsystem& EntitySubsystem, FMassExecutionContext& Context)
 {
 	// StateTree processor relies on signals to be ticked but we need an 'initial tick' to set the tree in the proper state.
 	// The initializer provides that by sending a signal to all new entities that use StateTree.
-	TArray<FLWEntity> EntitiesToSignal;
+	TArray<FMassEntityHandle> EntitiesToSignal;
 
 	FMassStateTreeExecutionContext StateTreeContext(EntitySubsystem, Context);
 	UMassStateTreeSubsystem* MassStateTreeSubsystem = UWorld::GetSubsystem<UMassStateTreeSubsystem>(EntitySubsystem.GetWorld());
@@ -164,7 +164,7 @@ void UMassStateTreeFragmentInitializer::Execute(UMassEntitySubsystem& EntitySubs
 	EntityQuery.ForEachEntityChunk(
 		EntitySubsystem,
 		Context,
-		[this, &StateTreeContext, &MassStateTreeSubsystem, &EntitiesToSignal](const FLWComponentSystemExecutionContext& Context)
+		[this, &StateTreeContext, &MassStateTreeSubsystem, &EntitiesToSignal](const FMassExecutionContext& Context)
 		{
 			UE::MassBehavior::ProcessChunk(
 				StateTreeContext,
@@ -199,10 +199,10 @@ UMassStateTreeFragmentDestructor::UMassStateTreeFragmentDestructor()
 
 void UMassStateTreeFragmentDestructor::ConfigureQueries()
 {
-	EntityQuery.AddRequirement<FMassStateTreeFragment>(ELWComponentAccess::ReadOnly);
+	EntityQuery.AddRequirement<FMassStateTreeFragment>(EMassFragmentAccess::ReadOnly);
 }
 
-void UMassStateTreeFragmentDestructor::Execute(UMassEntitySubsystem& EntitySubsystem, FLWComponentSystemExecutionContext& Context)
+void UMassStateTreeFragmentDestructor::Execute(UMassEntitySubsystem& EntitySubsystem, FMassExecutionContext& Context)
 {
 	FMassStateTreeExecutionContext StateTreeContext(EntitySubsystem, Context);
 	UMassStateTreeSubsystem* MassStateTreeSubsystem = UWorld::GetSubsystem<UMassStateTreeSubsystem>(EntitySubsystem.GetWorld());
@@ -210,7 +210,7 @@ void UMassStateTreeFragmentDestructor::Execute(UMassEntitySubsystem& EntitySubsy
 	EntityQuery.ForEachEntityChunk(
 		EntitySubsystem,
 		Context,
-		[this, &StateTreeContext, &MassStateTreeSubsystem](FLWComponentSystemExecutionContext&)
+		[this, &StateTreeContext, &MassStateTreeSubsystem](FMassExecutionContext&)
 		{
 			UE::MassBehavior::ProcessChunk(
 				StateTreeContext,
@@ -270,10 +270,10 @@ void UMassStateTreeProcessor::Initialize(UObject& Owner)
 
 void UMassStateTreeProcessor::ConfigureQueries()
 {
-	EntityQuery.AddRequirement<FMassStateTreeFragment>(ELWComponentAccess::ReadWrite);
+	EntityQuery.AddRequirement<FMassStateTreeFragment>(EMassFragmentAccess::ReadWrite);
 }
 
-void UMassStateTreeProcessor::SignalEntities(UMassEntitySubsystem& EntitySubsystem, FLWComponentSystemExecutionContext& Context, FMassSignalNameLookup& EntitySignals)
+void UMassStateTreeProcessor::SignalEntities(UMassEntitySubsystem& EntitySubsystem, FMassExecutionContext& Context, FMassSignalNameLookup& EntitySignals)
 {
 	if (!MassStateTreeSubsystem)
 	{
@@ -286,12 +286,12 @@ void UMassStateTreeProcessor::SignalEntities(UMassEntitySubsystem& EntitySubsyst
 	FMassStateTreeExecutionContext StateTreeContext(EntitySubsystem, Context);
 	const float TimeInSeconds = EntitySubsystem.GetWorld()->GetTimeSeconds();
 
-	TArray<FLWEntity> EntitiesToSignal;
+	TArray<FMassEntityHandle> EntitiesToSignal;
 
 	EntityQuery.ForEachEntityChunk(
 		EntitySubsystem,
 		Context,
-		[this, &StateTreeContext, TimeDelta, TimeInSeconds, &EntitiesToSignal](FLWComponentSystemExecutionContext& Context)
+		[this, &StateTreeContext, TimeDelta, TimeInSeconds, &EntitiesToSignal](FMassExecutionContext& Context)
 		{
 			TArrayView<FMassStateTreeFragment> StateTreeList = Context.GetMutableComponentView<FMassStateTreeFragment>();
 
diff --git a/Engine/Plugins/AI/MassAI/Source/MassAIBehavior/Private/MassStateTreeSchema.cpp b/Engine/Plugins/AI/MassAI/Source/MassAIBehavior/Private/MassStateTreeSchema.cpp
--- a/Engine/Plugins/AI/MassAI/Source/MassAIBehavior/Private/MassStateTreeSchema.cpp
+++ b/Engine/Plugins/AI/MassAI/Source/MassAIBehavior/Private/MassStateTreeSchema.cpp
@@ -7,8 +7,8 @@
 
 UScriptStruct* UMassStateTreeSchema::GetStorageSuperStruct() const
 {
-	// The runtime state must be derived from FLWComponentData since it's stored in the entity subsystem.
-	return FLWComponentData::StaticStruct();
+	// The runtime state must be derived from FMassFragment since it's stored in the entity subsystem.
+	return FMassFragment::StaticStruct();
 }
 
 bool UMassStateTreeSchema::IsStructAllowed(const UScriptStruct* InScriptStruct) const
@@ -20,5 +20,5 @@ bool UMassStateTreeSchema::IsStructAllowed(const UScriptStruct* InScriptStruct)
 bool UMassStateTreeSchema::IsExternalItemAllowed(const UStruct& InStruct) const
 {
 	// Allow only WorldSubsystems and fragments as external data.
-	return InStruct.IsChildOf(UWorldSubsystem::StaticClass()) || InStruct.IsChildOf(FLWComponentData::StaticStruct());
+	return InStruct.IsChildOf(UWorldSubsystem::StaticClass()) || InStruct.IsChildOf(FMassFragment::StaticStruct());
 }
diff --git a/Engine/Plugins/AI/MassAI/Source/MassAIBehavior/Private/MassStateTreeTrait.cpp b/Engine/Plugins/AI/MassAI/Source/MassAIBehavior/Private/MassStateTreeTrait.cpp
--- a/Engine/Plugins/AI/MassAI/Source/MassAIBehavior/Private/MassStateTreeTrait.cpp
+++ b/Engine/Plugins/AI/MassAI/Source/MassAIBehavior/Private/MassStateTreeTrait.cpp
@@ -65,7 +65,7 @@ void UMassStateTreeTrait::ValidateTemplate(FMassEntityTemplateBuildContext& Buil
 				USubsystem* Subsystem = World.GetSubsystemBase(SubClass);
 				UE_CVLOG(!Subsystem, MassStateTreeSubsystem, LogMassBehavior, Error, TEXT("StateTree %s: Could not find required subsystem %s"), *GetNameSafe(StateTree), *GetNameSafe(ItemDesc.Struct));
 			}
-			else if (ItemDesc.Struct->IsChildOf(FLWComponentData::StaticStruct()))
+			else if (ItemDesc.Struct->IsChildOf(FMassFragment::StaticStruct()))
 			{
 				const bool bContainsFragment = BuildContext.GetMutableFragments().GetFragments().ContainsByPredicate([&ItemDesc](const FInstancedStruct& Item) { return Item.GetScriptStruct() == ItemDesc.Struct; });
 				UE_CVLOG(!bContainsFragment, MassStateTreeSubsystem, LogMassBehavior, Error, TEXT("StateTree %s: Could not find required fragment %s"), *GetNameSafe(StateTree), *GetNameSafe(ItemDesc.Struct));
diff --git a/Engine/Plugins/AI/MassAI/Source/MassAIBehavior/Private/MassZoneGraphAnnotationFragments.cpp b/Engine/Plugins/AI/MassAI/Source/MassAIBehavior/Private/MassZoneGraphAnnotationFragments.cpp
--- a/Engine/Plugins/AI/MassAI/Source/MassAIBehavior/Private/MassZoneGraphAnnotationFragments.cpp
+++ b/Engine/Plugins/AI/MassAI/Source/MassAIBehavior/Private/MassZoneGraphAnnotationFragments.cpp
@@ -20,7 +20,7 @@ static const float OffLODMaxUpdateInterval = 2.10f;
 //  FMassZoneGraphAnnotationVariableTickChunkFragment
 //----------------------------------------------------------------------//
 
-bool FMassZoneGraphAnnotationVariableTickChunkFragment::UpdateChunk(FLWComponentSystemExecutionContext& Context)
+bool FMassZoneGraphAnnotationVariableTickChunkFragment::UpdateChunk(FMassExecutionContext& Context)
 {
 	FMassZoneGraphAnnotationVariableTickChunkFragment& ChunkFrag = Context.GetMutableChunkComponent<FMassZoneGraphAnnotationVariableTickChunkFragment>();
 	ChunkFrag.TimeUntilNextTick -= Context.GetDeltaTimeSeconds();
diff --git a/Engine/Plugins/AI/MassAI/Source/MassAIBehavior/Private/MassZoneGraphAnnotationProcessors.cpp b/Engine/Plugins/AI/MassAI/Source/MassAIBehavior/Private/MassZoneGraphAnnotationProcessors.cpp
--- a/Engine/Plugins/AI/MassAI/Source/MassAIBehavior/Private/MassZoneGraphAnnotationProcessors.cpp
+++ b/Engine/Plugins/AI/MassAI/Source/MassAIBehavior/Private/MassZoneGraphAnnotationProcessors.cpp
@@ -27,13 +27,13 @@ void UMassZoneGraphAnnotationTagsInitializer::Initialize(UObject& Owner)
 
 void UMassZoneGraphAnnotationTagsInitializer::ConfigureQueries()
 {
-	EntityQuery.AddRequirement<FMassZoneGraphAnnotationTagsFragment>(ELWComponentAccess::ReadWrite);
-	EntityQuery.AddRequirement<FMassZoneGraphLaneLocationFragment>(ELWComponentAccess::ReadOnly);
+	EntityQuery.AddRequirement<FMassZoneGraphAnnotationTagsFragment>(EMassFragmentAccess::ReadWrite);
+	EntityQuery.AddRequirement<FMassZoneGraphLaneLocationFragment>(EMassFragmentAccess::ReadOnly);
 }
 
-void UMassZoneGraphAnnotationTagsInitializer::Execute(UMassEntitySubsystem& EntitySubsystem, FLWComponentSystemExecutionContext& Context)
+void UMassZoneGraphAnnotationTagsInitializer::Execute(UMassEntitySubsystem& EntitySubsystem, FMassExecutionContext& Context)
 {
-	EntityQuery.ForEachEntityChunk(EntitySubsystem, Context, [this](FLWComponentSystemExecutionContext& Context)
+	EntityQuery.ForEachEntityChunk(EntitySubsystem, Context, [this](FMassExecutionContext& Context)
 	{
 		const int32 NumEntities = Context.GetEntitiesNum();
 		const TArrayView<FMassZoneGraphAnnotationTagsFragment> AnnotationTagsList = Context.GetMutableComponentView<FMassZoneGraphAnnotationTagsFragment>();
@@ -77,20 +77,20 @@ void UMassZoneGraphAnnotationTagUpdateProcessor::Initialize(UObject& Owner)
 void UMassZoneGraphAnnotationTagUpdateProcessor::ConfigureQueries()
 {
 	Super::ConfigureQueries();
-	EntityQuery.AddRequirement<FMassZoneGraphAnnotationTagsFragment>(ELWComponentAccess::ReadWrite);
-	EntityQuery.AddRequirement<FMassZoneGraphLaneLocationFragment>(ELWComponentAccess::ReadOnly);
-	EntityQuery.AddChunkRequirement<FMassZoneGraphAnnotationVariableTickChunkFragment>(ELWComponentAccess::ReadOnly);
-	EntityQuery.AddChunkRequirement<FMassSimulationVariableTickChunkFragment>(ELWComponentAccess::ReadOnly);
+	EntityQuery.AddRequirement<FMassZoneGraphAnnotationTagsFragment>(EMassFragmentAccess::ReadWrite);
+	EntityQuery.AddRequirement<FMassZoneGraphLaneLocationFragment>(EMassFragmentAccess::ReadOnly);
+	EntityQuery.AddChunkRequirement<FMassZoneGraphAnnotationVariableTickChunkFragment>(EMassFragmentAccess::ReadOnly);
+	EntityQuery.AddChunkRequirement<FMassSimulationVariableTickChunkFragment>(EMassFragmentAccess::ReadOnly);
 }
 
-void UMassZoneGraphAnnotationTagUpdateProcessor::Execute(UMassEntitySubsystem& EntitySubsystem, FLWComponentSystemExecutionContext& Context)
+void UMassZoneGraphAnnotationTagUpdateProcessor::Execute(UMassEntitySubsystem& EntitySubsystem, FMassExecutionContext& Context)
 {
 	TransientEntitiesToSignal.Reset();
 
 	// Calling super will update the signals, and call SignalEntities() below.
 	Super::Execute(EntitySubsystem, Context);
 
-	EntityQuery.ForEachEntityChunk(EntitySubsystem, Context, [this](FLWComponentSystemExecutionContext& Context)
+	EntityQuery.ForEachEntityChunk(EntitySubsystem, Context, [this](FMassExecutionContext& Context)
 	{
 		// Periodically update tags.
 		if (!FMassZoneGraphAnnotationVariableTickChunkFragment::UpdateChunk(Context))
@@ -117,7 +117,7 @@ void UMassZoneGraphAnnotationTagUpdateProcessor::Execute(UMassEntitySubsystem& E
 	}
 }
 
-void UMassZoneGraphAnnotationTagUpdateProcessor::UpdateAnnotationTags(FMassZoneGraphAnnotationTagsFragment& AnnotationTags, const FMassZoneGraphLaneLocationFragment& LaneLocation, const FLWEntity Entity)
+void UMassZoneGraphAnnotationTagUpdateProcessor::UpdateAnnotationTags(FMassZoneGraphAnnotationTagsFragment& AnnotationTags, const FMassZoneGraphLaneLocationFragment& LaneLocation, const FMassEntityHandle Entity)
 {
 	const FZoneGraphTagMask OldTags = AnnotationTags.Tags;
 
@@ -136,14 +136,14 @@ void UMassZoneGraphAnnotationTagUpdateProcessor::UpdateAnnotationTags(FMassZoneG
 	}
 }
 
-void UMassZoneGraphAnnotationTagUpdateProcessor::SignalEntities(UMassEntitySubsystem& EntitySubsystem, FLWComponentSystemExecutionContext& Context, FMassSignalNameLookup& EntitySignals)
+void UMassZoneGraphAnnotationTagUpdateProcessor::SignalEntities(UMassEntitySubsystem& EntitySubsystem, FMassExecutionContext& Context, FMassSignalNameLookup& EntitySignals)
 {
 	if (!ZoneGraphAnnotationSubsystem)
 	{
 		return;
 	}
 
-	EntityQuery.ForEachEntityChunk(EntitySubsystem, Context, [this](FLWComponentSystemExecutionContext& Context)
+	EntityQuery.ForEachEntityChunk(EntitySubsystem, Context, [this](FMassExecutionContext& Context)
 	{
 		const int32 NumEntities = Context.GetEntitiesNum();
 		const TConstArrayView<FMassZoneGraphLaneLocationFragment> LaneLocationList = Context.GetComponentView<FMassZoneGraphLaneLocationFragment>();
diff --git a/Engine/Plugins/AI/MassAI/Source/MassAIBehavior/Public/Evaluators/MassComponentHitEvaluator.h b/Engine/Plugins/AI/MassAI/Source/MassAIBehavior/Public/Evaluators/MassComponentHitEvaluator.h
--- a/Engine/Plugins/AI/MassAI/Source/MassAIBehavior/Public/Evaluators/MassComponentHitEvaluator.h
+++ b/Engine/Plugins/AI/MassAI/Source/MassAIBehavior/Public/Evaluators/MassComponentHitEvaluator.h
@@ -24,5 +24,5 @@ protected:
 	bool bGotHit = false;
 
 	UPROPERTY(VisibleAnywhere, Category = Output)
-	FLWEntity LastHitEntity;
+	FMassEntityHandle LastHitEntity;
 };
diff --git a/Engine/Plugins/AI/MassAI/Source/MassAIBehavior/Public/MassComponentHitSubsystem.h b/Engine/Plugins/AI/MassAI/Source/MassAIBehavior/Public/MassComponentHitSubsystem.h
--- a/Engine/Plugins/AI/MassAI/Source/MassAIBehavior/Public/MassComponentHitSubsystem.h
+++ b/Engine/Plugins/AI/MassAI/Source/MassAIBehavior/Public/MassComponentHitSubsystem.h
@@ -22,7 +22,7 @@ class MASSAIBEHAVIOR_API UMassComponentHitSubsystem : public UTickableWorldSubsy
 	GENERATED_BODY()
 
 public:
-	const FMassHitResult* GetLastHit(const FLWEntity Entity) const;
+	const FMassHitResult* GetLastHit(const FMassEntityHandle Entity) const;
 
 protected:
 	virtual void Initialize(FSubsystemCollectionBase& Collection) override;
@@ -30,8 +30,8 @@ protected:
 	virtual void Tick(float DeltaTime) override;
 	virtual TStatId GetStatId() const override;
 	
-	void RegisterForComponentHit(const FLWEntity Entity, UCapsuleComponent& CapsuleComponent);
-	void UnregisterForComponentHit(FLWEntity Entity, UCapsuleComponent& CapsuleComponent);
+	void RegisterForComponentHit(const FMassEntityHandle Entity, UCapsuleComponent& CapsuleComponent);
+	void UnregisterForComponentHit(FMassEntityHandle Entity, UCapsuleComponent& CapsuleComponent);
 
 	UFUNCTION()
 	void OnHitCallback(UPrimitiveComponent* HitComp, AActor* OtherActor, UPrimitiveComponent* OtherComp, FVector NormalImpulse, const FHitResult& Hit);
@@ -43,11 +43,11 @@ protected:
     UMassAgentSubsystem* AgentSubsystem;
 
 	UPROPERTY()
-	TMap<FLWEntity, FMassHitResult> HitResults;
+	TMap<FMassEntityHandle, FMassHitResult> HitResults;
 
 	UPROPERTY()
-	TMap<UActorComponent*, FLWEntity> ComponentToEntityMap;
+	TMap<UActorComponent*, FMassEntityHandle> ComponentToEntityMap;
 
 	UPROPERTY()
-	TMap<FLWEntity, UActorComponent*> EntityToComponentMap;
+	TMap<FMassEntityHandle, UActorComponent*> EntityToComponentMap;
 };
diff --git a/Engine/Plugins/AI/MassAI/Source/MassAIBehavior/Public/MassComponentHitTypes.h b/Engine/Plugins/AI/MassAI/Source/MassAIBehavior/Public/MassComponentHitTypes.h
--- a/Engine/Plugins/AI/MassAI/Source/MassAIBehavior/Public/MassComponentHitTypes.h
+++ b/Engine/Plugins/AI/MassAI/Source/MassAIBehavior/Public/MassComponentHitTypes.h
@@ -20,14 +20,14 @@ struct FMassHitResult
 
 	FMassHitResult() = default;
 
-	FMassHitResult(const FLWEntity& OtherEntity, const float Time)
+	FMassHitResult(const FMassEntityHandle& OtherEntity, const float Time)
 		: OtherEntity(OtherEntity)
 		, HitTime(Time)
 		, LastFilteredHitTime(Time)
 	{
 	}
 
-	FLWEntity OtherEntity;
+	FMassEntityHandle OtherEntity;
 
 	/** Time when first hit was received. */
 	float HitTime = 0.f;
diff --git a/Engine/Plugins/AI/MassAI/Source/MassAIBehavior/Public/MassLookAtFragments.h b/Engine/Plugins/AI/MassAI/Source/MassAIBehavior/Public/MassLookAtFragments.h
--- a/Engine/Plugins/AI/MassAI/Source/MassAIBehavior/Public/MassLookAtFragments.h
+++ b/Engine/Plugins/AI/MassAI/Source/MassAIBehavior/Public/MassLookAtFragments.h
@@ -36,7 +36,7 @@ enum class EMassLookAtGazeMode : uint8
  * Struct that holds all parameters of the current entity look at 
  */
 USTRUCT()
-struct MASSAIBEHAVIOR_API FMassLookAtFragment : public FLWComponentData
+struct MASSAIBEHAVIOR_API FMassLookAtFragment : public FMassFragment
 {
 	GENERATED_BODY()
 
@@ -71,11 +71,11 @@ struct MASSAIBEHAVIOR_API FMassLookAtFragment : public FLWComponentData
 
 	/** Specific entity that is being tracked as primary look at. */
 	UPROPERTY(Transient)
-	FLWEntity TrackedEntity;
+	FMassEntityHandle TrackedEntity;
 
 	/** Entity that is tracked as part of gazing. */
 	UPROPERTY(Transient)
-	FLWEntity GazeTrackedEntity;
+	FMassEntityHandle GazeTrackedEntity;
 
 	/** Start time of the current gaze. */
 	UPROPERTY(Transient)
@@ -114,7 +114,7 @@ struct MASSAIBEHAVIOR_API FMassLookAtFragment : public FLWComponentData
  * Special tag to mark an entity that could be tracked by the LookAt
  */
 USTRUCT()
-struct MASSAIBEHAVIOR_API FMassLookAtTargetTag : public FComponentTag
+struct MASSAIBEHAVIOR_API FMassLookAtTargetTag : public FMassTag
 {
 	GENERATED_BODY()
 };
@@ -143,7 +143,7 @@ struct MASSAIBEHAVIOR_API FMassLookAtTrajectoryPoint
 };
 
 USTRUCT()
-struct MASSAIBEHAVIOR_API FMassLookAtTrajectoryFragment : public FLWComponentData
+struct MASSAIBEHAVIOR_API FMassLookAtTrajectoryFragment : public FMassFragment
 {
 	GENERATED_BODY()
 
diff --git a/Engine/Plugins/AI/MassAI/Source/MassAIBehavior/Public/MassLookAtProcessors.h b/Engine/Plugins/AI/MassAI/Source/MassAIBehavior/Public/MassLookAtProcessors.h
--- a/Engine/Plugins/AI/MassAI/Source/MassAIBehavior/Public/MassLookAtProcessors.h
+++ b/Engine/Plugins/AI/MassAI/Source/MassAIBehavior/Public/MassLookAtProcessors.h
@@ -19,7 +19,7 @@ struct FMassMoveTargetFragment;
  * Processor to choose and assign LookAt configurations  
  */
 UCLASS()
-class MASSAIBEHAVIOR_API UMassLookAtProcessor : public UPipeProcessor
+class MASSAIBEHAVIOR_API UMassLookAtProcessor : public UMassProcessor
 {
 	GENERATED_BODY()
 
@@ -29,7 +29,7 @@ protected:
 
 	virtual void ConfigureQueries() override;
 	virtual void Initialize(UObject& Owner) override;
-	virtual void Execute(UMassEntitySubsystem& EntitySubsystem, FLWComponentSystemExecutionContext& Context) override;
+	virtual void Execute(UMassEntitySubsystem& EntitySubsystem, FMassExecutionContext& Context) override;
 
 	/** Selects a nearby target if possible or use a random fixed direction */
 	void FindNewGazeTarget(const UMassEntitySubsystem& EntitySubsystem, const float CurrentTime, const FTransform& Transform, FMassLookAtFragment& LookAt) const;
@@ -46,7 +46,7 @@ protected:
 
 	/** Builds look at trajectory along the current path. */
 	void BuildTrajectory(const FMassZoneGraphLaneLocationFragment& LaneLocation, const FMassZoneGraphShortPathFragment& ShortPath,
-							const FLWEntity Entity, const bool bDisplayDebug, FMassLookAtTrajectoryFragment& LookAtTrajectory);
+							const FMassEntityHandle Entity, const bool bDisplayDebug, FMassLookAtTrajectoryFragment& LookAtTrajectory);
 
 	/** Size of the query to find potential targets */
 	UPROPERTY(EditDefaultsOnly, Category = LookAt, config, meta = (UIMin = 0.0, ClampMin = 0.0))
@@ -68,7 +68,7 @@ protected:
 	UPROPERTY(EditDefaultsOnly, Category = LookAt, config, meta = (UIMin = 0.0, ClampMin = 0.0, UIMax = 180.0, ClampMax = 180.0, DisplayName="Angle Threshold (degrees)"))
 	float AngleThresholdInDegrees = 0.f;
 
-	FLWComponentQuery EntityQuery_Conditional;
+	FMassEntityQuery EntityQuery_Conditional;
 
 	UPROPERTY(Transient)
 	UMassMovementSubsystem* MovementSubsystem;
diff --git a/Engine/Plugins/AI/MassAI/Source/MassAIBehavior/Public/MassStateTreeExecutionContext.h b/Engine/Plugins/AI/MassAI/Source/MassAIBehavior/Public/MassStateTreeExecutionContext.h
--- a/Engine/Plugins/AI/MassAI/Source/MassAIBehavior/Public/MassStateTreeExecutionContext.h
+++ b/Engine/Plugins/AI/MassAI/Source/MassAIBehavior/Public/MassStateTreeExecutionContext.h
@@ -6,7 +6,7 @@
 #include "StateTreeExecutionContext.h"
 #include "MassStateTreeExecutionContext.generated.h"
 
-struct FLWComponentSystemExecutionContext;
+struct FMassExecutionContext;
 class UMassEntitySubsystem;
 class UMassSignalSubsystem;
 
@@ -20,13 +20,13 @@ struct MASSAIBEHAVIOR_API FMassStateTreeExecutionContext : public FStateTreeExec
 public:
 	/** Should never be used but has to be public for 'void UScriptStruct::TCppStructOps<FMassStateTreeExecutionContext>::ConstructForTests(void *)' */
 	FMassStateTreeExecutionContext() = default;
-	FMassStateTreeExecutionContext(UMassEntitySubsystem& InEntitySubsystem, FLWComponentSystemExecutionContext& InContext);
+	FMassStateTreeExecutionContext(UMassEntitySubsystem& InEntitySubsystem, FMassExecutionContext& InContext);
 
 	UMassEntitySubsystem& GetEntitySubsystem() const { return *EntitySubsystem; }
-	FLWComponentSystemExecutionContext& GetEntitySubsystemExecutionContext() const { return *EntitySubsystemExecutionContext; }
+	FMassExecutionContext& GetEntitySubsystemExecutionContext() const { return *EntitySubsystemExecutionContext; }
 
-	FLWEntity GetEntity() const { return Entity; }
-	void SetEntity(const FLWEntity& InEntity) { Entity = InEntity; }
+	FMassEntityHandle GetEntity() const { return Entity; }
+	void SetEntity(const FMassEntityHandle& InEntity) { Entity = InEntity; }
 
 	int32 GetEntityIndex() const { return EntityIndex; }
 	void SetEntityIndex(const int32 InIndex) { EntityIndex = InIndex; }
@@ -40,7 +40,7 @@ protected:
 
 	UMassEntitySubsystem* EntitySubsystem = nullptr;
 	UMassSignalSubsystem* SignalSubsystem = nullptr;
-	FLWComponentSystemExecutionContext* EntitySubsystemExecutionContext = nullptr;
-	FLWEntity Entity;
+	FMassExecutionContext* EntitySubsystemExecutionContext = nullptr;
+	FMassEntityHandle Entity;
 	int32 EntityIndex = INDEX_NONE;
 };
diff --git a/Engine/Plugins/AI/MassAI/Source/MassAIBehavior/Public/MassStateTreeFragments.h b/Engine/Plugins/AI/MassAI/Source/MassAIBehavior/Public/MassStateTreeFragments.h
--- a/Engine/Plugins/AI/MassAI/Source/MassAIBehavior/Public/MassStateTreeFragments.h
+++ b/Engine/Plugins/AI/MassAI/Source/MassAIBehavior/Public/MassStateTreeFragments.h
@@ -8,7 +8,7 @@
 
 
 USTRUCT()
-struct MASSAIBEHAVIOR_API FMassStateTreeFragment : public FLWComponentData
+struct MASSAIBEHAVIOR_API FMassStateTreeFragment : public FMassFragment
 {
 	GENERATED_BODY()
 
diff --git a/Engine/Plugins/AI/MassAI/Source/MassAIBehavior/Public/MassStateTreeProcessors.h b/Engine/Plugins/AI/MassAI/Source/MassAIBehavior/Public/MassStateTreeProcessors.h
--- a/Engine/Plugins/AI/MassAI/Source/MassAIBehavior/Public/MassStateTreeProcessors.h
+++ b/Engine/Plugins/AI/MassAI/Source/MassAIBehavior/Public/MassStateTreeProcessors.h
@@ -22,9 +22,9 @@ public:
 
 protected:
 	virtual void ConfigureQueries() override;
-	virtual void Execute(UMassEntitySubsystem& EntitySubsystem, FLWComponentSystemExecutionContext& Context) override;
+	virtual void Execute(UMassEntitySubsystem& EntitySubsystem, FMassExecutionContext& Context) override;
 
-	FLWComponentQuery EntityQuery;
+	FMassEntityQuery EntityQuery;
 };
 
 UCLASS()
@@ -37,9 +37,9 @@ public:
 
 protected:
 	virtual void ConfigureQueries() override;
-	virtual void Execute(UMassEntitySubsystem& EntitySubsystem, FLWComponentSystemExecutionContext& Context) override;
+	virtual void Execute(UMassEntitySubsystem& EntitySubsystem, FMassExecutionContext& Context) override;
 
-	FLWComponentQuery EntityQuery;
+	FMassEntityQuery EntityQuery;
 };
 
 /** 
@@ -56,7 +56,7 @@ public:
 protected:
 	virtual void Initialize(UObject& Owner) override;
 	virtual void ConfigureQueries() override;
-	virtual void SignalEntities(UMassEntitySubsystem& EntitySubsystem, FLWComponentSystemExecutionContext& Context, FMassSignalNameLookup& EntitySignals) override;
+	virtual void SignalEntities(UMassEntitySubsystem& EntitySubsystem, FMassExecutionContext& Context, FMassSignalNameLookup& EntitySignals) override;
 
 	UPROPERTY(Transient)
 	UMassStateTreeSubsystem* MassStateTreeSubsystem;
diff --git a/Engine/Plugins/AI/MassAI/Source/MassAIBehavior/Public/MassZoneGraphAnnotationFragments.h b/Engine/Plugins/AI/MassAI/Source/MassAIBehavior/Public/MassZoneGraphAnnotationFragments.h
--- a/Engine/Plugins/AI/MassAI/Source/MassAIBehavior/Public/MassZoneGraphAnnotationFragments.h
+++ b/Engine/Plugins/AI/MassAI/Source/MassAIBehavior/Public/MassZoneGraphAnnotationFragments.h
@@ -6,10 +6,10 @@
 #include "ZoneGraphTypes.h"
 #include "MassZoneGraphAnnotationFragments.generated.h"
 
-struct FLWComponentSystemExecutionContext;
+struct FMassExecutionContext;
 
 USTRUCT()
-struct MASSAIBEHAVIOR_API FMassZoneGraphAnnotationTagsFragment : public FLWComponentData
+struct MASSAIBEHAVIOR_API FMassZoneGraphAnnotationTagsFragment : public FMassFragment
 {
 	GENERATED_BODY()
 
@@ -19,12 +19,12 @@ struct MASSAIBEHAVIOR_API FMassZoneGraphAnnotationTagsFragment : public FLWCompo
 };
 
 USTRUCT()
-struct FMassZoneGraphAnnotationVariableTickChunkFragment : public FLWChunkComponent
+struct FMassZoneGraphAnnotationVariableTickChunkFragment : public FMassChunkFragment
 {
 	GENERATED_BODY();
 
 	/** Update the ticking frequency of the chunk and return if this chunk should be process this frame */
-	static bool UpdateChunk(FLWComponentSystemExecutionContext& Context);
+	static bool UpdateChunk(FMassExecutionContext& Context);
 
 	float TimeUntilNextTick = 0.0f;
 };
\ No newline at end of file
diff --git a/Engine/Plugins/AI/MassAI/Source/MassAIBehavior/Public/MassZoneGraphAnnotationProcessors.h b/Engine/Plugins/AI/MassAI/Source/MassAIBehavior/Public/MassZoneGraphAnnotationProcessors.h
--- a/Engine/Plugins/AI/MassAI/Source/MassAIBehavior/Public/MassZoneGraphAnnotationProcessors.h
+++ b/Engine/Plugins/AI/MassAI/Source/MassAIBehavior/Public/MassZoneGraphAnnotationProcessors.h
@@ -10,7 +10,7 @@ class UMassSignalSubsystem;
 class UZoneGraphAnnotationSubsystem;
 struct FMassZoneGraphAnnotationTagsFragment;
 struct FMassZoneGraphLaneLocationFragment;
-struct FLWEntity;
+struct FMassEntityHandle;
 
 /** 
  * Processor for initializing ZoneGraph annotation tags.
@@ -26,9 +26,9 @@ public:
 protected:
 	virtual void Initialize(UObject& Owner) override;
 	virtual void ConfigureQueries() override;
-	virtual void Execute(UMassEntitySubsystem& EntitySubsystem, FLWComponentSystemExecutionContext& Context) override;
+	virtual void Execute(UMassEntitySubsystem& EntitySubsystem, FMassExecutionContext& Context) override;
 
-	FLWComponentQuery EntityQuery;
+	FMassEntityQuery EntityQuery;
 
 	UPROPERTY(Transient)
 	UZoneGraphAnnotationSubsystem* ZoneGraphAnnotationSubsystem = nullptr;
@@ -48,15 +48,15 @@ public:
 protected:
 	virtual void Initialize(UObject& Owner) override;
 	virtual void ConfigureQueries() override;
-	virtual void Execute(UMassEntitySubsystem& EntitySubsystem, FLWComponentSystemExecutionContext& Context) override;
+	virtual void Execute(UMassEntitySubsystem& EntitySubsystem, FMassExecutionContext& Context) override;
 
-	void UpdateAnnotationTags(FMassZoneGraphAnnotationTagsFragment& AnnotationTags, const FMassZoneGraphLaneLocationFragment& LaneLocation, FLWEntity Entity);
+	void UpdateAnnotationTags(FMassZoneGraphAnnotationTagsFragment& AnnotationTags, const FMassZoneGraphLaneLocationFragment& LaneLocation, FMassEntityHandle Entity);
 
-	virtual void SignalEntities(UMassEntitySubsystem& EntitySubsystem, FLWComponentSystemExecutionContext& Context, FMassSignalNameLookup& EntitySignals) override;
+	virtual void SignalEntities(UMassEntitySubsystem& EntitySubsystem, FMassExecutionContext& Context, FMassSignalNameLookup& EntitySignals) override;
 
 	UPROPERTY(Transient)
 	UZoneGraphAnnotationSubsystem* ZoneGraphAnnotationSubsystem = nullptr;
 
 	// Frame buffer, it gets reset every frame.
-	TArray<FLWEntity> TransientEntitiesToSignal;
+	TArray<FMassEntityHandle> TransientEntitiesToSignal;
 };
diff --git a/Engine/Plugins/AI/MassAI/Source/MassAIBehavior/Public/Tasks/MassLookAtTask.h b/Engine/Plugins/AI/MassAI/Source/MassAIBehavior/Public/Tasks/MassLookAtTask.h
--- a/Engine/Plugins/AI/MassAI/Source/MassAIBehavior/Public/Tasks/MassLookAtTask.h
+++ b/Engine/Plugins/AI/MassAI/Source/MassAIBehavior/Public/Tasks/MassLookAtTask.h
@@ -36,7 +36,7 @@ protected:
 
 	/** Entity to set as the target for the LookAt behavior. */
 	UPROPERTY(VisibleAnywhere, Category = Params, meta = (Bindable))
-	FLWEntity TargetEntity;
+	FMassEntityHandle TargetEntity;
 
 	/** Random gaze Mode */
 	UPROPERTY(EditAnywhere, Category = Params, meta = (Bindable))
diff --git a/Engine/Plugins/AI/MassAI/Source/MassAIDebug/Private/GameplayDebuggerCategory_Mass.cpp b/Engine/Plugins/AI/MassAI/Source/MassAIDebug/Private/GameplayDebuggerCategory_Mass.cpp
--- a/Engine/Plugins/AI/MassAI/Source/MassAIDebug/Private/GameplayDebuggerCategory_Mass.cpp
+++ b/Engine/Plugins/AI/MassAI/Source/MassAIDebug/Private/GameplayDebuggerCategory_Mass.cpp
@@ -49,7 +49,7 @@ FGameplayDebuggerCategory_Mass::FGameplayDebuggerCategory_Mass()
 	BindKeyPress(EKeys::C.GetFName(), FGameplayDebuggerInputModifier::Shift, this, &FGameplayDebuggerCategory_Mass::OnToggleNearEntityPath, EGameplayDebuggerInputMode::Replicated);
 }
 
-void FGameplayDebuggerCategory_Mass::SetCachedEntity(FLWEntity Entity, UMassDebuggerSubsystem& Debugger)
+void FGameplayDebuggerCategory_Mass::SetCachedEntity(FMassEntityHandle Entity, UMassDebuggerSubsystem& Debugger)
 {
 	CachedEntity = Entity;
 	Debugger.SetSelectedEntity(Entity);
@@ -66,10 +66,10 @@ void FGameplayDebuggerCategory_Mass::PickEntity(const APlayerController& OwnerPC
 	const float MinViewDirDot = 0.707f; // 45 degrees
 	
 	float BestScore = MinViewDirDot;
-	FLWEntity BestEntity;
+	FMassEntityHandle BestEntity;
 	FVector BestLocation = FVector::ZeroVector;
 
-	TConstArrayView<FLWEntity> Entities = Debugger.GetEntities();
+	TConstArrayView<FMassEntityHandle> Entities = Debugger.GetEntities();
 	TConstArrayView<FVector> Locations = Debugger.GetLocations();
 	checkf(Entities.Num() == Locations.Num(), TEXT("Both Entities and Locations lists are expected to be of the same size: %d vs %d"), Entities.Num(), Locations.Num());
 
@@ -135,7 +135,7 @@ void FGameplayDebuggerCategory_Mass::CollectData(APlayerController* OwnerPC, AAc
 
 	auto GetEntityFromActorFunc = [](const AActor& Actor, const UMassAgentComponent** OutMassAgentComponent = nullptr)
 	{
-		FLWEntity LWEntity;
+		FMassEntityHandle LWEntity;
 		UMassAgentComponent* AgentComp = Actor.FindComponentByClass<UMassAgentComponent>();
 		if (AgentComp)
 		{
@@ -296,27 +296,27 @@ void FGameplayDebuggerCategory_Mass::CollectData(APlayerController* OwnerPC, AAc
 		FVector ViewDirection = FVector::ForwardVector;
 		AGameplayDebuggerPlayerManager::GetViewPoint(*OwnerPC, ViewLocation, ViewDirection);
 
-		FLWComponentQuery EntityQuery;
-		EntityQuery.AddRequirement<FMassStateTreeFragment>(ELWComponentAccess::ReadOnly);
-		EntityQuery.AddRequirement<FDataFragment_Transform>(ELWComponentAccess::ReadOnly);
-		EntityQuery.AddRequirement<FDataFragment_AgentRadius>(ELWComponentAccess::ReadOnly);
-		EntityQuery.AddRequirement<FMassSteeringFragment>(ELWComponentAccess::ReadOnly);
-		EntityQuery.AddRequirement<FMassSteeringGhostFragment>(ELWComponentAccess::ReadOnly);
-		EntityQuery.AddRequirement<FMassVelocityFragment>(ELWComponentAccess::ReadOnly);
-		EntityQuery.AddRequirement<FMassMoveTargetFragment>(ELWComponentAccess::ReadOnly);
-		EntityQuery.AddRequirement<FMassLookAtFragment>(ELWComponentAccess::ReadOnly);
-		EntityQuery.AddRequirement<FMassLookAtTrajectoryFragment>(ELWComponentAccess::ReadOnly);
-		EntityQuery.AddRequirement<FMassSimulationLODFragment>(ELWComponentAccess::ReadOnly);
-		EntityQuery.AddRequirement<FMassZoneGraphShortPathFragment>(ELWComponentAccess::ReadOnly);
+		FMassEntityQuery EntityQuery;
+		EntityQuery.AddRequirement<FMassStateTreeFragment>(EMassFragmentAccess::ReadOnly);
+		EntityQuery.AddRequirement<FDataFragment_Transform>(EMassFragmentAccess::ReadOnly);
+		EntityQuery.AddRequirement<FDataFragment_AgentRadius>(EMassFragmentAccess::ReadOnly);
+		EntityQuery.AddRequirement<FMassSteeringFragment>(EMassFragmentAccess::ReadOnly);
+		EntityQuery.AddRequirement<FMassSteeringGhostFragment>(EMassFragmentAccess::ReadOnly);
+		EntityQuery.AddRequirement<FMassVelocityFragment>(EMassFragmentAccess::ReadOnly);
+		EntityQuery.AddRequirement<FMassMoveTargetFragment>(EMassFragmentAccess::ReadOnly);
+		EntityQuery.AddRequirement<FMassLookAtFragment>(EMassFragmentAccess::ReadOnly);
+		EntityQuery.AddRequirement<FMassLookAtTrajectoryFragment>(EMassFragmentAccess::ReadOnly);
+		EntityQuery.AddRequirement<FMassSimulationLODFragment>(EMassFragmentAccess::ReadOnly);
+		EntityQuery.AddRequirement<FMassZoneGraphShortPathFragment>(EMassFragmentAccess::ReadOnly);
 
 		const float CurrentTime = World->GetTimeSeconds();
 		
 		UMassStateTreeSubsystem* MassStateTreeSubsystem = World->GetSubsystem<UMassStateTreeSubsystem>();
 		if (MassStateTreeSubsystem && EntitySystem)
 		{
-			FLWComponentSystemExecutionContext Context(0.0f);
+			FMassExecutionContext Context(0.0f);
 		
-			EntityQuery.ForEachEntityChunk(*EntitySystem, Context, [this, Debugger, MassStateTreeSubsystem, EntitySystem, OwnerPC, ViewLocation, ViewDirection, CurrentTime](FLWComponentSystemExecutionContext& Context)
+			EntityQuery.ForEachEntityChunk(*EntitySystem, Context, [this, Debugger, MassStateTreeSubsystem, EntitySystem, OwnerPC, ViewLocation, ViewDirection, CurrentTime](FMassExecutionContext& Context)
 			{
 				const int32 NumEntities = Context.GetEntitiesNum();
 				const TConstArrayView<FMassStateTreeFragment> StateTreeList = Context.GetComponentView<FMassStateTreeFragment>();
@@ -460,7 +460,7 @@ void FGameplayDebuggerCategory_Mass::CollectData(APlayerController* OwnerPC, AAc
 						FString Status;
 
 						// Entity name
-						FLWEntity Entity = Context.GetEntity(EntityIndex);
+						FMassEntityHandle Entity = Context.GetEntity(EntityIndex);
 						Status += TEXT("{orange}");
 						Status += Entity.DebugGetDescription();
 						Status += TEXT(" {white}LOD ");
@@ -521,11 +521,11 @@ void FGameplayDebuggerCategory_Mass::CollectData(APlayerController* OwnerPC, AAc
 
 		if (bShowNearEntityAvoidance && EntitySystem)
 		{
-			FLWComponentQuery EntityColliderQuery;
-			EntityColliderQuery.AddRequirement<FMassAvoidanceColliderFragment>(ELWComponentAccess::ReadOnly);
-			EntityColliderQuery.AddRequirement<FDataFragment_Transform>(ELWComponentAccess::ReadOnly);
-			FLWComponentSystemExecutionContext Context(0.f);
-			EntityColliderQuery.ForEachEntityChunk(*EntitySystem, Context, [this, Debugger, EntitySystem, OwnerPC, ViewLocation, ViewDirection](FLWComponentSystemExecutionContext& Context)
+			FMassEntityQuery EntityColliderQuery;
+			EntityColliderQuery.AddRequirement<FMassAvoidanceColliderFragment>(EMassFragmentAccess::ReadOnly);
+			EntityColliderQuery.AddRequirement<FDataFragment_Transform>(EMassFragmentAccess::ReadOnly);
+			FMassExecutionContext Context(0.f);
+			EntityColliderQuery.ForEachEntityChunk(*EntitySystem, Context, [this, Debugger, EntitySystem, OwnerPC, ViewLocation, ViewDirection](FMassExecutionContext& Context)
 			{
 				const int32 NumEntities = Context.GetEntitiesNum();
 				const TConstArrayView<FDataFragment_Transform> TransformList = Context.GetComponentView<FDataFragment_Transform>();
diff --git a/Engine/Plugins/AI/MassAI/Source/MassAIDebug/Private/MassDebugStateTreeProcessor.cpp b/Engine/Plugins/AI/MassAI/Source/MassAIDebug/Private/MassDebugStateTreeProcessor.cpp
--- a/Engine/Plugins/AI/MassAI/Source/MassAIDebug/Private/MassDebugStateTreeProcessor.cpp
+++ b/Engine/Plugins/AI/MassAI/Source/MassAIDebug/Private/MassDebugStateTreeProcessor.cpp
@@ -19,11 +19,11 @@ UMassDebugStateTreeProcessor::UMassDebugStateTreeProcessor()
 
 void UMassDebugStateTreeProcessor::ConfigureQueries()
 {
-	EntityQuery.AddRequirement<FMassStateTreeFragment>(ELWComponentAccess::ReadOnly);
-	EntityQuery.AddRequirement<FDataFragment_Transform>(ELWComponentAccess::ReadOnly);
+	EntityQuery.AddRequirement<FMassStateTreeFragment>(EMassFragmentAccess::ReadOnly);
+	EntityQuery.AddRequirement<FDataFragment_Transform>(EMassFragmentAccess::ReadOnly);
 }
 
-void UMassDebugStateTreeProcessor::Execute(UMassEntitySubsystem& EntitySubsystem, FLWComponentSystemExecutionContext& Context)
+void UMassDebugStateTreeProcessor::Execute(UMassEntitySubsystem& EntitySubsystem, FMassExecutionContext& Context)
 {
 #if WITH_MASSGAMEPLAY_DEBUG
 	UWorld* World = GetWorld();
@@ -50,9 +50,9 @@ void UMassDebugStateTreeProcessor::Execute(UMassEntitySubsystem& EntitySubsystem
 	}
 	
 	QUICK_SCOPE_CYCLE_COUNTER(UMassDebugStateTreeProcessor_Run);	
-	EntityQuery.ForEachEntityChunk(EntitySubsystem, Context, [this, Debugger, MassStateTreeSubsystem, &EntitySubsystem](FLWComponentSystemExecutionContext& Context)
+	EntityQuery.ForEachEntityChunk(EntitySubsystem, Context, [this, Debugger, MassStateTreeSubsystem, &EntitySubsystem](FMassExecutionContext& Context)
 		{
-			const FLWEntity SelectedEntity = Debugger->GetSelectedEntity();
+			const FMassEntityHandle SelectedEntity = Debugger->GetSelectedEntity();
 			const int32 NumEntities = Context.GetEntitiesNum();
 			const TConstArrayView<FMassStateTreeFragment> StateTreeList = Context.GetComponentView<FMassStateTreeFragment>();
 			const TConstArrayView<FDataFragment_Transform> TransformList = Context.GetComponentView<FDataFragment_Transform>();
@@ -66,7 +66,7 @@ void UMassDebugStateTreeProcessor::Execute(UMassEntitySubsystem& EntitySubsystem
 		
 			for (int32 i = 0; i < NumEntities; ++i)
 			{
-				FLWEntity Entity = Context.GetEntity(i);
+				FMassEntityHandle Entity = Context.GetEntity(i);
 				if (Entity == SelectedEntity)
 				{
 					FMassStateTreeExecutionContext StateTreeContext(EntitySubsystem, Context);
diff --git a/Engine/Plugins/AI/MassAI/Source/MassAIDebug/Public/GameplayDebuggerCategory_Mass.h b/Engine/Plugins/AI/MassAI/Source/MassAIDebug/Public/GameplayDebuggerCategory_Mass.h
--- a/Engine/Plugins/AI/MassAI/Source/MassAIDebug/Public/GameplayDebuggerCategory_Mass.h
+++ b/Engine/Plugins/AI/MassAI/Source/MassAIDebug/Public/GameplayDebuggerCategory_Mass.h
@@ -24,7 +24,7 @@ protected:
 	virtual void CollectData(APlayerController* OwnerPC, AActor* DebugActor) override;
 	virtual void DrawData(APlayerController* OwnerPC, FGameplayDebuggerCanvasContext& CanvasContext) override;
 	
-	void SetCachedEntity(FLWEntity BestEntity, UMassDebuggerSubsystem& Debugger);
+	void SetCachedEntity(FMassEntityHandle BestEntity, UMassDebuggerSubsystem& Debugger);
 
 	void OnToggleArchetypes() { bShowArchetypes = !bShowArchetypes; }
 	void OnToggleShapes() { bShowShapes = !bShowShapes; }
@@ -39,7 +39,7 @@ protected:
 
 protected:
 	AActor* CachedDebugActor;
-	FLWEntity CachedEntity;
+	FMassEntityHandle CachedEntity;
 	bool bShowArchetypes;
 	bool bShowShapes;
 	bool bShowAgentFragments;
diff --git a/Engine/Plugins/AI/MassAI/Source/MassAIDebug/Public/MassDebugStateTreeProcessor.h b/Engine/Plugins/AI/MassAI/Source/MassAIDebug/Public/MassDebugStateTreeProcessor.h
--- a/Engine/Plugins/AI/MassAI/Source/MassAIDebug/Public/MassDebugStateTreeProcessor.h
+++ b/Engine/Plugins/AI/MassAI/Source/MassAIDebug/Public/MassDebugStateTreeProcessor.h
@@ -7,11 +7,11 @@
 
 class UMassEntitySubsystem;
 struct FMassStateTreeFragment;
-struct FLWComponentQuery;
-struct FLWComponentSystemExecutionContext;
+struct FMassEntityQuery;
+struct FMassExecutionContext;
 
 UCLASS()
-class MASSAIDEBUG_API UMassDebugStateTreeProcessor : public UPipeProcessor
+class MASSAIDEBUG_API UMassDebugStateTreeProcessor : public UMassProcessor
 {
 	GENERATED_BODY()
 
@@ -19,7 +19,7 @@ protected:
 	UMassDebugStateTreeProcessor();
 
 	virtual void ConfigureQueries() override;
-	virtual void Execute(UMassEntitySubsystem& EntitySubsystem, FLWComponentSystemExecutionContext& Context) override;
+	virtual void Execute(UMassEntitySubsystem& EntitySubsystem, FMassExecutionContext& Context) override;
 
-	FLWComponentQuery EntityQuery;
+	FMassEntityQuery EntityQuery;
 };
diff --git a/Engine/Plugins/AI/MassAI/Source/MassAIMovement/Private/AssignRandomNavLocationProcessor.cpp b/Engine/Plugins/AI/MassAI/Source/MassAIMovement/Private/AssignRandomNavLocationProcessor.cpp
--- a/Engine/Plugins/AI/MassAI/Source/MassAIMovement/Private/AssignRandomNavLocationProcessor.cpp
+++ b/Engine/Plugins/AI/MassAI/Source/MassAIMovement/Private/AssignRandomNavLocationProcessor.cpp
@@ -17,18 +17,18 @@ UAssignRandomNavLocationProcessor::UAssignRandomNavLocationProcessor()
 
 void UAssignRandomNavLocationProcessor::ConfigureQueries()
 {
-	EntityQuery.AddRequirement<FDataFragment_Transform>(ELWComponentAccess::ReadWrite);
-	EntityQuery.AddRequirement<FDataFragment_NavLocation>(ELWComponentAccess::ReadWrite);
+	EntityQuery.AddRequirement<FDataFragment_Transform>(EMassFragmentAccess::ReadWrite);
+	EntityQuery.AddRequirement<FDataFragment_NavLocation>(EMassFragmentAccess::ReadWrite);
 }
 
-void UAssignRandomNavLocationProcessor::Execute(UMassEntitySubsystem& EntitySubsystem, FLWComponentSystemExecutionContext& Context)
+void UAssignRandomNavLocationProcessor::Execute(UMassEntitySubsystem& EntitySubsystem, FMassExecutionContext& Context)
 {
 	if (!ensure(NavigationSystem))
 	{
 		return;
 	}
 
-	EntityQuery.ForEachEntityChunk(EntitySubsystem, Context, [this](FLWComponentSystemExecutionContext& Context)
+	EntityQuery.ForEachEntityChunk(EntitySubsystem, Context, [this](FMassExecutionContext& Context)
 		{
 			TArrayView<FDataFragment_Transform> LocationList = Context.GetMutableComponentView<FDataFragment_Transform>();
 			TArrayView<FDataFragment_NavLocation> NavLocationList = Context.GetMutableComponentView<FDataFragment_NavLocation>();
diff --git a/Engine/Plugins/AI/MassAI/Source/MassAIMovement/Private/MassAvoidanceProcessors.cpp b/Engine/Plugins/AI/MassAI/Source/MassAIMovement/Private/MassAvoidanceProcessors.cpp
--- a/Engine/Plugins/AI/MassAI/Source/MassAIMovement/Private/MassAvoidanceProcessors.cpp
+++ b/Engine/Plugins/AI/MassAI/Source/MassAIMovement/Private/MassAvoidanceProcessors.cpp
@@ -212,7 +212,7 @@ namespace UE::MassAvoidance
 	//----------------------------------------------------------------------//
 	struct FDebugContext
 	{
-		FDebugContext(const UObject* InLogOwner, const FLogCategoryBase& InCategory, const UWorld* InWorld, const FLWEntity InEntity)
+		FDebugContext(const UObject* InLogOwner, const FLogCategoryBase& InCategory, const UWorld* InWorld, const FMassEntityHandle InEntity)
 			: LogOwner(InLogOwner)
 			, Category(InCategory)
 			, World(InWorld)
@@ -222,10 +222,10 @@ namespace UE::MassAvoidance
 		const UObject* LogOwner;
 		const FLogCategoryBase& Category;
 		const UWorld* World;
-		const FLWEntity Entity;
+		const FMassEntityHandle Entity;
 	};
 
-	static bool DebugIsSelected(const FLWEntity& Entity)
+	static bool DebugIsSelected(const FMassEntityHandle& Entity)
 	{
 		FColor Color;
 		return UE::MassDebug::IsDebuggingEntity(Entity, &Color);
@@ -397,15 +397,15 @@ UMassAvoidanceProcessor::UMassAvoidanceProcessor()
 
 void UMassAvoidanceProcessor::ConfigureQueries()
 {
-	EntityQuery.AddRequirement<FMassSteeringFragment>(ELWComponentAccess::ReadWrite);
-	EntityQuery.AddRequirement<FMassNavigationEdgesFragment>(ELWComponentAccess::ReadOnly);
-	EntityQuery.AddRequirement<FMassSimulationLODFragment>(ELWComponentAccess::ReadOnly);
-	EntityQuery.AddRequirement<FMassMoveTargetFragment>(ELWComponentAccess::ReadOnly);
-	EntityQuery.AddRequirement<FDataFragment_Transform>(ELWComponentAccess::ReadOnly);
-	EntityQuery.AddRequirement<FMassVelocityFragment>(ELWComponentAccess::ReadOnly);
-	EntityQuery.AddRequirement<FDataFragment_AgentRadius>(ELWComponentAccess::ReadOnly);
-	EntityQuery.AddRequirement<FMassMovementConfigFragment>(ELWComponentAccess::ReadOnly);
-	EntityQuery.AddTagRequirement<FMassHighLODTag>(ELWComponentPresence::All);
+	EntityQuery.AddRequirement<FMassSteeringFragment>(EMassFragmentAccess::ReadWrite);
+	EntityQuery.AddRequirement<FMassNavigationEdgesFragment>(EMassFragmentAccess::ReadOnly);
+	EntityQuery.AddRequirement<FMassSimulationLODFragment>(EMassFragmentAccess::ReadOnly);
+	EntityQuery.AddRequirement<FMassMoveTargetFragment>(EMassFragmentAccess::ReadOnly);
+	EntityQuery.AddRequirement<FDataFragment_Transform>(EMassFragmentAccess::ReadOnly);
+	EntityQuery.AddRequirement<FMassVelocityFragment>(EMassFragmentAccess::ReadOnly);
+	EntityQuery.AddRequirement<FDataFragment_AgentRadius>(EMassFragmentAccess::ReadOnly);
+	EntityQuery.AddRequirement<FMassMovementConfigFragment>(EMassFragmentAccess::ReadOnly);
+	EntityQuery.AddTagRequirement<FMassHighLODTag>(EMassFragmentPresence::All);
 }
 
 void UMassAvoidanceProcessor::Initialize(UObject& Owner)
@@ -416,7 +416,7 @@ void UMassAvoidanceProcessor::Initialize(UObject& Owner)
 	WeakMovementSubsystem = UWorld::GetSubsystem<UMassMovementSubsystem>(Owner.GetWorld());
 }
 
-void UMassAvoidanceProcessor::Execute(UMassEntitySubsystem& EntitySubsystem, FLWComponentSystemExecutionContext& Context)
+void UMassAvoidanceProcessor::Execute(UMassEntitySubsystem& EntitySubsystem, FMassExecutionContext& Context)
 {
 	QUICK_SCOPE_CYCLE_COUNTER(UMassAvoidanceProcessor);
 
@@ -493,7 +493,7 @@ void UMassAvoidanceProcessor::Execute(UMassEntitySubsystem& EntitySubsystem, FLW
 	static const FVector DebugAgentHeightOffset = FVector(0.f, 0.f, 185.f);
 	static const FVector DebugLowCylinderOffset = FVector(0.f, 0.f, 20.f);
 
-	EntityQuery.ForEachEntityChunk(EntitySubsystem, Context, [=, &EntitySubsystem](FLWComponentSystemExecutionContext& Ctx)
+	EntityQuery.ForEachEntityChunk(EntitySubsystem, Context, [=, &EntitySubsystem](FMassExecutionContext& Ctx)
 	{
 		const int32 NumEntities = Ctx.GetEntitiesNum();
 		const float DistanceCutOffSquare = FMath::Square(UE::MassAvoidance::Tweakables::AgentDetectionDistance);
@@ -557,7 +557,7 @@ void UMassAvoidanceProcessor::Execute(UMassEntitySubsystem& EntitySubsystem, FLW
 				continue;
 			}
 
-			FLWEntity Entity = Ctx.GetEntity(EntityIndex);
+			FMassEntityHandle Entity = Ctx.GetEntity(EntityIndex);
 			
 			const FMassMovementConfigFragment& MovementConfig = MovementConfigList[EntityIndex];
 			if (MovementConfig.ConfigHandle != CurrentConfigHandle)
@@ -881,7 +881,7 @@ void UMassAvoidanceProcessor::Execute(UMassEntitySubsystem& EntitySubsystem, FLW
 				}
 
 				FSortingAgent& OtherAgent = ClosestAgents[Index];
-				FEntityView OtherEntityView(EntitySubsystem, OtherAgent.ObstacleItem.Entity);
+				FMassEntityView OtherEntityView(EntitySubsystem, OtherAgent.ObstacleItem.Entity);
 
 				const FMassVelocityFragment* OtherVelocityFragment = OtherEntityView.GetComponentDataPtr<FMassVelocityFragment>();
 				const FVector OtherVelocity = OtherVelocityFragment != nullptr ? OtherVelocityFragment->Value : FVector::ZeroVector; // Get velocity from FAvoidanceComponent
@@ -1148,12 +1148,12 @@ UMassStandingAvoidanceProcessor::UMassStandingAvoidanceProcessor()
 
 void UMassStandingAvoidanceProcessor::ConfigureQueries()
 {
-	EntityQuery.AddRequirement<FMassSteeringGhostFragment>(ELWComponentAccess::ReadWrite);
-	EntityQuery.AddRequirement<FMassNavigationEdgesFragment>(ELWComponentAccess::ReadOnly);
-	EntityQuery.AddRequirement<FMassMoveTargetFragment>(ELWComponentAccess::ReadOnly);
-	EntityQuery.AddRequirement<FDataFragment_Transform>(ELWComponentAccess::ReadOnly);
-	EntityQuery.AddRequirement<FDataFragment_AgentRadius>(ELWComponentAccess::ReadOnly);
-	EntityQuery.AddTagRequirement<FMassHighLODTag>(ELWComponentPresence::All);
+	EntityQuery.AddRequirement<FMassSteeringGhostFragment>(EMassFragmentAccess::ReadWrite);
+	EntityQuery.AddRequirement<FMassNavigationEdgesFragment>(EMassFragmentAccess::ReadOnly);
+	EntityQuery.AddRequirement<FMassMoveTargetFragment>(EMassFragmentAccess::ReadOnly);
+	EntityQuery.AddRequirement<FDataFragment_Transform>(EMassFragmentAccess::ReadOnly);
+	EntityQuery.AddRequirement<FDataFragment_AgentRadius>(EMassFragmentAccess::ReadOnly);
+	EntityQuery.AddTagRequirement<FMassHighLODTag>(EMassFragmentPresence::All);
 }
 
 void UMassStandingAvoidanceProcessor::Initialize(UObject& Owner)
@@ -1164,7 +1164,7 @@ void UMassStandingAvoidanceProcessor::Initialize(UObject& Owner)
 	WeakMovementSubsystem = UWorld::GetSubsystem<UMassMovementSubsystem>(Owner.GetWorld());
 }
 
-void UMassStandingAvoidanceProcessor::Execute(UMassEntitySubsystem& EntitySubsystem, FLWComponentSystemExecutionContext& Context)
+void UMassStandingAvoidanceProcessor::Execute(UMassEntitySubsystem& EntitySubsystem, FMassExecutionContext& Context)
 {
 	QUICK_SCOPE_CYCLE_COUNTER(UMassStandingAvoidanceProcessor);
 
@@ -1176,7 +1176,7 @@ void UMassStandingAvoidanceProcessor::Execute(UMassEntitySubsystem& EntitySubsys
 	}
 
 	// Avoidance while standing
-	EntityQuery.ForEachEntityChunk(EntitySubsystem, Context, [&EntitySubsystem, MovementSubsystem, World](FLWComponentSystemExecutionContext& Context)
+	EntityQuery.ForEachEntityChunk(EntitySubsystem, Context, [&EntitySubsystem, MovementSubsystem, World](FMassExecutionContext& Context)
 	{
 		const int32 NumEntities = Context.GetEntitiesNum();
 		const float DeltaTime = Context.GetDeltaTimeSeconds();
@@ -1193,9 +1193,9 @@ void UMassStandingAvoidanceProcessor::Execute(UMassEntitySubsystem& EntitySubsys
 		struct FSortingAgent
 		{
 			FSortingAgent() = default;
-			FSortingAgent(const FLWEntity InEntity, const FVector InLocation, const FVector InForward, const float InDistSq) : Entity(InEntity), Location(InLocation), Forward(InForward), DistSq(InDistSq) {}
+			FSortingAgent(const FMassEntityHandle InEntity, const FVector InLocation, const FVector InForward, const float InDistSq) : Entity(InEntity), Location(InLocation), Forward(InForward), DistSq(InDistSq) {}
 			
-			FLWEntity Entity;
+			FMassEntityHandle Entity;
 			FVector Location = FVector::ZeroVector;
 			FVector Forward = FVector::ForwardVector;
 			float DistSq = 0.0f;
@@ -1221,7 +1221,7 @@ void UMassStandingAvoidanceProcessor::Execute(UMassEntitySubsystem& EntitySubsys
 			const FDataFragment_Transform& Location = LocationList[EntityIndex];
 			const FDataFragment_AgentRadius& Radius = RadiusList[EntityIndex];
 
-			FLWEntity Entity = Context.GetEntity(EntityIndex);
+			FMassEntityHandle Entity = Context.GetEntity(EntityIndex);
 			const FVector AgentLocation = Location.GetTransform().GetTranslation();
 			const float AgentRadius = Radius.Radius;
 
@@ -1288,7 +1288,7 @@ void UMassStandingAvoidanceProcessor::Execute(UMassEntitySubsystem& EntitySubsys
 			for (int32 Index = 0; Index < NumCloseAgents; Index++)
 			{
 				FSortingAgent& OtherAgent = ClosestAgents[Index];
-				FEntityView OtherEntityView(EntitySubsystem, OtherAgent.Entity);
+				FMassEntityView OtherEntityView(EntitySubsystem, OtherAgent.Entity);
 
 				const FMassVelocityFragment* OtherVelocityFragment = OtherEntityView.GetComponentDataPtr<FMassVelocityFragment>();
 				const float OtherRadius = OtherEntityView.GetComponentData<FDataFragment_AgentRadius>().Radius;
@@ -1400,20 +1400,20 @@ UMassAvoidanceObstacleProcessor::UMassAvoidanceObstacleProcessor()
 
 void UMassAvoidanceObstacleProcessor::ConfigureQueries()
 {
-	AddToGridEntityQuery.AddRequirement<FDataFragment_Transform>(ELWComponentAccess::ReadOnly);
-	AddToGridEntityQuery.AddRequirement<FDataFragment_AgentRadius>(ELWComponentAccess::ReadOnly);
-	AddToGridEntityQuery.AddRequirement<FMassAvoidanceObstacleGridCellLocationFragment>(ELWComponentAccess::ReadWrite);
+	AddToGridEntityQuery.AddRequirement<FDataFragment_Transform>(EMassFragmentAccess::ReadOnly);
+	AddToGridEntityQuery.AddRequirement<FDataFragment_AgentRadius>(EMassFragmentAccess::ReadOnly);
+	AddToGridEntityQuery.AddRequirement<FMassAvoidanceObstacleGridCellLocationFragment>(EMassFragmentAccess::ReadWrite);
 	UpdateGridEntityQuery = AddToGridEntityQuery;
 	RemoveFromGridEntityQuery = AddToGridEntityQuery;
 
-	AddToGridEntityQuery.AddTagRequirement<FMassOffLODTag>(ELWComponentPresence::None);
-	AddToGridEntityQuery.AddTagRequirement<FMassInAvoidanceObstacleGridTag>(ELWComponentPresence::None);
+	AddToGridEntityQuery.AddTagRequirement<FMassOffLODTag>(EMassFragmentPresence::None);
+	AddToGridEntityQuery.AddTagRequirement<FMassInAvoidanceObstacleGridTag>(EMassFragmentPresence::None);
 
-	UpdateGridEntityQuery.AddTagRequirement<FMassOffLODTag>(ELWComponentPresence::None);
-	UpdateGridEntityQuery.AddTagRequirement<FMassInAvoidanceObstacleGridTag>(ELWComponentPresence::All);
+	UpdateGridEntityQuery.AddTagRequirement<FMassOffLODTag>(EMassFragmentPresence::None);
+	UpdateGridEntityQuery.AddTagRequirement<FMassInAvoidanceObstacleGridTag>(EMassFragmentPresence::All);
 
-	RemoveFromGridEntityQuery.AddTagRequirement<FMassOffLODTag>(ELWComponentPresence::All);
-	RemoveFromGridEntityQuery.AddTagRequirement<FMassInAvoidanceObstacleGridTag>(ELWComponentPresence::All);
+	RemoveFromGridEntityQuery.AddTagRequirement<FMassOffLODTag>(EMassFragmentPresence::All);
+	RemoveFromGridEntityQuery.AddTagRequirement<FMassInAvoidanceObstacleGridTag>(EMassFragmentPresence::All);
 }
 
 void UMassAvoidanceObstacleProcessor::Initialize(UObject& Owner)
@@ -1423,7 +1423,7 @@ void UMassAvoidanceObstacleProcessor::Initialize(UObject& Owner)
 	WeakMovementSubsystem = UWorld::GetSubsystem<UMassMovementSubsystem>(Owner.GetWorld());
 }
 
-void UMassAvoidanceObstacleProcessor::Execute(UMassEntitySubsystem& EntitySubsystem, FLWComponentSystemExecutionContext& Context)
+void UMassAvoidanceObstacleProcessor::Execute(UMassEntitySubsystem& EntitySubsystem, FMassExecutionContext& Context)
 {
 	using namespace UE::MassAvoidance;
 
@@ -1434,7 +1434,7 @@ void UMassAvoidanceObstacleProcessor::Execute(UMassEntitySubsystem& EntitySubsys
 	}
 
 	// can't be ParallelFor due to MovementSubsystem->GetGridMutable().Move not being thread-safe
-	AddToGridEntityQuery.ForEachEntityChunk(EntitySubsystem, Context, [this, MovementSubsystem, &EntitySubsystem](FLWComponentSystemExecutionContext& Context)
+	AddToGridEntityQuery.ForEachEntityChunk(EntitySubsystem, Context, [this, MovementSubsystem, &EntitySubsystem](FMassExecutionContext& Context)
 		{
 			const int32 NumEntities = Context.GetEntitiesNum();
 
@@ -1450,7 +1450,7 @@ void UMassAvoidanceObstacleProcessor::Execute(UMassEntitySubsystem& EntitySubsys
 
 				FMassAvoidanceObstacleItem ObstacleItem;
 				ObstacleItem.Entity = Context.GetEntity(EntityIndex);
-				FEntityView EntityView(EntitySubsystem, ObstacleItem.Entity);
+				FMassEntityView EntityView(EntitySubsystem, ObstacleItem.Entity);
 				const FMassAvoidanceColliderFragment* Collider = EntityView.GetComponentDataPtr<FMassAvoidanceColliderFragment>();
 				if (Collider)
 				{
@@ -1464,7 +1464,7 @@ void UMassAvoidanceObstacleProcessor::Execute(UMassEntitySubsystem& EntitySubsys
 			}
 		});
 
-	UpdateGridEntityQuery.ForEachEntityChunk(EntitySubsystem, Context, [this, MovementSubsystem, &EntitySubsystem](FLWComponentSystemExecutionContext& Context)
+	UpdateGridEntityQuery.ForEachEntityChunk(EntitySubsystem, Context, [this, MovementSubsystem, &EntitySubsystem](FMassExecutionContext& Context)
 		{
 			const int32 NumEntities = Context.GetEntitiesNum();
 
@@ -1479,7 +1479,7 @@ void UMassAvoidanceObstacleProcessor::Execute(UMassEntitySubsystem& EntitySubsys
 				const float Radius = RadiiList[EntityIndex].Radius;
 				FMassAvoidanceObstacleItem ObstacleItem;
 				ObstacleItem.Entity = Context.GetEntity(EntityIndex);
-				FEntityView EntityView(EntitySubsystem, ObstacleItem.Entity);
+				FMassEntityView EntityView(EntitySubsystem, ObstacleItem.Entity);
 				const FMassAvoidanceColliderFragment* Collider = EntityView.GetComponentDataPtr<FMassAvoidanceColliderFragment>();
 				if (Collider)
 				{
@@ -1500,7 +1500,7 @@ void UMassAvoidanceObstacleProcessor::Execute(UMassEntitySubsystem& EntitySubsys
 			}
 		});
 
-	RemoveFromGridEntityQuery.ForEachEntityChunk(EntitySubsystem, Context, [this, MovementSubsystem](FLWComponentSystemExecutionContext& Context)
+	RemoveFromGridEntityQuery.ForEachEntityChunk(EntitySubsystem, Context, [this, MovementSubsystem](FMassExecutionContext& Context)
 	{
 		const int32 NumEntities = Context.GetEntitiesNum();
 
@@ -1530,10 +1530,10 @@ UMassNavigationBoundaryProcessor::UMassNavigationBoundaryProcessor()
 
 void UMassNavigationBoundaryProcessor::ConfigureQueries()
 {
-	EntityQuery.AddRequirement<FDataFragment_Transform>(ELWComponentAccess::ReadOnly);
-	EntityQuery.AddRequirement<FDataFragment_NavLocation>(ELWComponentAccess::ReadOnly);
-	EntityQuery.AddRequirement<FMassEdgeDetectionParamsFragment>(ELWComponentAccess::ReadOnly);
-	EntityQuery.AddRequirement<FMassNavigationEdgesFragment>(ELWComponentAccess::ReadWrite);
+	EntityQuery.AddRequirement<FDataFragment_Transform>(EMassFragmentAccess::ReadOnly);
+	EntityQuery.AddRequirement<FDataFragment_NavLocation>(EMassFragmentAccess::ReadOnly);
+	EntityQuery.AddRequirement<FMassEdgeDetectionParamsFragment>(EMassFragmentAccess::ReadOnly);
+	EntityQuery.AddRequirement<FMassNavigationEdgesFragment>(EMassFragmentAccess::ReadWrite);
 }
 
 void UMassNavigationBoundaryProcessor::Initialize(UObject& Owner)
@@ -1552,7 +1552,7 @@ void UMassNavigationBoundaryProcessor::Initialize(UObject& Owner)
 	WeakNavData = NavData;
 }
 
-void UMassNavigationBoundaryProcessor::Execute(UMassEntitySubsystem& EntitySubsystem, FLWComponentSystemExecutionContext& Context)
+void UMassNavigationBoundaryProcessor::Execute(UMassEntitySubsystem& EntitySubsystem, FMassExecutionContext& Context)
 {
 	using namespace UE::MassAvoidance;
 
@@ -1564,7 +1564,7 @@ void UMassNavigationBoundaryProcessor::Execute(UMassEntitySubsystem& EntitySubsy
 		return;
 	}
 
-	EntityQuery.ForEachEntityChunk(EntitySubsystem, Context, [this, NavData](FLWComponentSystemExecutionContext& Context)
+	EntityQuery.ForEachEntityChunk(EntitySubsystem, Context, [this, NavData](FMassExecutionContext& Context)
 		{
 			const int32 NumEntities = Context.GetEntitiesNum();
 			TConstArrayView<FDataFragment_Transform> LocationList = Context.GetComponentView<FDataFragment_Transform>();
@@ -1610,13 +1610,13 @@ UMassLaneBoundaryProcessor::UMassLaneBoundaryProcessor()
 
 void UMassLaneBoundaryProcessor::ConfigureQueries()
 {
-	EntityQuery.AddRequirement<FDataFragment_Transform>(ELWComponentAccess::ReadOnly);							// need agent position to get closest point on lane
-	EntityQuery.AddRequirement<FMassSimulationLODFragment>(ELWComponentAccess::ReadOnly);						// LOD level
-	EntityQuery.AddRequirement<FMassNavigationEdgesFragment>(ELWComponentAccess::ReadWrite);					// output edges
-	EntityQuery.AddRequirement<FMassLastUpdatePositionFragment>(ELWComponentAccess::ReadWrite);					// to keep position when boundaries where last updated
-	EntityQuery.AddRequirement<FMassZoneGraphLaneLocationFragment>(ELWComponentAccess::ReadOnly);				// current lane location
-	EntityQuery.AddRequirement<FMassAvoidanceBoundaryLastLaneHandleFragment>(ELWComponentAccess::ReadWrite);	// keep track of the last used lane
-	EntityQuery.AddTagRequirement<FMassHighLODTag>(ELWComponentPresence::All);
+	EntityQuery.AddRequirement<FDataFragment_Transform>(EMassFragmentAccess::ReadOnly);							// need agent position to get closest point on lane
+	EntityQuery.AddRequirement<FMassSimulationLODFragment>(EMassFragmentAccess::ReadOnly);						// LOD level
+	EntityQuery.AddRequirement<FMassNavigationEdgesFragment>(EMassFragmentAccess::ReadWrite);					// output edges
+	EntityQuery.AddRequirement<FMassLastUpdatePositionFragment>(EMassFragmentAccess::ReadWrite);					// to keep position when boundaries where last updated
+	EntityQuery.AddRequirement<FMassZoneGraphLaneLocationFragment>(EMassFragmentAccess::ReadOnly);				// current lane location
+	EntityQuery.AddRequirement<FMassAvoidanceBoundaryLastLaneHandleFragment>(EMassFragmentAccess::ReadWrite);	// keep track of the last used lane
+	EntityQuery.AddTagRequirement<FMassHighLODTag>(EMassFragmentPresence::All);
 }
 
 void UMassLaneBoundaryProcessor::Initialize(UObject& Owner)
@@ -1627,7 +1627,7 @@ void UMassLaneBoundaryProcessor::Initialize(UObject& Owner)
 	WeakZoneGraph = UWorld::GetSubsystem<UZoneGraphSubsystem>(Owner.GetWorld());
 }
 
-void UMassLaneBoundaryProcessor::Execute(UMassEntitySubsystem& EntitySubsystem, FLWComponentSystemExecutionContext& Context)
+void UMassLaneBoundaryProcessor::Execute(UMassEntitySubsystem& EntitySubsystem, FMassExecutionContext& Context)
 {
 	using namespace UE::MassAvoidance;
 
@@ -1646,7 +1646,7 @@ void UMassLaneBoundaryProcessor::Execute(UMassEntitySubsystem& EntitySubsystem,
 	static const FColor PaleTurquoise = FColor(175, 238, 238);
 	static const FColor LaneColor = PaleTurquoise;
 
-	EntityQuery.ForEachEntityChunk(EntitySubsystem, Context, [this, World, ZoneGraphSubsystem](FLWComponentSystemExecutionContext& Context)
+	EntityQuery.ForEachEntityChunk(EntitySubsystem, Context, [this, World, ZoneGraphSubsystem](FMassExecutionContext& Context)
 		{
 			const int32 NumEntities = Context.GetEntitiesNum();
 
@@ -1685,7 +1685,7 @@ void UMassLaneBoundaryProcessor::Execute(UMassEntitySubsystem& EntitySubsystem,
 				EdgesFragment.AvoidanceEdges.Reset();
 
 #if WITH_MASSGAMEPLAY_DEBUG && UNSAFE_FOR_MT
-				const FLWEntity Entity = Context.GetEntity(EntityIndex);
+				const FMassEntityHandle Entity = Context.GetEntity(EntityIndex);
 				const FDebugContext BaseDebugContext(this, LogAvoidance, World, Entity);
 				const FDebugContext ObstacleDebugContext(this, LogAvoidanceObstacles, World, Entity);
 
@@ -1888,13 +1888,13 @@ UMassLaneCacheBoundaryProcessor::UMassLaneCacheBoundaryProcessor()
 
 void UMassLaneCacheBoundaryProcessor::ConfigureQueries()
 {
-	EntityQuery.AddRequirement<FMassZoneGraphCachedLaneFragment>(ELWComponentAccess::ReadOnly);
-	EntityQuery.AddRequirement<FMassMoveTargetFragment>(ELWComponentAccess::ReadOnly);
-	EntityQuery.AddRequirement<FMassZoneGraphLaneLocationFragment>(ELWComponentAccess::ReadOnly);
-	EntityQuery.AddRequirement<FMassSimulationLODFragment>(ELWComponentAccess::ReadOnly);
-	EntityQuery.AddRequirement<FMassLaneCacheBoundaryFragment>(ELWComponentAccess::ReadWrite);
-	EntityQuery.AddRequirement<FMassNavigationEdgesFragment>(ELWComponentAccess::ReadWrite);	// output edges
-	EntityQuery.AddTagRequirement<FMassOffLODTag>(ELWComponentPresence::None);
+	EntityQuery.AddRequirement<FMassZoneGraphCachedLaneFragment>(EMassFragmentAccess::ReadOnly);
+	EntityQuery.AddRequirement<FMassMoveTargetFragment>(EMassFragmentAccess::ReadOnly);
+	EntityQuery.AddRequirement<FMassZoneGraphLaneLocationFragment>(EMassFragmentAccess::ReadOnly);
+	EntityQuery.AddRequirement<FMassSimulationLODFragment>(EMassFragmentAccess::ReadOnly);
+	EntityQuery.AddRequirement<FMassLaneCacheBoundaryFragment>(EMassFragmentAccess::ReadWrite);
+	EntityQuery.AddRequirement<FMassNavigationEdgesFragment>(EMassFragmentAccess::ReadWrite);	// output edges
+	EntityQuery.AddTagRequirement<FMassOffLODTag>(EMassFragmentPresence::None);
 }
 
 void UMassLaneCacheBoundaryProcessor::Initialize(UObject& Owner)
@@ -1904,7 +1904,7 @@ void UMassLaneCacheBoundaryProcessor::Initialize(UObject& Owner)
 	WeakWorld = Owner.GetWorld();
 }
 
-void UMassLaneCacheBoundaryProcessor::Execute(UMassEntitySubsystem& EntitySubsystem, FLWComponentSystemExecutionContext& Context)
+void UMassLaneCacheBoundaryProcessor::Execute(UMassEntitySubsystem& EntitySubsystem, FMassExecutionContext& Context)
 {
 	using namespace UE::MassAvoidance;
 
@@ -1916,7 +1916,7 @@ void UMassLaneCacheBoundaryProcessor::Execute(UMassEntitySubsystem& EntitySubsys
 		return;
 	}
 
-	EntityQuery.ForEachEntityChunk(EntitySubsystem, Context, [this, World](FLWComponentSystemExecutionContext& Context)
+	EntityQuery.ForEachEntityChunk(EntitySubsystem, Context, [this, World](FMassExecutionContext& Context)
 	{
 		const int32 NumEntities = Context.GetEntitiesNum();
 
@@ -1936,7 +1936,7 @@ void UMassLaneCacheBoundaryProcessor::Execute(UMassEntitySubsystem& EntitySubsys
 			const FMassMoveTargetFragment& MovementTarget = MovementTargetList[EntityIndex];
 			FMassNavigationEdgesFragment& Edges = EdgesList[EntityIndex];
 			FMassLaneCacheBoundaryFragment& LaneCacheBoundary = LaneCacheBoundaryList[EntityIndex];
-			const FLWEntity Entity = Context.GetEntity(EntityIndex);
+			const FMassEntityHandle Entity = Context.GetEntity(EntityIndex);
 
 			// First check if we moved enough for an update
 			const float DeltaDistSquared = FVector::DistSquared(MovementTarget.Center, LaneCacheBoundary.LastUpdatePosition);
diff --git a/Engine/Plugins/AI/MassAI/Source/MassAIMovement/Private/MassZoneGraphMovementProcessors.cpp b/Engine/Plugins/AI/MassAI/Source/MassAIMovement/Private/MassZoneGraphMovementProcessors.cpp
--- a/Engine/Plugins/AI/MassAI/Source/MassAIMovement/Private/MassZoneGraphMovementProcessors.cpp
+++ b/Engine/Plugins/AI/MassAI/Source/MassAIMovement/Private/MassZoneGraphMovementProcessors.cpp
@@ -77,10 +77,10 @@ UMassZoneGraphLocationInitializer::UMassZoneGraphLocationInitializer()
 
 void UMassZoneGraphLocationInitializer::ConfigureQueries()
 {
-	EntityQuery.AddRequirement<FMassZoneGraphLaneLocationFragment>(ELWComponentAccess::ReadWrite);
-	EntityQuery.AddRequirement<FDataFragment_Transform>(ELWComponentAccess::ReadOnly);
-	EntityQuery.AddRequirement<FMassMovementConfigFragment>(ELWComponentAccess::ReadOnly);
-	EntityQuery.AddRequirement<FMassMoveTargetFragment>(ELWComponentAccess::ReadWrite); // Make optional?
+	EntityQuery.AddRequirement<FMassZoneGraphLaneLocationFragment>(EMassFragmentAccess::ReadWrite);
+	EntityQuery.AddRequirement<FDataFragment_Transform>(EMassFragmentAccess::ReadOnly);
+	EntityQuery.AddRequirement<FMassMovementConfigFragment>(EMassFragmentAccess::ReadOnly);
+	EntityQuery.AddRequirement<FMassMoveTargetFragment>(EMassFragmentAccess::ReadWrite); // Make optional?
 }
 
 void UMassZoneGraphLocationInitializer::Initialize(UObject& Owner)
@@ -90,7 +90,7 @@ void UMassZoneGraphLocationInitializer::Initialize(UObject& Owner)
 	SignalSubsystem = UWorld::GetSubsystem<UMassSignalSubsystem>(Owner.GetWorld());
 }
 
-void UMassZoneGraphLocationInitializer::Execute(UMassEntitySubsystem& EntitySubsystem, FLWComponentSystemExecutionContext& Context)
+void UMassZoneGraphLocationInitializer::Execute(UMassEntitySubsystem& EntitySubsystem, FMassExecutionContext& Context)
 {
 	const UMassMovementSettings* Settings = GetDefault<UMassMovementSettings>();
 	if (!ZoneGraphSubsystem || !Settings || !SignalSubsystem)
@@ -98,7 +98,7 @@ void UMassZoneGraphLocationInitializer::Execute(UMassEntitySubsystem& EntitySubs
 		return;
 	}
 
-	EntityQuery.ForEachEntityChunk(EntitySubsystem, Context, [this, Settings](FLWComponentSystemExecutionContext& Context)
+	EntityQuery.ForEachEntityChunk(EntitySubsystem, Context, [this, Settings](FMassExecutionContext& Context)
 	{
 		const int32 NumEntities = Context.GetEntitiesNum();
 		const TArrayView<FMassZoneGraphLaneLocationFragment> LaneLocationList = Context.GetMutableComponentView<FMassZoneGraphLaneLocationFragment>();
@@ -126,7 +126,7 @@ void UMassZoneGraphLocationInitializer::Execute(UMassEntitySubsystem& EntitySubs
 			if (!CurrentMovementConfig)
 			{
 #if WITH_MASSGAMEPLAY_DEBUG && UNSAFE_FOR_MT
-				const FLWEntity Entity = Context.GetEntity(EntityIndex);
+				const FMassEntityHandle Entity = Context.GetEntity(EntityIndex);
 				UE_VLOG(this, LogMassNavigation, Log, TEXT("Entity [%s] Invalid movement config."), *Entity.DebugGetDescription());
 #endif
 				continue;
@@ -187,25 +187,25 @@ void UMassZoneGraphPathFollowProcessor::Initialize(UObject& Owner)
 
 void UMassZoneGraphPathFollowProcessor::ConfigureQueries()
 {
-	EntityQuery_Conditional.AddRequirement<FMassZoneGraphShortPathFragment>(ELWComponentAccess::ReadWrite);
-	EntityQuery_Conditional.AddRequirement<FMassZoneGraphLaneLocationFragment>(ELWComponentAccess::ReadWrite);
-	EntityQuery_Conditional.AddRequirement<FMassMoveTargetFragment>(ELWComponentAccess::ReadWrite);
-	EntityQuery_Conditional.AddRequirement<FMassSimulationLODFragment>(ELWComponentAccess::ReadOnly);
-	EntityQuery_Conditional.AddChunkRequirement<FMassSimulationVariableTickChunkFragment>(ELWComponentAccess::ReadOnly);
+	EntityQuery_Conditional.AddRequirement<FMassZoneGraphShortPathFragment>(EMassFragmentAccess::ReadWrite);
+	EntityQuery_Conditional.AddRequirement<FMassZoneGraphLaneLocationFragment>(EMassFragmentAccess::ReadWrite);
+	EntityQuery_Conditional.AddRequirement<FMassMoveTargetFragment>(EMassFragmentAccess::ReadWrite);
+	EntityQuery_Conditional.AddRequirement<FMassSimulationLODFragment>(EMassFragmentAccess::ReadOnly);
+	EntityQuery_Conditional.AddChunkRequirement<FMassSimulationVariableTickChunkFragment>(EMassFragmentAccess::ReadOnly);
 	EntityQuery_Conditional.SetChunkFilter(&FMassSimulationVariableTickChunkFragment::ShouldTickChunkThisFrame);
 }
 
-void UMassZoneGraphPathFollowProcessor::Execute(UMassEntitySubsystem& EntitySubsystem, FLWComponentSystemExecutionContext& Context)
+void UMassZoneGraphPathFollowProcessor::Execute(UMassEntitySubsystem& EntitySubsystem, FMassExecutionContext& Context)
 {
 	if (!SignalSubsystem || !ZoneGraphSubsystem)
 	{
 		return;
 	}
 	
-	TArray<FLWEntity> EntitiesToSignalPathDone;
-	TArray<FLWEntity> EntitiesToSignalLaneChanged;
+	TArray<FMassEntityHandle> EntitiesToSignalPathDone;
+	TArray<FMassEntityHandle> EntitiesToSignalLaneChanged;
 
-	EntityQuery_Conditional.ForEachEntityChunk(EntitySubsystem, Context, [this, &EntitiesToSignalPathDone, &EntitiesToSignalLaneChanged](FLWComponentSystemExecutionContext& Context)
+	EntityQuery_Conditional.ForEachEntityChunk(EntitySubsystem, Context, [this, &EntitiesToSignalPathDone, &EntitiesToSignalLaneChanged](FMassExecutionContext& Context)
 		{
 		const int32 NumEntities = Context.GetEntitiesNum();
 		const TArrayView<FMassZoneGraphShortPathFragment> ShortPathList = Context.GetMutableComponentView<FMassZoneGraphShortPathFragment>();
@@ -219,7 +219,7 @@ void UMassZoneGraphPathFollowProcessor::Execute(UMassEntitySubsystem& EntitySubs
 			FMassZoneGraphLaneLocationFragment& LaneLocation = LaneLocationList[EntityIndex];
 			FMassMoveTargetFragment& MoveTarget = MoveTargetList[EntityIndex];
 			const FMassSimulationLODFragment& SimLOD = SimLODList[EntityIndex];
-			const FLWEntity Entity = Context.GetEntity(EntityIndex);
+			const FMassEntityHandle Entity = Context.GetEntity(EntityIndex);
 			const float DeltaTime = SimLOD.DeltaTime;
 
 			bool bDisplayDebug = false;
@@ -474,18 +474,18 @@ void UMassZoneGraphSteeringProcessor::Initialize(UObject& Owner)
 
 void UMassZoneGraphSteeringProcessor::ConfigureQueries()
 {
-	EntityQuery.AddRequirement<FMassMoveTargetFragment>(ELWComponentAccess::ReadWrite);
-	EntityQuery.AddRequirement<FMassMovementConfigFragment>(ELWComponentAccess::ReadOnly);
-	EntityQuery.AddRequirement<FDataFragment_Transform>(ELWComponentAccess::ReadWrite);
-	EntityQuery.AddRequirement<FMassSteeringFragment>(ELWComponentAccess::ReadWrite);
-	EntityQuery.AddRequirement<FMassSteeringGhostFragment>(ELWComponentAccess::ReadWrite);
-	EntityQuery.AddRequirement<FMassVelocityFragment>(ELWComponentAccess::ReadOnly);
+	EntityQuery.AddRequirement<FMassMoveTargetFragment>(EMassFragmentAccess::ReadWrite);
+	EntityQuery.AddRequirement<FMassMovementConfigFragment>(EMassFragmentAccess::ReadOnly);
+	EntityQuery.AddRequirement<FDataFragment_Transform>(EMassFragmentAccess::ReadWrite);
+	EntityQuery.AddRequirement<FMassSteeringFragment>(EMassFragmentAccess::ReadWrite);
+	EntityQuery.AddRequirement<FMassSteeringGhostFragment>(EMassFragmentAccess::ReadWrite);
+	EntityQuery.AddRequirement<FMassVelocityFragment>(EMassFragmentAccess::ReadOnly);
 
 	// No need for Off LOD to do steering, applying move target directly
-	EntityQuery.AddTagRequirement<FMassOffLODTag>(ELWComponentPresence::None);
+	EntityQuery.AddTagRequirement<FMassOffLODTag>(EMassFragmentPresence::None);
 }
 
-void UMassZoneGraphSteeringProcessor::Execute(UMassEntitySubsystem& EntitySubsystem, FLWComponentSystemExecutionContext& Context)
+void UMassZoneGraphSteeringProcessor::Execute(UMassEntitySubsystem& EntitySubsystem, FMassExecutionContext& Context)
 {
 	const UMassMovementSettings* Settings = GetDefault<UMassMovementSettings>();
 	check(Settings);
@@ -495,7 +495,7 @@ void UMassZoneGraphSteeringProcessor::Execute(UMassEntitySubsystem& EntitySubsys
 		return;
 	}
 	
-	EntityQuery.ForEachEntityChunk(EntitySubsystem, Context, [this, Settings](FLWComponentSystemExecutionContext& Context)
+	EntityQuery.ForEachEntityChunk(EntitySubsystem, Context, [this, Settings](FMassExecutionContext& Context)
 	{
 		const int32 NumEntities = Context.GetEntitiesNum();
 		const TArrayView<FMassMoveTargetFragment> MoveTargetList = Context.GetMutableComponentView<FMassMoveTargetFragment>();
@@ -523,7 +523,7 @@ void UMassZoneGraphSteeringProcessor::Execute(UMassEntitySubsystem& EntitySubsys
 			FMassMoveTargetFragment& MoveTarget = MoveTargetList[EntityIndex];
 			const FMassMovementConfigFragment& MovementConfig = MovementConfigList[EntityIndex];
 			const FMassVelocityFragment& Velocity = VelocityList[EntityIndex];
-			const FLWEntity Entity = Context.GetEntity(EntityIndex);
+			const FMassEntityHandle Entity = Context.GetEntity(EntityIndex);
 
 			if (MovementConfig.ConfigHandle != CurrentConfigHandle)
 			{
diff --git a/Engine/Plugins/AI/MassAI/Source/MassAIMovement/Private/MassZoneGraphMovementUtils.cpp b/Engine/Plugins/AI/MassAI/Source/MassAIMovement/Private/MassZoneGraphMovementUtils.cpp
--- a/Engine/Plugins/AI/MassAI/Source/MassAIMovement/Private/MassZoneGraphMovementUtils.cpp
+++ b/Engine/Plugins/AI/MassAI/Source/MassAIMovement/Private/MassZoneGraphMovementUtils.cpp
@@ -13,7 +13,7 @@ namespace UE::MassMovement
 
 	bool ActivateActionMove(const UWorld& World,
 							const UObject* Requester,
-							const FLWEntity Entity,
+							const FMassEntityHandle Entity,
 							const UZoneGraphSubsystem& ZoneGraphSubsystem,
 							const FMassZoneGraphLaneLocationFragment& LaneLocation,
 							const FZoneGraphShortPathRequest& PathRequest,
@@ -77,7 +77,7 @@ namespace UE::MassMovement
 
 	bool ActivateActionStand(const UWorld& World,
 							 const UObject* Requester,
-							 const FLWEntity Entity,
+							 const FMassEntityHandle Entity,
 							 const UZoneGraphSubsystem& ZoneGraphSubsystem,
 							 const FMassZoneGraphLaneLocationFragment& LaneLocation,
 							 const float DesiredSpeed,
@@ -108,7 +108,7 @@ namespace UE::MassMovement
 
 	bool ActivateActionAnimate(const UWorld& World,
 							   const UObject* Requester,
-							   const FLWEntity Entity,
+							   const FMassEntityHandle Entity,
 							   FMassMoveTargetFragment& MoveTarget)
 	{
 		MoveTarget.DistanceToGoal = 0.0f;
diff --git a/Engine/Plugins/AI/MassAI/Source/MassAIMovement/Private/MovementProcessor.cpp b/Engine/Plugins/AI/MassAI/Source/MassAIMovement/Private/MovementProcessor.cpp
--- a/Engine/Plugins/AI/MassAI/Source/MassAIMovement/Private/MovementProcessor.cpp
+++ b/Engine/Plugins/AI/MassAI/Source/MassAIMovement/Private/MovementProcessor.cpp
@@ -58,7 +58,7 @@ namespace UE::MassMovement
 		FAutoConsoleVariableRef(TEXT("ai.mass.AvoidDistance"), Tweakables::AvoidDistanceCutOff, TEXT(""), ECVF_Default), FAutoConsoleVariableRef(TEXT("ai.mass.AvoidRadius"), Tweakables::AvoidAgentRadius, TEXT(""), ECVF_Default), FAutoConsoleVariableRef(TEXT("ai.mass.AvoidMaxForce"), Tweakables::AvoidMaxForce, TEXT(""), ECVF_Default), FAutoConsoleVariableRef(TEXT("ai.mass.AvoidTimeHorizon"), Tweakables::AvoidTimeHorizon, TEXT(""), ECVF_Default)
 	};
 
- 	static bool DebugIsSelected(const FLWEntity& Entity)
+ 	static bool DebugIsSelected(const FMassEntityHandle& Entity)
  	{
  #if WITH_MASSGAMEPLAY_DEBUG
  		return UE::MassDebug::IsDebuggingEntity(Entity);
@@ -114,17 +114,17 @@ UMassProcessor_Movement::UMassProcessor_Movement()
 //----------------------------------------------------------------------//
 void UMassProcessor_Movement::ConfigureQueries()
 {
-	EntityQuery.AddRequirement<FDataFragment_Transform>(ELWComponentAccess::ReadWrite);
-	EntityQuery.AddRequirement<FMassVelocityFragment>(ELWComponentAccess::ReadWrite);
+	EntityQuery.AddRequirement<FDataFragment_Transform>(EMassFragmentAccess::ReadWrite);
+	EntityQuery.AddRequirement<FMassVelocityFragment>(EMassFragmentAccess::ReadWrite);
 }
 
-void UMassProcessor_Movement::Execute(UMassEntitySubsystem& EntitySubsystem, FLWComponentSystemExecutionContext& Context)
+void UMassProcessor_Movement::Execute(UMassEntitySubsystem& EntitySubsystem, FMassExecutionContext& Context)
 {
 	QUICK_SCOPE_CYCLE_COUNTER(MovementProcessor_Run);
 
 	const float TimeDelta = Context.GetDeltaTimeSeconds();
 
-	EntityQuery.ForEachEntityChunk(EntitySubsystem, Context, [this, TimeDelta](FLWComponentSystemExecutionContext& Context)
+	EntityQuery.ForEachEntityChunk(EntitySubsystem, Context, [this, TimeDelta](FMassExecutionContext& Context)
 		{
 			const int32 NumEntities = Context.GetEntitiesNum();
 			const TArrayView<FDataFragment_Transform> LocationList = Context.GetMutableComponentView<FDataFragment_Transform>();
@@ -153,7 +153,7 @@ UMassProcessor_AgentMovement::UMassProcessor_AgentMovement()
 void UMassProcessor_AgentMovement::ConfigureQueries()
 {
 	Super::ConfigureQueries();
-	EntityQuery.AddRequirement<FDataFragment_AgentRadius>(ELWComponentAccess::ReadOnly);
+	EntityQuery.AddRequirement<FDataFragment_AgentRadius>(EMassFragmentAccess::ReadOnly);
 }
 
 //----------------------------------------------------------------------//
@@ -167,7 +167,7 @@ UMassAvoidanceObstacleRemoverFragmentDestructor::UMassAvoidanceObstacleRemoverFr
 
 void UMassAvoidanceObstacleRemoverFragmentDestructor::ConfigureQueries()
 {
-	EntityQuery.AddRequirement<FMassAvoidanceObstacleGridCellLocationFragment>(ELWComponentAccess::ReadWrite);
+	EntityQuery.AddRequirement<FMassAvoidanceObstacleGridCellLocationFragment>(EMassFragmentAccess::ReadWrite);
 }
 
 void UMassAvoidanceObstacleRemoverFragmentDestructor::Initialize(UObject& Owner)
@@ -176,7 +176,7 @@ void UMassAvoidanceObstacleRemoverFragmentDestructor::Initialize(UObject& Owner)
 	WeakMovementSubsystem = UWorld::GetSubsystem<UMassMovementSubsystem>(Owner.GetWorld());
 }
 
-void UMassAvoidanceObstacleRemoverFragmentDestructor::Execute(UMassEntitySubsystem& EntitySubsystem, FLWComponentSystemExecutionContext& Context)
+void UMassAvoidanceObstacleRemoverFragmentDestructor::Execute(UMassEntitySubsystem& EntitySubsystem, FMassExecutionContext& Context)
 {
 	UMassMovementSubsystem* MovementSubsystem = WeakMovementSubsystem.Get();
 	if (!MovementSubsystem)
@@ -184,7 +184,7 @@ void UMassAvoidanceObstacleRemoverFragmentDestructor::Execute(UMassEntitySubsyst
 		return;
 	}
 
-	EntityQuery.ForEachEntityChunk(EntitySubsystem, Context, [this, MovementSubsystem](FLWComponentSystemExecutionContext& Context)
+	EntityQuery.ForEachEntityChunk(EntitySubsystem, Context, [this, MovementSubsystem](FMassExecutionContext& Context)
 		{
 			const int32 NumEntities = Context.GetEntitiesNum();
 			const TArrayView<FMassAvoidanceObstacleGridCellLocationFragment> AvoidanceObstacleCellLocationList = Context.GetMutableComponentView<FMassAvoidanceObstacleGridCellLocationFragment>();
@@ -209,12 +209,12 @@ UVelocityFragmentInitializer::UVelocityFragmentInitializer()
 
 void UVelocityFragmentInitializer::ConfigureQueries()
 {
-	EntityQuery.AddRequirement<FMassVelocityFragment>(ELWComponentAccess::ReadWrite);
+	EntityQuery.AddRequirement<FMassVelocityFragment>(EMassFragmentAccess::ReadWrite);
 }
 
-void UVelocityFragmentInitializer::Execute(UMassEntitySubsystem& EntitySubsystem, FLWComponentSystemExecutionContext& Context)
+void UVelocityFragmentInitializer::Execute(UMassEntitySubsystem& EntitySubsystem, FMassExecutionContext& Context)
 {
-	EntityQuery.ForEachEntityChunk(EntitySubsystem, Context, [this](FLWComponentSystemExecutionContext& Context)
+	EntityQuery.ForEachEntityChunk(EntitySubsystem, Context, [this](FMassExecutionContext& Context)
 		{
 			const TArrayView<FMassVelocityFragment> VelocitiesList = Context.GetMutableComponentView<FMassVelocityFragment>();
 			for (FMassVelocityFragment& VelocityFragment : VelocitiesList)
diff --git a/Engine/Plugins/AI/MassAI/Source/MassAIMovement/Private/NavigationProcessor.cpp b/Engine/Plugins/AI/MassAI/Source/MassAIMovement/Private/NavigationProcessor.cpp
--- a/Engine/Plugins/AI/MassAI/Source/MassAIMovement/Private/NavigationProcessor.cpp
+++ b/Engine/Plugins/AI/MassAI/Source/MassAIMovement/Private/NavigationProcessor.cpp
@@ -85,21 +85,21 @@ UMassApplyVelocityMoveTargetProcessor::UMassApplyVelocityMoveTargetProcessor()
 
 void UMassApplyVelocityMoveTargetProcessor::ConfigureQueries()
 {
-	HighResEntityQuery.AddRequirement<FMassMoveTargetFragment>(ELWComponentAccess::ReadOnly);
-	HighResEntityQuery.AddRequirement<FMassVelocityFragment>(ELWComponentAccess::ReadWrite);
-	HighResEntityQuery.AddRequirement<FDataFragment_Transform>(ELWComponentAccess::ReadWrite);
-	HighResEntityQuery.AddRequirement<FMassSteeringFragment>(ELWComponentAccess::ReadWrite);
-	HighResEntityQuery.AddTagRequirement<FMassOffLODTag>(ELWComponentPresence::None);
-
-	LowResEntityQuery_Conditional.AddRequirement<FDataFragment_Transform>(ELWComponentAccess::ReadWrite);
-	LowResEntityQuery_Conditional.AddRequirement<FMassMoveTargetFragment>(ELWComponentAccess::ReadOnly);
-	LowResEntityQuery_Conditional.AddTagRequirement<FMassOffLODTag>(ELWComponentPresence::All);
-	LowResEntityQuery_Conditional.AddChunkRequirement<FMassSimulationVariableTickChunkFragment>(ELWComponentAccess::ReadOnly);
+	HighResEntityQuery.AddRequirement<FMassMoveTargetFragment>(EMassFragmentAccess::ReadOnly);
+	HighResEntityQuery.AddRequirement<FMassVelocityFragment>(EMassFragmentAccess::ReadWrite);
+	HighResEntityQuery.AddRequirement<FDataFragment_Transform>(EMassFragmentAccess::ReadWrite);
+	HighResEntityQuery.AddRequirement<FMassSteeringFragment>(EMassFragmentAccess::ReadWrite);
+	HighResEntityQuery.AddTagRequirement<FMassOffLODTag>(EMassFragmentPresence::None);
+
+	LowResEntityQuery_Conditional.AddRequirement<FDataFragment_Transform>(EMassFragmentAccess::ReadWrite);
+	LowResEntityQuery_Conditional.AddRequirement<FMassMoveTargetFragment>(EMassFragmentAccess::ReadOnly);
+	LowResEntityQuery_Conditional.AddTagRequirement<FMassOffLODTag>(EMassFragmentPresence::All);
+	LowResEntityQuery_Conditional.AddChunkRequirement<FMassSimulationVariableTickChunkFragment>(EMassFragmentAccess::ReadOnly);
 	LowResEntityQuery_Conditional.SetChunkFilter(&FMassSimulationVariableTickChunkFragment::ShouldTickChunkThisFrame);
 }
 
 void UMassApplyVelocityMoveTargetProcessor::Execute(UMassEntitySubsystem& EntitySubsystem,
-													FLWComponentSystemExecutionContext& Context)
+													FMassExecutionContext& Context)
 {
 	// Clamp max delta time to avoid force explosion on large time steps (i.e. during initialization).
 	const float TimeDelta = FMath::Min(0.1f, Context.GetDeltaTimeSeconds());
@@ -110,7 +110,7 @@ void UMassApplyVelocityMoveTargetProcessor::Execute(UMassEntitySubsystem& Entity
 	{
 		QUICK_SCOPE_CYCLE_COUNTER(HighRes);
 
-		HighResEntityQuery.ForEachEntityChunk(EntitySubsystem, Context, [this, Settings, TimeDelta](FLWComponentSystemExecutionContext& Context)
+		HighResEntityQuery.ForEachEntityChunk(EntitySubsystem, Context, [this, Settings, TimeDelta](FMassExecutionContext& Context)
 		{
 			const int32 NumEntities = Context.GetEntitiesNum();
 			constexpr float ZDamperHalfLife = 0.2f;
@@ -203,7 +203,7 @@ void UMassApplyVelocityMoveTargetProcessor::Execute(UMassEntitySubsystem& Entity
 	{
 		QUICK_SCOPE_CYCLE_COUNTER(LowRes);
 
-		LowResEntityQuery_Conditional.ForEachEntityChunk(EntitySubsystem, Context, [this](FLWComponentSystemExecutionContext& Context)
+		LowResEntityQuery_Conditional.ForEachEntityChunk(EntitySubsystem, Context, [this](FMassExecutionContext& Context)
 			{
 				const int32 NumEntities = Context.GetEntitiesNum();
 
@@ -234,16 +234,16 @@ UMassDynamicObstacleProcessor::UMassDynamicObstacleProcessor()
 
 void UMassDynamicObstacleProcessor::ConfigureQueries()
 {
-	EntityQuery_Conditional.AddRequirement<FDataFragment_Transform>(ELWComponentAccess::ReadOnly);
-	EntityQuery_Conditional.AddRequirement<FDataFragment_AgentRadius>(ELWComponentAccess::ReadOnly);
-	EntityQuery_Conditional.AddRequirement<FMassDynamicObstacleFragment>(ELWComponentAccess::ReadWrite);
-	EntityQuery_Conditional.AddChunkRequirement<FMassSimulationVariableTickChunkFragment>(ELWComponentAccess::ReadOnly);
+	EntityQuery_Conditional.AddRequirement<FDataFragment_Transform>(EMassFragmentAccess::ReadOnly);
+	EntityQuery_Conditional.AddRequirement<FDataFragment_AgentRadius>(EMassFragmentAccess::ReadOnly);
+	EntityQuery_Conditional.AddRequirement<FMassDynamicObstacleFragment>(EMassFragmentAccess::ReadWrite);
+	EntityQuery_Conditional.AddChunkRequirement<FMassSimulationVariableTickChunkFragment>(EMassFragmentAccess::ReadOnly);
 	EntityQuery_Conditional.SetChunkFilter(&FMassSimulationVariableTickChunkFragment::ShouldTickChunkThisFrame);
 }
 
-void UMassDynamicObstacleProcessor::Execute(UMassEntitySubsystem& EntitySubsystem, FLWComponentSystemExecutionContext& Context)
+void UMassDynamicObstacleProcessor::Execute(UMassEntitySubsystem& EntitySubsystem, FMassExecutionContext& Context)
 {
-	EntityQuery_Conditional.ForEachEntityChunk(EntitySubsystem, Context, [this](FLWComponentSystemExecutionContext& Context)
+	EntityQuery_Conditional.ForEachEntityChunk(EntitySubsystem, Context, [this](FMassExecutionContext& Context)
 		{
 			const int32 NumEntities = Context.GetEntitiesNum();
 
diff --git a/Engine/Plugins/AI/MassAI/Source/MassAIMovement/Private/SnapToNavigationProcessor.cpp b/Engine/Plugins/AI/MassAI/Source/MassAIMovement/Private/SnapToNavigationProcessor.cpp
--- a/Engine/Plugins/AI/MassAI/Source/MassAIMovement/Private/SnapToNavigationProcessor.cpp
+++ b/Engine/Plugins/AI/MassAI/Source/MassAIMovement/Private/SnapToNavigationProcessor.cpp
@@ -19,11 +19,11 @@ USnapToNavigationProcessor::USnapToNavigationProcessor()
 
 void USnapToNavigationProcessor::ConfigureQueries()
 {
-	EntityQuery.AddRequirement<FDataFragment_Transform>(ELWComponentAccess::ReadWrite);
-	EntityQuery.AddRequirement<FDataFragment_NavLocation>(ELWComponentAccess::ReadWrite);
+	EntityQuery.AddRequirement<FDataFragment_Transform>(EMassFragmentAccess::ReadWrite);
+	EntityQuery.AddRequirement<FDataFragment_NavLocation>(EMassFragmentAccess::ReadWrite);
 }
 
-void USnapToNavigationProcessor::Execute(UMassEntitySubsystem& EntitySubsystem, FLWComponentSystemExecutionContext& Context)
+void USnapToNavigationProcessor::Execute(UMassEntitySubsystem& EntitySubsystem, FMassExecutionContext& Context)
 {
 	ANavigationData* NavData = WeakNavData.Get();
 	if (NavData == nullptr)
@@ -33,7 +33,7 @@ void USnapToNavigationProcessor::Execute(UMassEntitySubsystem& EntitySubsystem,
 
 	QUICK_SCOPE_CYCLE_COUNTER(UMovementProcessor_Run);
 
-	EntityQuery.ForEachEntityChunk(EntitySubsystem, Context, [this, NavData](FLWComponentSystemExecutionContext& Context)
+	EntityQuery.ForEachEntityChunk(EntitySubsystem, Context, [this, NavData](FMassExecutionContext& Context)
 		{
 			const int32 NumEntities = Context.GetEntitiesNum();
 			const TArrayView<FDataFragment_Transform> LocationList = Context.GetMutableComponentView<FDataFragment_Transform>();
