AggregateTicking and Pipe type renames:

diff --git a/Engine/Plugins/AI/MassAI/Source/MassAIMovement/Public/AssignRandomNavLocationProcessor.h b/Engine/Plugins/AI/MassAI/Source/MassAIMovement/Public/AssignRandomNavLocationProcessor.h
--- a/Engine/Plugins/AI/MassAI/Source/MassAIMovement/Public/AssignRandomNavLocationProcessor.h
+++ b/Engine/Plugins/AI/MassAI/Source/MassAIMovement/Public/AssignRandomNavLocationProcessor.h
@@ -10,7 +10,7 @@
 class UNavigationSystemV1;
 
 UCLASS()
-class MASSAIMOVEMENT_API UAssignRandomNavLocationProcessor : public UPipeProcessor
+class MASSAIMOVEMENT_API UAssignRandomNavLocationProcessor : public UMassProcessor
 {
 	GENERATED_BODY()
 
@@ -19,7 +19,7 @@ public:
 
 protected:
 	virtual void ConfigureQueries() override;
-	virtual void Execute(UMassEntitySubsystem& EntitySubsystem, FLWComponentSystemExecutionContext& Context) override;
+	virtual void Execute(UMassEntitySubsystem& EntitySubsystem, FMassExecutionContext& Context) override;
 	virtual void Initialize(UObject& InOwner) override;
 
 	UPROPERTY()
@@ -30,5 +30,5 @@ protected:
 
 	FVector Origin = FNavigationSystem::InvalidLocation;
 
-	FLWComponentQuery EntityQuery;
+	FMassEntityQuery EntityQuery;
 };
\ No newline at end of file
diff --git a/Engine/Plugins/AI/MassAI/Source/MassAIMovement/Public/MassAIMovementFragments.h b/Engine/Plugins/AI/MassAI/Source/MassAIMovement/Public/MassAIMovementFragments.h
--- a/Engine/Plugins/AI/MassAI/Source/MassAIMovement/Public/MassAIMovementFragments.h
+++ b/Engine/Plugins/AI/MassAI/Source/MassAIMovement/Public/MassAIMovementFragments.h
@@ -41,7 +41,7 @@ struct FMassPillCollider
 
 /** Fragment holding data for avoidance coliders */
 USTRUCT()
-struct MASSAIMOVEMENT_API FMassAvoidanceColliderFragment : public FLWComponentData
+struct MASSAIMOVEMENT_API FMassAvoidanceColliderFragment : public FMassFragment
 {
 	GENERATED_BODY()
 
@@ -84,7 +84,7 @@ struct MASSAIMOVEMENT_API FMassAvoidanceColliderFragment : public FLWComponentDa
 
 /** Cell location for dynamic obstacles */
 USTRUCT()
-struct MASSAIMOVEMENT_API FMassAvoidanceObstacleGridCellLocationFragment : public FLWComponentData
+struct MASSAIMOVEMENT_API FMassAvoidanceObstacleGridCellLocationFragment : public FMassFragment
 {
 	GENERATED_BODY()
 	FAvoidanceObstacleHashGrid2D::FCellLocation CellLoc;
@@ -92,7 +92,7 @@ struct MASSAIMOVEMENT_API FMassAvoidanceObstacleGridCellLocationFragment : publi
 
 /** Move target. */
 USTRUCT()
-struct MASSAIMOVEMENT_API FMassMoveTargetFragment : public FLWComponentData
+struct MASSAIMOVEMENT_API FMassMoveTargetFragment : public FMassFragment
 {
 	GENERATED_BODY()
 
@@ -164,7 +164,7 @@ public:
 
 /** Steering fragment. */
 USTRUCT()
-struct MASSAIMOVEMENT_API FMassSteeringFragment : public FLWComponentData
+struct MASSAIMOVEMENT_API FMassSteeringFragment : public FMassFragment
 {
 	GENERATED_BODY()
 
@@ -183,7 +183,7 @@ struct MASSAIMOVEMENT_API FMassSteeringFragment : public FLWComponentData
 
 /** Steering ghost fragment. */
 USTRUCT()
-struct MASSAIMOVEMENT_API FMassSteeringGhostFragment : public FLWComponentData
+struct MASSAIMOVEMENT_API FMassSteeringGhostFragment : public FMassFragment
 {
 	GENERATED_BODY()
 
@@ -218,7 +218,7 @@ struct MASSAIMOVEMENT_API FMassSteeringGhostFragment : public FLWComponentData
 };
 
 USTRUCT()
-struct MASSAIMOVEMENT_API FMassDynamicObstacleFragment : public FLWComponentData
+struct MASSAIMOVEMENT_API FMassDynamicObstacleFragment : public FMassFragment
 {
 	GENERATED_BODY()
 
@@ -238,7 +238,7 @@ struct MASSAIMOVEMENT_API FMassDynamicObstacleFragment : public FLWComponentData
 
 /** Stores handle to rich movement config */
 USTRUCT()
-struct MASSAIMOVEMENT_API FMassMovementConfigFragment : public FLWComponentData
+struct MASSAIMOVEMENT_API FMassMovementConfigFragment : public FMassFragment
 {
 	GENERATED_BODY()
 
diff --git a/Engine/Plugins/AI/MassAI/Source/MassAIMovement/Public/MassAvoidanceProcessors.h b/Engine/Plugins/AI/MassAI/Source/MassAIMovement/Public/MassAvoidanceProcessors.h
--- a/Engine/Plugins/AI/MassAI/Source/MassAIMovement/Public/MassAvoidanceProcessors.h
+++ b/Engine/Plugins/AI/MassAI/Source/MassAIMovement/Public/MassAvoidanceProcessors.h
@@ -18,7 +18,7 @@ MASSAIMOVEMENT_API DECLARE_LOG_CATEGORY_EXTERN(LogAvoidanceAgents, Warning, All)
 MASSAIMOVEMENT_API DECLARE_LOG_CATEGORY_EXTERN(LogAvoidanceObstacles, Warning, All);
 
 USTRUCT()
-struct MASSAIMOVEMENT_API FMassEdgeDetectionParamsFragment : public FLWComponentData
+struct MASSAIMOVEMENT_API FMassEdgeDetectionParamsFragment : public FMassFragment
 {
 	GENERATED_BODY()
 	float EdgeDetectionRange = 500.f;
@@ -39,7 +39,7 @@ struct FNavigationAvoidanceEdge
 };
 
 USTRUCT()
-struct MASSAIMOVEMENT_API FMassNavigationEdgesFragment : public FLWComponentData
+struct MASSAIMOVEMENT_API FMassNavigationEdgesFragment : public FMassFragment
 {
 	GENERATED_BODY()
 
@@ -49,14 +49,14 @@ struct MASSAIMOVEMENT_API FMassNavigationEdgesFragment : public FLWComponentData
 
 /** Component tag to tell the avoidance to extend the size of this obstacle when too close to edges. */
 USTRUCT()
-struct MASSAIMOVEMENT_API FMassAvoidanceExtendToEdgeObstacleTag : public FComponentTag
+struct MASSAIMOVEMENT_API FMassAvoidanceExtendToEdgeObstacleTag : public FMassTag
 {
 	GENERATED_BODY()
 };
 
 /** Experimental: move using cumulative forces to avoid close agents */
 UCLASS()
-class MASSAIMOVEMENT_API UMassAvoidanceProcessor : public UPipeProcessor
+class MASSAIMOVEMENT_API UMassAvoidanceProcessor : public UMassProcessor
 {
 	GENERATED_BODY()
 
@@ -66,17 +66,17 @@ public:
 protected:
 	virtual void ConfigureQueries() override;
 	virtual void Initialize(UObject& Owner) override;
-	virtual void Execute(UMassEntitySubsystem& EntitySubsystem, FLWComponentSystemExecutionContext& Context) override;
+	virtual void Execute(UMassEntitySubsystem& EntitySubsystem, FMassExecutionContext& Context) override;
 
 private:
 	TWeakObjectPtr<UWorld> WeakWorld;
 	TWeakObjectPtr<UMassMovementSubsystem> WeakMovementSubsystem;
-	FLWComponentQuery EntityQuery;
+	FMassEntityQuery EntityQuery;
 };
 
 /** Avoidance while standing. */
 UCLASS()
-class MASSAIMOVEMENT_API UMassStandingAvoidanceProcessor : public UPipeProcessor
+class MASSAIMOVEMENT_API UMassStandingAvoidanceProcessor : public UMassProcessor
 {
 	GENERATED_BODY()
 
@@ -86,24 +86,24 @@ public:
 protected:
 	virtual void ConfigureQueries() override;
 	virtual void Initialize(UObject& Owner) override;
-	virtual void Execute(UMassEntitySubsystem& EntitySubsystem, FLWComponentSystemExecutionContext& Context) override;
+	virtual void Execute(UMassEntitySubsystem& EntitySubsystem, FMassExecutionContext& Context) override;
 
 private:
 	TWeakObjectPtr<UWorld> WeakWorld;
 	TWeakObjectPtr<UMassMovementSubsystem> WeakMovementSubsystem;
-	FLWComponentQuery EntityQuery;
+	FMassEntityQuery EntityQuery;
 };
 
 /** Component Tag to tell if the entity is in the avoidance obstacle grid */
 USTRUCT()
-struct MASSAIMOVEMENT_API FMassInAvoidanceObstacleGridTag : public FComponentTag
+struct MASSAIMOVEMENT_API FMassInAvoidanceObstacleGridTag : public FMassTag
 {
 	GENERATED_BODY()
 };
 
 /** Processor to update avoidance obstacle data */
 UCLASS()
-class MASSAIMOVEMENT_API UMassAvoidanceObstacleProcessor : public UPipeProcessor
+class MASSAIMOVEMENT_API UMassAvoidanceObstacleProcessor : public UMassProcessor
 {
 	GENERATED_BODY()
 
@@ -113,18 +113,18 @@ public:
 protected:
 	virtual void ConfigureQueries() override;
 	virtual void Initialize(UObject& Owner) override;
-	virtual void Execute(UMassEntitySubsystem& EntitySubsystem, FLWComponentSystemExecutionContext& Context) override;
+	virtual void Execute(UMassEntitySubsystem& EntitySubsystem, FMassExecutionContext& Context) override;
 
 private:
 	TWeakObjectPtr<UMassMovementSubsystem> WeakMovementSubsystem;
-	FLWComponentQuery AddToGridEntityQuery;
-	FLWComponentQuery UpdateGridEntityQuery;
-	FLWComponentQuery RemoveFromGridEntityQuery;
+	FMassEntityQuery AddToGridEntityQuery;
+	FMassEntityQuery UpdateGridEntityQuery;
+	FMassEntityQuery RemoveFromGridEntityQuery;
 };
 
 /** Experimental: navmesh edges gathering */
 UCLASS()
-class MASSAIMOVEMENT_API UMassNavigationBoundaryProcessor : public UPipeProcessor
+class MASSAIMOVEMENT_API UMassNavigationBoundaryProcessor : public UMassProcessor
 {
 	GENERATED_BODY()
 
@@ -134,24 +134,24 @@ public:
 protected:
 	virtual void ConfigureQueries() override;
 	virtual void Initialize(UObject& Owner) override;
-	virtual void Execute(UMassEntitySubsystem& EntitySubsystem, FLWComponentSystemExecutionContext& Context) override;
+	virtual void Execute(UMassEntitySubsystem& EntitySubsystem, FMassExecutionContext& Context) override;
 
 private:
 	TWeakObjectPtr<UWorld> WeakWorld;
 	TWeakObjectPtr<UMassMovementSubsystem> WeakMovementSubsystem;
 	TWeakObjectPtr<ANavigationData> WeakNavData;
-	FLWComponentQuery EntityQuery;
+	FMassEntityQuery EntityQuery;
 };
 
 USTRUCT()
-struct MASSAIMOVEMENT_API FMassLastUpdatePositionFragment : public FLWComponentData
+struct MASSAIMOVEMENT_API FMassLastUpdatePositionFragment : public FMassFragment
 {
 	GENERATED_BODY()
 	FVector Value;
 };
 
 USTRUCT()
-struct MASSAIMOVEMENT_API FMassAvoidanceBoundaryLastLaneHandleFragment : public FLWComponentData
+struct MASSAIMOVEMENT_API FMassAvoidanceBoundaryLastLaneHandleFragment : public FMassFragment
 {
 	GENERATED_BODY()
 
@@ -161,7 +161,7 @@ struct MASSAIMOVEMENT_API FMassAvoidanceBoundaryLastLaneHandleFragment : public
 
 /** Experimental: lane borders gathering */
 UCLASS()
-class MASSAIMOVEMENT_API UMassLaneBoundaryProcessor : public UPipeProcessor
+class MASSAIMOVEMENT_API UMassLaneBoundaryProcessor : public UMassProcessor
 {
 	GENERATED_BODY()
 
@@ -171,16 +171,16 @@ public:
 protected:
 	virtual void ConfigureQueries() override;
 	virtual void Initialize(UObject& Owner) override;
-	virtual void Execute(UMassEntitySubsystem& EntitySubsystem, FLWComponentSystemExecutionContext& Context) override;
+	virtual void Execute(UMassEntitySubsystem& EntitySubsystem, FMassExecutionContext& Context) override;
 
 private:
 	TWeakObjectPtr<UWorld> WeakWorld;
 	TWeakObjectPtr<UZoneGraphSubsystem> WeakZoneGraph;
-	FLWComponentQuery EntityQuery;
+	FMassEntityQuery EntityQuery;
 };
 
 USTRUCT()
-struct MASSAIMOVEMENT_API FMassLaneCacheBoundaryFragment : public FLWComponentData
+struct MASSAIMOVEMENT_API FMassLaneCacheBoundaryFragment : public FMassFragment
 {
 	GENERATED_BODY()
 
@@ -194,7 +194,7 @@ struct MASSAIMOVEMENT_API FMassLaneCacheBoundaryFragment : public FLWComponentDa
 /** ZoneGraph lane cache boundary processor */
 // @todo MassMovement: Make this signal based.
 UCLASS()
-class MASSAIMOVEMENT_API UMassLaneCacheBoundaryProcessor : public UPipeProcessor
+class MASSAIMOVEMENT_API UMassLaneCacheBoundaryProcessor : public UMassProcessor
 {
 	GENERATED_BODY()
 
@@ -204,9 +204,9 @@ public:
 protected:
 	virtual void ConfigureQueries() override;
 	virtual void Initialize(UObject& Owner) override;
-	virtual void Execute(UMassEntitySubsystem& EntitySubsystem, FLWComponentSystemExecutionContext& Context) override;
+	virtual void Execute(UMassEntitySubsystem& EntitySubsystem, FMassExecutionContext& Context) override;
 
 private:
 	TWeakObjectPtr<UWorld> WeakWorld;
-	FLWComponentQuery EntityQuery;
+	FMassEntityQuery EntityQuery;
 };
\ No newline at end of file
diff --git a/Engine/Plugins/AI/MassAI/Source/MassAIMovement/Public/MassZoneGraphMovementFragments.h b/Engine/Plugins/AI/MassAI/Source/MassAIMovement/Public/MassZoneGraphMovementFragments.h
--- a/Engine/Plugins/AI/MassAI/Source/MassAIMovement/Public/MassZoneGraphMovementFragments.h
+++ b/Engine/Plugins/AI/MassAI/Source/MassAIMovement/Public/MassZoneGraphMovementFragments.h
@@ -10,7 +10,7 @@
 
 /** Stores path request associated to a new movement action. This is used to replicate actions. */
 USTRUCT()
-struct MASSAIMOVEMENT_API FMassZoneGraphPathRequestFragment : public FLWComponentData
+struct MASSAIMOVEMENT_API FMassZoneGraphPathRequestFragment : public FMassFragment
 {
 	GENERATED_BODY()
 
@@ -21,7 +21,7 @@ struct MASSAIMOVEMENT_API FMassZoneGraphPathRequestFragment : public FLWComponen
 
 /** Describes current location on ZoneGraph */ 
 USTRUCT()
-struct MASSAIMOVEMENT_API FMassZoneGraphLaneLocationFragment : public FLWComponentData
+struct MASSAIMOVEMENT_API FMassZoneGraphLaneLocationFragment : public FMassFragment
 {
 	GENERATED_BODY()
 
@@ -37,7 +37,7 @@ struct MASSAIMOVEMENT_API FMassZoneGraphLaneLocationFragment : public FLWCompone
 
 /** Describes part of a ZoneGraph lane. */
 USTRUCT()
-struct MASSAIMOVEMENT_API FMassZoneGraphCachedLaneFragment : public FLWComponentData
+struct MASSAIMOVEMENT_API FMassZoneGraphCachedLaneFragment : public FMassFragment
 {
 	GENERATED_BODY()
 
@@ -168,7 +168,7 @@ struct MASSAIMOVEMENT_API FMassZoneGraphCachedLaneFragment : public FLWComponent
 // - remove debug lane handle, and replace other with index
 // - see if we can remove move tangent?
 USTRUCT()
-struct MASSAIMOVEMENT_API FMassZoneGraphShortPathFragment : public FLWComponentData
+struct MASSAIMOVEMENT_API FMassZoneGraphShortPathFragment : public FMassFragment
 {
 	GENERATED_BODY()
 
diff --git a/Engine/Plugins/AI/MassAI/Source/MassAIMovement/Public/MassZoneGraphMovementProcessors.h b/Engine/Plugins/AI/MassAI/Source/MassAIMovement/Public/MassZoneGraphMovementProcessors.h
--- a/Engine/Plugins/AI/MassAI/Source/MassAIMovement/Public/MassZoneGraphMovementProcessors.h
+++ b/Engine/Plugins/AI/MassAI/Source/MassAIMovement/Public/MassZoneGraphMovementProcessors.h
@@ -23,9 +23,9 @@ public:
 protected:
 	virtual void Initialize(UObject& Owner) override;
 	virtual void ConfigureQueries() override;
-	virtual void Execute(UMassEntitySubsystem& EntitySubsystem, FLWComponentSystemExecutionContext& Context) override;
+	virtual void Execute(UMassEntitySubsystem& EntitySubsystem, FMassExecutionContext& Context) override;
 
-	FLWComponentQuery EntityQuery;
+	FMassEntityQuery EntityQuery;
 
 	UPROPERTY(Transient)
 	UZoneGraphSubsystem* ZoneGraphSubsystem = nullptr;
@@ -38,7 +38,7 @@ protected:
  * Processor for updating move target on ZoneGraph path.
  */
 UCLASS()
-class MASSAIMOVEMENT_API UMassZoneGraphPathFollowProcessor : public UPipeProcessor
+class MASSAIMOVEMENT_API UMassZoneGraphPathFollowProcessor : public UMassProcessor
 {
 	GENERATED_BODY()
 
@@ -47,9 +47,9 @@ protected:
 	
 	virtual void Initialize(UObject& Owner) override;
 	virtual void ConfigureQueries() override;
-	virtual void Execute(UMassEntitySubsystem& EntitySubsystem, FLWComponentSystemExecutionContext& Context) override;
+	virtual void Execute(UMassEntitySubsystem& EntitySubsystem, FMassExecutionContext& Context) override;
 
-	FLWComponentQuery EntityQuery_Conditional;
+	FMassEntityQuery EntityQuery_Conditional;
 
 	UPROPERTY(Transient)
 	UZoneGraphSubsystem* ZoneGraphSubsystem = nullptr;
@@ -62,7 +62,7 @@ protected:
  * Processor for updating steering towards MoveTarget.
  */
 UCLASS()
-class MASSAIMOVEMENT_API UMassZoneGraphSteeringProcessor : public UPipeProcessor
+class MASSAIMOVEMENT_API UMassZoneGraphSteeringProcessor : public UMassProcessor
 {
 	GENERATED_BODY()
 
@@ -71,9 +71,9 @@ protected:
 	
 	virtual void Initialize(UObject& Owner) override;
 	virtual void ConfigureQueries() override;
-	virtual void Execute(UMassEntitySubsystem& EntitySubsystem, FLWComponentSystemExecutionContext& Context) override;
+	virtual void Execute(UMassEntitySubsystem& EntitySubsystem, FMassExecutionContext& Context) override;
 
-	FLWComponentQuery EntityQuery;
+	FMassEntityQuery EntityQuery;
 
 	UPROPERTY(Transient)
 	UMassSignalSubsystem* SignalSubsystem = nullptr;
diff --git a/Engine/Plugins/AI/MassAI/Source/MassAIMovement/Public/MassZoneGraphMovementUtils.h b/Engine/Plugins/AI/MassAI/Source/MassAIMovement/Public/MassZoneGraphMovementUtils.h
--- a/Engine/Plugins/AI/MassAI/Source/MassAIMovement/Public/MassZoneGraphMovementUtils.h
+++ b/Engine/Plugins/AI/MassAI/Source/MassAIMovement/Public/MassZoneGraphMovementUtils.h
@@ -18,7 +18,7 @@ namespace UE::MassMovement
 {
 	MASSAIMOVEMENT_API bool ActivateActionMove(const UWorld& World,
 											 const UObject* Requester,
-											 const FLWEntity Entity,
+											 const FMassEntityHandle Entity,
 											 const UZoneGraphSubsystem& ZoneGraphSubsystem,
 											 const FMassZoneGraphLaneLocationFragment& LaneLocation,
 											 const FZoneGraphShortPathRequest& PathRequest,
@@ -30,7 +30,7 @@ namespace UE::MassMovement
 
 	MASSAIMOVEMENT_API bool ActivateActionStand(const UWorld& World,
 											  const UObject* Requester,
-											  const FLWEntity Entity,
+											  const FMassEntityHandle Entity,
 											  const UZoneGraphSubsystem& ZoneGraphSubsystem,
 											  const FMassZoneGraphLaneLocationFragment& LaneLocation,
 											  const float DesiredSpeed,
@@ -40,6 +40,6 @@ namespace UE::MassMovement
 
 	MASSAIMOVEMENT_API bool ActivateActionAnimate(const UWorld& World,
 												const UObject* Requester,
-												const FLWEntity Entity,
+												const FMassEntityHandle Entity,
 												FMassMoveTargetFragment& MoveTarget);
 };
diff --git a/Engine/Plugins/AI/MassAI/Source/MassAIMovement/Public/MovementProcessor.h b/Engine/Plugins/AI/MassAI/Source/MassAIMovement/Public/MovementProcessor.h
--- a/Engine/Plugins/AI/MassAI/Source/MassAIMovement/Public/MovementProcessor.h
+++ b/Engine/Plugins/AI/MassAI/Source/MassAIMovement/Public/MovementProcessor.h
@@ -9,7 +9,7 @@
 
 /** Move and orient */ 
 UCLASS()
-class MASSAIMOVEMENT_API UMassProcessor_Movement : public UPipeProcessor
+class MASSAIMOVEMENT_API UMassProcessor_Movement : public UMassProcessor
 {
 	GENERATED_BODY()
 
@@ -18,9 +18,9 @@ public:
 
 protected:
 	virtual void ConfigureQueries() override;
-	virtual void Execute(UMassEntitySubsystem& EntitySubsystem, FLWComponentSystemExecutionContext& Context) override;
+	virtual void Execute(UMassEntitySubsystem& EntitySubsystem, FMassExecutionContext& Context) override;
 
-	FLWComponentQuery EntityQuery;
+	FMassEntityQuery EntityQuery;
 };
 
 /** Base movement processor for 'grid localized circular agent' */
@@ -47,9 +47,9 @@ class MASSAIMOVEMENT_API UMassAvoidanceObstacleRemoverFragmentDestructor : publi
 protected:
 	virtual void ConfigureQueries() override;
 	virtual void Initialize(UObject& Owner) override;
-	virtual void Execute(UMassEntitySubsystem& EntitySubsystem, FLWComponentSystemExecutionContext& Context) override;
+	virtual void Execute(UMassEntitySubsystem& EntitySubsystem, FMassExecutionContext& Context) override;
 
-	FLWComponentQuery EntityQuery;
+	FMassEntityQuery EntityQuery;
 };
 
 UCLASS()
@@ -61,7 +61,7 @@ public:
 
 protected:
 	virtual void ConfigureQueries() override;
-	virtual void Execute(UMassEntitySubsystem& EntitySubsystem, FLWComponentSystemExecutionContext& Context) override;
+	virtual void Execute(UMassEntitySubsystem& EntitySubsystem, FMassExecutionContext& Context) override;
 
-	FLWComponentQuery EntityQuery;
+	FMassEntityQuery EntityQuery;
 };
diff --git a/Engine/Plugins/AI/MassAI/Source/MassAIMovement/Public/NavigationProcessor.h b/Engine/Plugins/AI/MassAI/Source/MassAIMovement/Public/NavigationProcessor.h
--- a/Engine/Plugins/AI/MassAI/Source/MassAIMovement/Public/NavigationProcessor.h
+++ b/Engine/Plugins/AI/MassAI/Source/MassAIMovement/Public/NavigationProcessor.h
@@ -15,7 +15,7 @@ class UMassMovementSubsystem;
  * @todo: separate apply velocity and height adjustment
  */
 UCLASS()
-class MASSAIMOVEMENT_API UMassApplyVelocityMoveTargetProcessor : public UPipeProcessor
+class MASSAIMOVEMENT_API UMassApplyVelocityMoveTargetProcessor : public UMassProcessor
 {
 	GENERATED_BODY()
 
@@ -24,16 +24,16 @@ public:
 
 protected:
 	virtual void ConfigureQueries() override;
-	virtual void Execute(UMassEntitySubsystem& EntitySubsystem, FLWComponentSystemExecutionContext& Context) override;
+	virtual void Execute(UMassEntitySubsystem& EntitySubsystem, FMassExecutionContext& Context) override;
 
 private:
-	FLWComponentQuery HighResEntityQuery;
-	FLWComponentQuery LowResEntityQuery_Conditional;
+	FMassEntityQuery HighResEntityQuery;
+	FMassEntityQuery LowResEntityQuery_Conditional;
 };
 
 /** Handle dynamic obstacles. */
 UCLASS()
-class MASSAIMOVEMENT_API UMassDynamicObstacleProcessor : public UPipeProcessor
+class MASSAIMOVEMENT_API UMassDynamicObstacleProcessor : public UMassProcessor
 {
 	GENERATED_BODY()
 
@@ -50,10 +50,10 @@ public:
 	
 protected:
 	virtual void ConfigureQueries() override;
-	virtual void Execute(UMassEntitySubsystem& EntitySubsystem, FLWComponentSystemExecutionContext& Context) override;
+	virtual void Execute(UMassEntitySubsystem& EntitySubsystem, FMassExecutionContext& Context) override;
 	virtual void OnStop(FMassDynamicObstacleFragment& OutObstacle, const float BlockingRadius) {}
 	virtual void OnMove(FMassDynamicObstacleFragment& OutObstacle) {}
 
 private:
-	FLWComponentQuery EntityQuery_Conditional;
+	FMassEntityQuery EntityQuery_Conditional;
 };
diff --git a/Engine/Plugins/AI/MassAI/Source/MassAIMovement/Public/SnapToNavigationProcessor.h b/Engine/Plugins/AI/MassAI/Source/MassAIMovement/Public/SnapToNavigationProcessor.h
--- a/Engine/Plugins/AI/MassAI/Source/MassAIMovement/Public/SnapToNavigationProcessor.h
+++ b/Engine/Plugins/AI/MassAI/Source/MassAIMovement/Public/SnapToNavigationProcessor.h
@@ -8,7 +8,7 @@
 class ANavigationData;
 
 UCLASS()
-class MASSAIMOVEMENT_API USnapToNavigationProcessor : public UPipeProcessor
+class MASSAIMOVEMENT_API USnapToNavigationProcessor : public UMassProcessor
 {
 	GENERATED_BODY()
 public:
@@ -16,10 +16,10 @@ public:
 
 protected:	
 	virtual void ConfigureQueries() override;
-	virtual void Execute(UMassEntitySubsystem& EntitySubsystem, FLWComponentSystemExecutionContext& Context) override;
+	virtual void Execute(UMassEntitySubsystem& EntitySubsystem, FMassExecutionContext& Context) override;
 	virtual void Initialize(UObject& InOwner) override;
 
 	TWeakObjectPtr<ANavigationData> WeakNavData;
 
-	FLWComponentQuery EntityQuery;
+	FMassEntityQuery EntityQuery;
 };
diff --git a/Engine/Plugins/AI/MassAI/Source/MassAIReplication/Private/MassReplicationPathHandlers.cpp b/Engine/Plugins/AI/MassAI/Source/MassAIReplication/Private/MassReplicationPathHandlers.cpp
--- a/Engine/Plugins/AI/MassAI/Source/MassAIReplication/Private/MassReplicationPathHandlers.cpp
+++ b/Engine/Plugins/AI/MassAI/Source/MassAIReplication/Private/MassReplicationPathHandlers.cpp
@@ -11,14 +11,14 @@
 #include "GameFramework/GameStateBase.h"
 #include "VisualLogger/VisualLogger.h"
 
-void FMassReplicationProcessorPathHandler::AddRequirements(FLWComponentQuery& InQuery)
+void FMassReplicationProcessorPathHandler::AddRequirements(FMassEntityQuery& InQuery)
 {
-	InQuery.AddRequirement<FMassZoneGraphPathRequestFragment>(ELWComponentAccess::ReadOnly);
-	InQuery.AddRequirement<FMassMoveTargetFragment>(ELWComponentAccess::ReadOnly);
-	InQuery.AddRequirement<FMassZoneGraphLaneLocationFragment>(ELWComponentAccess::ReadOnly);
+	InQuery.AddRequirement<FMassZoneGraphPathRequestFragment>(EMassFragmentAccess::ReadOnly);
+	InQuery.AddRequirement<FMassMoveTargetFragment>(EMassFragmentAccess::ReadOnly);
+	InQuery.AddRequirement<FMassZoneGraphLaneLocationFragment>(EMassFragmentAccess::ReadOnly);
 }
 
-void FMassReplicationProcessorPathHandler::CacheComponentViews(FLWComponentSystemExecutionContext& ExecContext)
+void FMassReplicationProcessorPathHandler::CacheComponentViews(FMassExecutionContext& ExecContext)
 {
 	PathRequestList = ExecContext.GetMutableComponentView<FMassZoneGraphPathRequestFragment>();
 	MoveTargetList = ExecContext.GetMutableComponentView<FMassMoveTargetFragment>();
@@ -54,12 +54,12 @@ FReplicatedAgentPathData::FReplicatedAgentPathData(const FMassZoneGraphPathReque
 }
 
 void FReplicatedAgentPathData::InitEntity(const UWorld& InWorld,
-										  const FEntityView& InEntityView,
+										  const FMassEntityView& InEntityView,
 										  FMassZoneGraphLaneLocationFragment& OutLaneLocation,
 										  FMassMoveTargetFragment& OutMoveTarget,
 										  FMassZoneGraphPathRequestFragment& OutActionRequest) const
 {
-	const FLWEntity Entity = InEntityView.GetEntity();
+	const FMassEntityHandle Entity = InEntityView.GetEntity();
 
 	const UZoneGraphSubsystem* ZoneGraphSubsystem = InWorld.GetSubsystem<UZoneGraphSubsystem>();
 	const UMassSimulationSubsystem* SimulationSubsystem = InWorld.GetSubsystem<UMassSimulationSubsystem>();
@@ -94,9 +94,9 @@ void FReplicatedAgentPathData::InitEntity(const UWorld& InWorld,
 	ApplyToEntity(InWorld, InEntityView);
 }
 
-void FReplicatedAgentPathData::ApplyToEntity(const UWorld& InWorld, const FEntityView& InEntityView) const
+void FReplicatedAgentPathData::ApplyToEntity(const UWorld& InWorld, const FMassEntityView& InEntityView) const
 {
-	const FLWEntity Entity = InEntityView.GetEntity();
+	const FMassEntityHandle Entity = InEntityView.GetEntity();
 	FMassMoveTargetFragment& MoveTarget = InEntityView.GetComponentData<FMassMoveTargetFragment>();
 	if (MoveTarget.GetCurrentActionID() == ActionID)
 	{
diff --git a/Engine/Plugins/AI/MassAI/Source/MassAIReplication/Public/MassReplicationPathHandlers.h b/Engine/Plugins/AI/MassAI/Source/MassAIReplication/Public/MassReplicationPathHandlers.h
--- a/Engine/Plugins/AI/MassAI/Source/MassAIReplication/Public/MassReplicationPathHandlers.h
+++ b/Engine/Plugins/AI/MassAI/Source/MassAIReplication/Public/MassReplicationPathHandlers.h
@@ -10,7 +10,7 @@
 
 #include "MassReplicationPathHandlers.generated.h"
 
-struct FLWComponentQuery;
+struct FMassEntityQuery;
 
 //////////////////////////////////////////////////////////////////////////// FReplicatedAgentPathData ////////////////////////////////////////////////////////////////////////////
 /**
@@ -30,12 +30,12 @@ struct MASSAIREPLICATION_API FReplicatedAgentPathData
 		const FMassZoneGraphLaneLocationFragment& LaneLocationFragment);
 
 	void InitEntity(const UWorld& InWorld,
-					const FEntityView& InEntityView,
+					const FMassEntityView& InEntityView,
 					FMassZoneGraphLaneLocationFragment& OutLaneLocation,
 					FMassMoveTargetFragment& OutMoveTarget,
 					FMassZoneGraphPathRequestFragment& OutActionRequest) const;
 
-	void ApplyToEntity(const UWorld& InWorld, const FEntityView& InEntityView) const;
+	void ApplyToEntity(const UWorld& InWorld, const FMassEntityView& InEntityView) const;
 
 	UPROPERTY(Transient)
 	mutable FZoneGraphShortPathRequest PathRequest;
@@ -94,14 +94,14 @@ public:
 	 * When entities are spawned in Mass by the replication system on the client, a spawn query is used to set the data on the spawned entities.
 	 * The following functions are used to configure the query and then set that data for path following.
 	 */
-	static void AddRequirementsForSpawnQuery(FLWComponentQuery& InQuery);
-	void CacheComponentViewsForSpawnQuery(FLWComponentSystemExecutionContext& InExecContext);
+	static void AddRequirementsForSpawnQuery(FMassEntityQuery& InQuery);
+	void CacheComponentViewsForSpawnQuery(FMassExecutionContext& InExecContext);
 	void ClearComponentViewsForSpawnQuery();
 
-	void SetSpawnedEntityData(const FEntityView& EntityView, const FReplicatedAgentPathData& ReplicatedPathData, const int32 EntityIdx) const;
+	void SetSpawnedEntityData(const FMassEntityView& EntityView, const FReplicatedAgentPathData& ReplicatedPathData, const int32 EntityIdx) const;
 
 	/** Call this when an Entity that has already been spawned is modified on the client */
-	void SetModifiedEntityData(const FEntityView& EntityView, const FReplicatedAgentPathData& ReplicatedPathData) const;
+	void SetModifiedEntityData(const FMassEntityView& EntityView, const FReplicatedAgentPathData& ReplicatedPathData) const;
 #endif // UE_REPLICATION_COMPILE_CLIENT_CODE
 
 protected:
@@ -140,17 +140,17 @@ void TMassClientBubblePathHandler<AgentArrayItem>::SetBubblePathData(const FMass
 
 #if UE_REPLICATION_COMPILE_CLIENT_CODE
 template<typename AgentArrayItem>
-void TMassClientBubblePathHandler<AgentArrayItem>::AddRequirementsForSpawnQuery(FLWComponentQuery& InQuery)
+void TMassClientBubblePathHandler<AgentArrayItem>::AddRequirementsForSpawnQuery(FMassEntityQuery& InQuery)
 {
-	InQuery.AddRequirement<FMassZoneGraphPathRequestFragment>(ELWComponentAccess::ReadWrite);
-	InQuery.AddRequirement<FMassMoveTargetFragment>(ELWComponentAccess::ReadWrite);
-	InQuery.AddRequirement<FMassZoneGraphLaneLocationFragment>(ELWComponentAccess::ReadWrite);
+	InQuery.AddRequirement<FMassZoneGraphPathRequestFragment>(EMassFragmentAccess::ReadWrite);
+	InQuery.AddRequirement<FMassMoveTargetFragment>(EMassFragmentAccess::ReadWrite);
+	InQuery.AddRequirement<FMassZoneGraphLaneLocationFragment>(EMassFragmentAccess::ReadWrite);
 }
 #endif // UE_REPLICATION_COMPILE_CLIENT_CODE
 
 #if UE_REPLICATION_COMPILE_CLIENT_CODE
 template<typename AgentArrayItem>
-void TMassClientBubblePathHandler<AgentArrayItem>::CacheComponentViewsForSpawnQuery(FLWComponentSystemExecutionContext& InExecContext)
+void TMassClientBubblePathHandler<AgentArrayItem>::CacheComponentViewsForSpawnQuery(FMassExecutionContext& InExecContext)
 {
 	PathRequestList = InExecContext.GetMutableComponentView<FMassZoneGraphPathRequestFragment>();
 	MoveTargetList = InExecContext.GetMutableComponentView<FMassMoveTargetFragment>();
@@ -170,7 +170,7 @@ void TMassClientBubblePathHandler<AgentArrayItem>::ClearComponentViewsForSpawnQu
 
 #if UE_REPLICATION_COMPILE_CLIENT_CODE
 template<typename AgentArrayItem>
-void TMassClientBubblePathHandler<AgentArrayItem>::SetSpawnedEntityData(const FEntityView& EntityView, const FReplicatedAgentPathData& ReplicatedPathData, const int32 EntityIdx) const
+void TMassClientBubblePathHandler<AgentArrayItem>::SetSpawnedEntityData(const FMassEntityView& EntityView, const FReplicatedAgentPathData& ReplicatedPathData, const int32 EntityIdx) const
 {
 	UWorld* World = OwnerHandler.Serializer->GetWorld();
 	check(World);
@@ -180,7 +180,7 @@ void TMassClientBubblePathHandler<AgentArrayItem>::SetSpawnedEntityData(const FE
 
 #if UE_REPLICATION_COMPILE_CLIENT_CODE
 template<typename AgentArrayItem>
-void TMassClientBubblePathHandler<AgentArrayItem>::SetModifiedEntityData(const FEntityView& EntityView, const FReplicatedAgentPathData& ReplicatedPathData) const
+void TMassClientBubblePathHandler<AgentArrayItem>::SetModifiedEntityData(const FMassEntityView& EntityView, const FReplicatedAgentPathData& ReplicatedPathData) const
 {
 	UWorld* World = OwnerHandler.Serializer->GetWorld();
 	check(World);
@@ -197,10 +197,10 @@ class MASSAIREPLICATION_API FMassReplicationProcessorPathHandler
 {
 public:
 	/** Adds the requirements for the path following to the query. */
-	static void AddRequirements(FLWComponentQuery& InQuery);
+	static void AddRequirements(FMassEntityQuery& InQuery);
 
 	/** Cache any component views you want to, this will get called before we iterate through entities. */
-	void CacheComponentViews(FLWComponentSystemExecutionContext& ExecContext);
+	void CacheComponentViews(FMassExecutionContext& ExecContext);
 
 	/**
 	 * Set the replicated path data when we are adding an entity to the client bubble.
diff --git a/Engine/Plugins/AI/MassCrowd/Source/MassCrowd/Private/MassCrowdBubble.cpp b/Engine/Plugins/AI/MassCrowd/Source/MassCrowd/Private/MassCrowdBubble.cpp
--- a/Engine/Plugins/AI/MassCrowd/Source/MassCrowd/Private/MassCrowdBubble.cpp
+++ b/Engine/Plugins/AI/MassCrowd/Source/MassCrowd/Private/MassCrowdBubble.cpp
@@ -17,12 +17,12 @@ namespace UE::Mass::Crowd
 	// @todo provide a better way of selecting agents to debug
 	constexpr int32 MaxAgentsDraw = 300;
 
-	void DebugDrawReplicatedAgent(FLWEntity Entity, const UMassEntitySubsystem& EntitySystem)
+	void DebugDrawReplicatedAgent(FMassEntityHandle Entity, const UMassEntitySubsystem& EntitySystem)
 	{
 		static const FVector DebugCylinderHeight = FVector(0.f, 0.f, 200.f);
 		static constexpr float DebugCylinderRadius = 50.f;
 
-		const FEntityView EntityView(EntitySystem, Entity);
+		const FMassEntityView EntityView(EntitySystem, Entity);
 
 		const FDataFragment_Transform& TransformFragment = EntityView.GetComponentData<FDataFragment_Transform>();
 		const FMassNetworkIDFragment& NetworkIDFragment = EntityView.GetComponentData<FMassNetworkIDFragment>();
@@ -123,25 +123,25 @@ void FMassCrowdClientBubbleHandler::DebugValidateBubbleOnClient()
 #if UE_REPLICATION_COMPILE_CLIENT_CODE
 void FMassCrowdClientBubbleHandler::PostReplicatedAdd(const TArrayView<int32> AddedIndices, int32 FinalSize)
 {
-	auto AddRequirementsForSpawnQuery = [this](FLWComponentQuery& InQuery)
+	auto AddRequirementsForSpawnQuery = [this](FMassEntityQuery& InQuery)
 	{
 		PathHandler.AddRequirementsForSpawnQuery(InQuery);
 		TransformHandler.AddRequirementsForSpawnQuery(InQuery);
 	};
 
-	auto CacheComponentViewsForSpawnQuery = [this](FLWComponentSystemExecutionContext& InExecContext)
+	auto CacheComponentViewsForSpawnQuery = [this](FMassExecutionContext& InExecContext)
 	{
 		PathHandler.CacheComponentViewsForSpawnQuery(InExecContext);
 		TransformHandler.CacheComponentViewsForSpawnQuery(InExecContext);
 	};
 
-	auto SetSpawnedEntityData = [this](const FEntityView& EntityView, const FReplicatedCrowdAgent& ReplicatedEntity, const int32 EntityIdx)
+	auto SetSpawnedEntityData = [this](const FMassEntityView& EntityView, const FReplicatedCrowdAgent& ReplicatedEntity, const int32 EntityIdx)
 	{
 		PathHandler.SetSpawnedEntityData(EntityView, ReplicatedEntity.GetReplicatedPathData(), EntityIdx);
 		TransformHandler.SetSpawnedEntityData(EntityIdx, ReplicatedEntity.GetReplicatedPositionYawData());
 	};
 
-	auto SetModifiedEntityData = [this](const FEntityView& EntityView, const FReplicatedCrowdAgent& Item)
+	auto SetModifiedEntityData = [this](const FMassEntityView& EntityView, const FReplicatedCrowdAgent& Item)
 	{
 		PostReplicatedChangeEntity(EntityView, Item);
 	};
@@ -156,7 +156,7 @@ void FMassCrowdClientBubbleHandler::PostReplicatedAdd(const TArrayView<int32> Ad
 #if UE_REPLICATION_COMPILE_CLIENT_CODE
 void FMassCrowdClientBubbleHandler::FMassCrowdClientBubbleHandler::PostReplicatedChange(const TArrayView<int32> ChangedIndices, int32 FinalSize)
 {
-	auto SetModifiedEntityData = [this](const FEntityView& EntityView, const FReplicatedCrowdAgent& Item)
+	auto SetModifiedEntityData = [this](const FMassEntityView& EntityView, const FReplicatedCrowdAgent& Item)
 	{
 		PostReplicatedChangeEntity(EntityView, Item);
 	};
@@ -166,7 +166,7 @@ void FMassCrowdClientBubbleHandler::FMassCrowdClientBubbleHandler::PostReplicate
 #endif //UE_REPLICATION_COMPILE_SERVER_CODE
 
 #if UE_REPLICATION_COMPILE_CLIENT_CODE
-void FMassCrowdClientBubbleHandler::PostReplicatedChangeEntity(const FEntityView& EntityView, const FReplicatedCrowdAgent& Item) const
+void FMassCrowdClientBubbleHandler::PostReplicatedChangeEntity(const FMassEntityView& EntityView, const FReplicatedCrowdAgent& Item) const
 {
 	PathHandler.SetModifiedEntityData(EntityView, Item.GetReplicatedPathData());
 
diff --git a/Engine/Plugins/AI/MassCrowd/Source/MassCrowd/Private/MassCrowdNavigationProcessor.cpp b/Engine/Plugins/AI/MassCrowd/Source/MassCrowd/Private/MassCrowdNavigationProcessor.cpp
--- a/Engine/Plugins/AI/MassCrowd/Source/MassCrowd/Private/MassCrowdNavigationProcessor.cpp
+++ b/Engine/Plugins/AI/MassCrowd/Source/MassCrowd/Private/MassCrowdNavigationProcessor.cpp
@@ -23,9 +23,9 @@ UMassCrowdLaneTrackingSignalProcessor::UMassCrowdLaneTrackingSignalProcessor()
 
 void UMassCrowdLaneTrackingSignalProcessor::ConfigureQueries()
 {
-	EntityQuery.AddTagRequirement<FTagFragment_MassCrowd>(ELWComponentPresence::All);
-	EntityQuery.AddRequirement<FMassCrowdLaneTrackingFragment>(ELWComponentAccess::ReadWrite);
-	EntityQuery.AddRequirement<FMassZoneGraphLaneLocationFragment>(ELWComponentAccess::ReadOnly);
+	EntityQuery.AddTagRequirement<FTagFragment_MassCrowd>(EMassFragmentPresence::All);
+	EntityQuery.AddRequirement<FMassCrowdLaneTrackingFragment>(EMassFragmentAccess::ReadWrite);
+	EntityQuery.AddRequirement<FMassZoneGraphLaneLocationFragment>(EMassFragmentAccess::ReadOnly);
 }
 
 void UMassCrowdLaneTrackingSignalProcessor::Initialize(UObject& Owner)
@@ -38,14 +38,14 @@ void UMassCrowdLaneTrackingSignalProcessor::Initialize(UObject& Owner)
 	SubscribeToSignal(UE::Mass::Signals::CurrentLaneChanged);
 }
 
-void UMassCrowdLaneTrackingSignalProcessor::SignalEntities(UMassEntitySubsystem& EntitySubsystem, FLWComponentSystemExecutionContext& Context, FMassSignalNameLookup& EntitySignals)
+void UMassCrowdLaneTrackingSignalProcessor::SignalEntities(UMassEntitySubsystem& EntitySubsystem, FMassExecutionContext& Context, FMassSignalNameLookup& EntitySignals)
 {
 	if (!MassCrowdSubsystem)
 	{
 		return;
 	}
 
-	EntityQuery.ForEachEntityChunk(EntitySubsystem, Context, [this](FLWComponentSystemExecutionContext& Context)
+	EntityQuery.ForEachEntityChunk(EntitySubsystem, Context, [this](FMassExecutionContext& Context)
 	{
 		const int32 NumEntities = Context.GetEntitiesNum();
 		const TConstArrayView<FMassZoneGraphLaneLocationFragment> LaneLocationList = Context.GetComponentView<FMassZoneGraphLaneLocationFragment>();
@@ -82,13 +82,13 @@ void UMassCrowdLaneTrackingDestructor::Initialize(UObject& Owner)
 
 void UMassCrowdLaneTrackingDestructor::ConfigureQueries()
 {
-	EntityQuery.AddTagRequirement<FTagFragment_MassCrowd>(ELWComponentPresence::All);
-	EntityQuery.AddRequirement<FMassCrowdLaneTrackingFragment>(ELWComponentAccess::ReadOnly);
+	EntityQuery.AddTagRequirement<FTagFragment_MassCrowd>(EMassFragmentPresence::All);
+	EntityQuery.AddRequirement<FMassCrowdLaneTrackingFragment>(EMassFragmentAccess::ReadOnly);
 }
 
-void UMassCrowdLaneTrackingDestructor::Execute(UMassEntitySubsystem& EntitySubsystem, FLWComponentSystemExecutionContext& Context)
+void UMassCrowdLaneTrackingDestructor::Execute(UMassEntitySubsystem& EntitySubsystem, FMassExecutionContext& Context)
 {
-	EntityQuery.ForEachEntityChunk(EntitySubsystem, Context, [this](const FLWComponentSystemExecutionContext& Context)
+	EntityQuery.ForEachEntityChunk(EntitySubsystem, Context, [this](const FMassExecutionContext& Context)
 	{
 		const int32 NumEntities = Context.GetEntitiesNum();
 		const TConstArrayView<FMassCrowdLaneTrackingFragment> LaneTrackingList = Context.GetComponentView<FMassCrowdLaneTrackingFragment>();
diff --git a/Engine/Plugins/AI/MassCrowd/Source/MassCrowd/Private/MassCrowdReplicationProcessor.cpp b/Engine/Plugins/AI/MassCrowd/Source/MassCrowd/Private/MassCrowdReplicationProcessor.cpp
--- a/Engine/Plugins/AI/MassCrowd/Source/MassCrowd/Private/MassCrowdReplicationProcessor.cpp
+++ b/Engine/Plugins/AI/MassCrowd/Source/MassCrowd/Private/MassCrowdReplicationProcessor.cpp
@@ -36,7 +36,7 @@ void UMassCrowdReplicationProcessor::ConfigureQueries()
 	FMassReplicationProcessorPositionYawHandler::AddRequirements(EntityQuery);
 	FMassReplicationProcessorPathHandler::AddRequirements(EntityQuery);
 
-	EntityQuery.AddTagRequirement<FTagFragment_MassCrowd>(ELWComponentPresence::All);
+	EntityQuery.AddTagRequirement<FTagFragment_MassCrowd>(EMassFragmentPresence::All);
 }
 
 void UMassCrowdReplicationProcessor::Initialize(UObject& Owner)
@@ -48,7 +48,7 @@ void UMassCrowdReplicationProcessor::Initialize(UObject& Owner)
 	BubbleInfoClassHandle = ReplicationManager->GetBubbleInfoClassHandle(AMassCrowdClientBubbleInfo::StaticClass());
 }
 
-void UMassCrowdReplicationProcessor::Execute(UMassEntitySubsystem& EntitySubsystem, FLWComponentSystemExecutionContext& Context)
+void UMassCrowdReplicationProcessor::Execute(UMassEntitySubsystem& EntitySubsystem, FMassExecutionContext& Context)
 {
 	Super::Execute(EntitySubsystem, Context);
 
@@ -58,7 +58,7 @@ void UMassCrowdReplicationProcessor::Execute(UMassEntitySubsystem& EntitySubsyst
 	// Optional debug display
 	if (UE::Mass::Crowd::bDebugReplicationViewerLOD)
 	{
-		EntityQuery.ForEachEntityChunk(EntitySubsystem, Context, [this](FLWComponentSystemExecutionContext& Context)
+		EntityQuery.ForEachEntityChunk(EntitySubsystem, Context, [this](FMassExecutionContext& Context)
 			{
 				const TConstArrayView<FDataFragment_Transform> TransformList = Context.GetComponentView<FDataFragment_Transform>();
 				const TConstArrayView<FMassReplicationViewerLODFragment> ViewerLODList = Context.GetComponentView<FMassReplicationViewerLODFragment>();
@@ -68,7 +68,7 @@ void UMassCrowdReplicationProcessor::Execute(UMassEntitySubsystem& EntitySubsyst
 #endif // WITH_MASSGAMEPLAY_DEBUG
 }
 
-void UMassCrowdReplicationProcessor::ProcessClientReplication(UMassEntitySubsystem& EntitySubsystem, FLWComponentSystemExecutionContext& Context)
+void UMassCrowdReplicationProcessor::ProcessClientReplication(UMassEntitySubsystem& EntitySubsystem, FMassExecutionContext& Context)
 {
 #if UE_REPLICATION_COMPILE_SERVER_CODE
 
@@ -111,7 +111,7 @@ void UMassCrowdReplicationProcessor::ProcessClientReplication(UMassEntitySubsyst
 
 	QUICK_SCOPE_CYCLE_COUNTER(UMassCrowdReplicationProcessor_ProcessClientReplication);
 
-	EntityQuery.ForEachEntityChunk(EntitySubsystem, Context, [this, &CacheViewsCallback, &AddEntityCallback, &ModifyEntityCallback, &RemoveEntityCallback](FLWComponentSystemExecutionContext& Context)
+	EntityQuery.ForEachEntityChunk(EntitySubsystem, Context, [this, &CacheViewsCallback, &AddEntityCallback, &ModifyEntityCallback, &RemoveEntityCallback](FMassExecutionContext& Context)
 	{
 		CalculateClientReplication<FCrowdFastArrayItem>(Context, CacheViewsCallback, AddEntityCallback, ModifyEntityCallback, RemoveEntityCallback);
 	});
diff --git a/Engine/Plugins/AI/MassCrowd/Source/MassCrowd/Private/MassCrowdRepresentationProcessor.cpp b/Engine/Plugins/AI/MassCrowd/Source/MassCrowd/Private/MassCrowdRepresentationProcessor.cpp
--- a/Engine/Plugins/AI/MassCrowd/Source/MassCrowd/Private/MassCrowdRepresentationProcessor.cpp
+++ b/Engine/Plugins/AI/MassCrowd/Source/MassCrowd/Private/MassCrowdRepresentationProcessor.cpp
@@ -30,17 +30,17 @@ void UMassCrowdRepresentationProcessor::Initialize(UObject& Owner)
 void UMassCrowdRepresentationProcessor::ConfigureQueries()
 {
 	Super::ConfigureQueries();
-	EntityQuery.AddTagRequirement<FTagFragment_MassCrowd>(ELWComponentPresence::All);
+	EntityQuery.AddTagRequirement<FTagFragment_MassCrowd>(EMassFragmentPresence::All);
 
 	CharacterMovementEntitiesQuery_Conditional = EntityQuery;
-	CharacterMovementEntitiesQuery_Conditional.AddRequirement<FDataFragment_CharacterMovementComponentWrapper>(ELWComponentAccess::ReadWrite);
-	CharacterMovementEntitiesQuery_Conditional.AddRequirement<FMassVelocityFragment>(ELWComponentAccess::ReadOnly);
+	CharacterMovementEntitiesQuery_Conditional.AddRequirement<FDataFragment_CharacterMovementComponentWrapper>(EMassFragmentAccess::ReadWrite);
+	CharacterMovementEntitiesQuery_Conditional.AddRequirement<FMassVelocityFragment>(EMassFragmentAccess::ReadOnly);
 	CharacterMovementEntitiesQuery_Conditional.SetChunkFilter(&FMassVisualizationChunkFragment::AreAnyEntitiesVisibleInChunk);
 }
 
-void UMassCrowdRepresentationProcessor::InitializeVelocity(UMassEntitySubsystem& EntitySubsystem, FLWComponentSystemExecutionContext& Context)
+void UMassCrowdRepresentationProcessor::InitializeVelocity(UMassEntitySubsystem& EntitySubsystem, FMassExecutionContext& Context)
 {
-	CharacterMovementEntitiesQuery_Conditional.ForEachEntityChunk(EntitySubsystem, Context, [this](FLWComponentSystemExecutionContext& Context)
+	CharacterMovementEntitiesQuery_Conditional.ForEachEntityChunk(EntitySubsystem, Context, [this](FMassExecutionContext& Context)
 		{
 			const TConstArrayView<FMassRepresentationFragment> VisualizationList = Context.GetComponentView<FMassRepresentationFragment>();
 			const TConstArrayView<FMassVelocityFragment> VelocityList = Context.GetComponentView<FMassVelocityFragment>();
@@ -69,7 +69,7 @@ void UMassCrowdRepresentationProcessor::InitializeVelocity(UMassEntitySubsystem&
 		});
 }
 
-void UMassCrowdRepresentationProcessor::SetActorEnabled(const EActorEnabledType EnabledType, AActor& Actor, const int32 EntityIdx, FLWComponentSystemExecutionContext& Context)
+void UMassCrowdRepresentationProcessor::SetActorEnabled(const EActorEnabledType EnabledType, AActor& Actor, const int32 EntityIdx, FMassExecutionContext& Context)
 {
 	Super::SetActorEnabled(EnabledType, Actor, EntityIdx, Context);
 
@@ -140,7 +140,7 @@ AActor* UMassCrowdRepresentationProcessor::GetOrSpawnActor(const FMassHandle Mas
 	return Super::GetOrSpawnActor(MassAgent, ActorInfo, RootTransform, TemplateActorIndex, SpawnRequestHandle, Priority);
 }
 
-void UMassCrowdRepresentationProcessor::TeleportActor(const FTransform& Transform, AActor& Actor, FLWComponentSystemExecutionContext& Context)
+void UMassCrowdRepresentationProcessor::TeleportActor(const FTransform& Transform, AActor& Actor, FMassExecutionContext& Context)
 {
 	FTransform RootTransform = Transform;
 
diff --git a/Engine/Plugins/AI/MassCrowd/Source/MassCrowd/Private/MassCrowdServerRepresentationLODProcessor.cpp b/Engine/Plugins/AI/MassCrowd/Source/MassCrowd/Private/MassCrowdServerRepresentationLODProcessor.cpp
--- a/Engine/Plugins/AI/MassCrowd/Source/MassCrowd/Private/MassCrowdServerRepresentationLODProcessor.cpp
+++ b/Engine/Plugins/AI/MassCrowd/Source/MassCrowd/Private/MassCrowdServerRepresentationLODProcessor.cpp
@@ -30,10 +30,10 @@ UMassCrowdServerRepresentationLODProcessor::UMassCrowdServerRepresentationLODPro
 
 void UMassCrowdServerRepresentationLODProcessor::ConfigureQueries()
 {
-	EntityQuery.AddTagRequirement<FTagFragment_MassCrowd>(ELWComponentPresence::All);
-	EntityQuery.AddRequirement<FDataFragment_Transform>(ELWComponentAccess::ReadOnly);
-	EntityQuery.AddRequirement<FMassLODInfoFragment>(ELWComponentAccess::ReadOnly);
-	EntityQuery.AddRequirement<FMassRepresentationLODFragment>(ELWComponentAccess::ReadWrite);
+	EntityQuery.AddTagRequirement<FTagFragment_MassCrowd>(EMassFragmentPresence::All);
+	EntityQuery.AddRequirement<FDataFragment_Transform>(EMassFragmentAccess::ReadOnly);
+	EntityQuery.AddRequirement<FMassLODInfoFragment>(EMassFragmentAccess::ReadOnly);
+	EntityQuery.AddRequirement<FMassRepresentationLODFragment>(EMassFragmentAccess::ReadWrite);
 }
 
 void UMassCrowdServerRepresentationLODProcessor::Initialize(UObject& InOwner)
@@ -43,7 +43,7 @@ void UMassCrowdServerRepresentationLODProcessor::Initialize(UObject& InOwner)
 	Super::Initialize(InOwner);
 }
 
-void UMassCrowdServerRepresentationLODProcessor::Execute(UMassEntitySubsystem& EntitySubsystem, FLWComponentSystemExecutionContext& Context)
+void UMassCrowdServerRepresentationLODProcessor::Execute(UMassEntitySubsystem& EntitySubsystem, FMassExecutionContext& Context)
 {
 	TRACE_CPUPROFILER_EVENT_SCOPE(TEXT("CrowdServerRepresentationLOD"))
 
@@ -54,7 +54,7 @@ void UMassCrowdServerRepresentationLODProcessor::Execute(UMassEntitySubsystem& E
 	{
 		TRACE_CPUPROFILER_EVENT_SCOPE(TEXT("CalculateLOD"))
 		
-		EntityQuery.ForEachEntityChunk(EntitySubsystem, Context, [this](FLWComponentSystemExecutionContext& Context)
+		EntityQuery.ForEachEntityChunk(EntitySubsystem, Context, [this](FMassExecutionContext& Context)
 		{
 			const TConstArrayView<FMassLODInfoFragment> ViewersInfoList = Context.GetComponentView<FMassLODInfoFragment>();
 			const TArrayView<FMassRepresentationLODFragment> RepresentationLODFragments = Context.GetMutableComponentView<FMassRepresentationLODFragment>();
@@ -67,7 +67,7 @@ void UMassCrowdServerRepresentationLODProcessor::Execute(UMassEntitySubsystem& E
 		
 		if (LODCalculator.AdjustDistancesFromCount())
 		{
-			EntityQuery.ForEachEntityChunk(EntitySubsystem, Context, [this](FLWComponentSystemExecutionContext& Context)
+			EntityQuery.ForEachEntityChunk(EntitySubsystem, Context, [this](FMassExecutionContext& Context)
 			{
 				const TConstArrayView<FMassLODInfoFragment> ViewersInfoList = Context.GetComponentView<FMassLODInfoFragment>();
 				const TArrayView<FMassRepresentationLODFragment> RepresentationLODFragments = Context.GetMutableComponentView<FMassRepresentationLODFragment>();
@@ -81,7 +81,7 @@ void UMassCrowdServerRepresentationLODProcessor::Execute(UMassEntitySubsystem& E
 	{
 		TRACE_CPUPROFILER_EVENT_SCOPE(TEXT("DebugDisplayLOD"))
 		
-		EntityQuery.ForEachEntityChunk(EntitySubsystem, Context, [this](FLWComponentSystemExecutionContext& Context)
+		EntityQuery.ForEachEntityChunk(EntitySubsystem, Context, [this](FMassExecutionContext& Context)
 		{
 			const TConstArrayView<FDataFragment_Transform> LocationList = Context.GetComponentView<FDataFragment_Transform>();
 			const TConstArrayView<FMassRepresentationLODFragment> RepresentationLODFragments = Context.GetComponentView<FMassRepresentationLODFragment>();
diff --git a/Engine/Plugins/AI/MassCrowd/Source/MassCrowd/Private/MassCrowdServerRepresentationProcessor.cpp b/Engine/Plugins/AI/MassCrowd/Source/MassCrowd/Private/MassCrowdServerRepresentationProcessor.cpp
--- a/Engine/Plugins/AI/MassCrowd/Source/MassCrowd/Private/MassCrowdServerRepresentationProcessor.cpp
+++ b/Engine/Plugins/AI/MassCrowd/Source/MassCrowd/Private/MassCrowdServerRepresentationProcessor.cpp
@@ -16,9 +16,9 @@ UMassCrowdServerRepresentationProcessor::UMassCrowdServerRepresentationProcessor
 	LODRepresentation[EMassLOD::Off] = ERepresentationType::None;
 }
 
-void UMassCrowdServerRepresentationProcessor::Execute(UMassEntitySubsystem& EntitySubsystem, FLWComponentSystemExecutionContext& Context)
+void UMassCrowdServerRepresentationProcessor::Execute(UMassEntitySubsystem& EntitySubsystem, FMassExecutionContext& Context)
 {
-	EntityQuery.ForEachEntityChunk(EntitySubsystem, Context, [this](FLWComponentSystemExecutionContext& Context)
+	EntityQuery.ForEachEntityChunk(EntitySubsystem, Context, [this](FMassExecutionContext& Context)
 	{
 		UpdateRepresentation(Context);
 	});
diff --git a/Engine/Plugins/AI/MassCrowd/Source/MassCrowd/Private/MassCrowdSubsystem.cpp b/Engine/Plugins/AI/MassCrowd/Source/MassCrowd/Private/MassCrowdSubsystem.cpp
--- a/Engine/Plugins/AI/MassCrowd/Source/MassCrowd/Private/MassCrowdSubsystem.cpp
+++ b/Engine/Plugins/AI/MassCrowd/Source/MassCrowd/Private/MassCrowdSubsystem.cpp
@@ -364,7 +364,7 @@ const FRegisteredCrowdLaneData* UMassCrowdSubsystem::GetCrowdData(const FZoneGra
 	return &LanesData;
 }
 
-void UMassCrowdSubsystem::OnEntityLaneChanged(const FLWEntity& Entity, const FZoneGraphLaneHandle PreviousLaneHandle, const FZoneGraphLaneHandle CurrentLaneHandle)
+void UMassCrowdSubsystem::OnEntityLaneChanged(const FMassEntityHandle& Entity, const FZoneGraphLaneHandle PreviousLaneHandle, const FZoneGraphLaneHandle CurrentLaneHandle)
 {
 	const bool bPreviousLocationValid = PreviousLaneHandle.IsValid();
 	const bool bCurrentLocationValid = CurrentLaneHandle.IsValid();
@@ -544,14 +544,14 @@ void UMassCrowdSubsystem::UpdateDensityMask()
 	}
 }
 
-void UMassCrowdSubsystem::OnEnterTrackedLane(const FLWEntity& Entity, const int32 LaneIndex, FCrowdTrackingLaneData& TrackingData)
+void UMassCrowdSubsystem::OnEnterTrackedLane(const FMassEntityHandle& Entity, const int32 LaneIndex, FCrowdTrackingLaneData& TrackingData)
 {
 	++TrackingData.NumEntitiesOnLane;
 
 	UE_VLOG_UELOG(this, LogMassNavigation, Verbose, TEXT("[%s] enters lane %d. Num entities on lane: %d"), *Entity.DebugGetDescription(), LaneIndex, TrackingData.NumEntitiesOnLane);
 }
 
-void UMassCrowdSubsystem::OnExitTrackedLane(const FLWEntity& Entity, const int32 LaneIndex, FCrowdTrackingLaneData& TrackingData)
+void UMassCrowdSubsystem::OnExitTrackedLane(const FMassEntityHandle& Entity, const int32 LaneIndex, FCrowdTrackingLaneData& TrackingData)
 {
 	ensureMsgf(TrackingData.NumEntitiesOnLane >= 1, TEXT("OnExitTrackedLane should not be called more often than OnEnterTrackedLane"));
 	--TrackingData.NumEntitiesOnLane;
@@ -559,7 +559,7 @@ void UMassCrowdSubsystem::OnExitTrackedLane(const FLWEntity& Entity, const int32
 	UE_VLOG_UELOG(this, LogMassNavigation, Verbose, TEXT("[%s] exits lane %d. Num entities on lane: %d"), *Entity.DebugGetDescription(), LaneIndex, TrackingData.NumEntitiesOnLane);
 }
 
-int32 UMassCrowdSubsystem::AcquireWaitingSlot(const FLWEntity Entity, const FVector& EntityPosition, const FZoneGraphLaneHandle LaneHandle,
+int32 UMassCrowdSubsystem::AcquireWaitingSlot(const FMassEntityHandle Entity, const FVector& EntityPosition, const FZoneGraphLaneHandle LaneHandle,
 											  FVector& OutSlotPosition, FVector& OutSlotDirection)
 {
 	if (!LaneHandle.IsValid())
@@ -635,7 +635,7 @@ int32 UMassCrowdSubsystem::AcquireWaitingSlot(const FLWEntity Entity, const FVec
 	return BestSlotIndex;
 }
 
-void UMassCrowdSubsystem::ReleaseWaitingSlot(const FLWEntity Entity, const FZoneGraphLaneHandle LaneHandle, const int32 SlotIndex)
+void UMassCrowdSubsystem::ReleaseWaitingSlot(const FMassEntityHandle Entity, const FZoneGraphLaneHandle LaneHandle, const int32 SlotIndex)
 {
 	if (!LaneHandle.IsValid())
 	{
diff --git a/Engine/Plugins/AI/MassCrowd/Source/MassCrowd/Private/MassCrowdVisualizationLODProcessor.cpp b/Engine/Plugins/AI/MassCrowd/Source/MassCrowd/Private/MassCrowdVisualizationLODProcessor.cpp
--- a/Engine/Plugins/AI/MassCrowd/Source/MassCrowd/Private/MassCrowdVisualizationLODProcessor.cpp
+++ b/Engine/Plugins/AI/MassCrowd/Source/MassCrowd/Private/MassCrowdVisualizationLODProcessor.cpp
@@ -56,19 +56,19 @@ void UMassCrowdVisualizationLODProcessor::ConfigureQueries()
 	// @todo remove, no need for this anymore since we have the common LOD collector
 	// Do not call super as we do have our own LODInfo fragment, so need to duplicate
 
-	CloseEntityQuery.AddTagRequirement<FTagFragment_MassCrowd>(ELWComponentPresence::All);
-	CloseEntityQuery.AddRequirement<FMassLODInfoFragment>(ELWComponentAccess::ReadOnly);
-	CloseEntityQuery.AddRequirement<FMassRepresentationFragment>(ELWComponentAccess::ReadOnly);
-	CloseEntityQuery.AddRequirement<FMassRepresentationLODFragment>(ELWComponentAccess::ReadWrite);
-	CloseEntityQuery.AddRequirement<FDataFragment_Transform>(ELWComponentAccess::ReadOnly);
+	CloseEntityQuery.AddTagRequirement<FTagFragment_MassCrowd>(EMassFragmentPresence::All);
+	CloseEntityQuery.AddRequirement<FMassLODInfoFragment>(EMassFragmentAccess::ReadOnly);
+	CloseEntityQuery.AddRequirement<FMassRepresentationFragment>(EMassFragmentAccess::ReadOnly);
+	CloseEntityQuery.AddRequirement<FMassRepresentationLODFragment>(EMassFragmentAccess::ReadWrite);
+	CloseEntityQuery.AddRequirement<FDataFragment_Transform>(EMassFragmentAccess::ReadOnly);
 
 	FarEntityQuery = CloseEntityQuery;
-	CloseEntityQuery.AddTagRequirement<FMassVisibilityCulledByDistanceTag>(ELWComponentPresence::None);
-	FarEntityQuery.AddTagRequirement<FMassVisibilityCulledByDistanceTag>(ELWComponentPresence::All);
-	FarEntityQuery.AddChunkRequirement<FMassVisualizationChunkFragment>(ELWComponentAccess::ReadOnly);
+	CloseEntityQuery.AddTagRequirement<FMassVisibilityCulledByDistanceTag>(EMassFragmentPresence::None);
+	FarEntityQuery.AddTagRequirement<FMassVisibilityCulledByDistanceTag>(EMassFragmentPresence::All);
+	FarEntityQuery.AddChunkRequirement<FMassVisualizationChunkFragment>(EMassFragmentAccess::ReadOnly);
 }
 
-void UMassCrowdVisualizationLODProcessor::Execute(UMassEntitySubsystem& EntitySubsystem, FLWComponentSystemExecutionContext& Context)
+void UMassCrowdVisualizationLODProcessor::Execute(UMassEntitySubsystem& EntitySubsystem, FMassExecutionContext& Context)
 {
 	ForceOffLOD((bool)UE::MassCrowd::GCrowdTurnOffVisualization);
 
@@ -80,7 +80,7 @@ void UMassCrowdVisualizationLODProcessor::Execute(UMassEntitySubsystem& EntitySu
 	{
 		TRACE_CPUPROFILER_EVENT_SCOPE(TEXT("DebugDisplayLOD"))
 
-		auto DebugDisplayLOD = [this](FLWComponentSystemExecutionContext& Context)
+		auto DebugDisplayLOD = [this](FMassExecutionContext& Context)
 		{
 			const TConstArrayView<FDataFragment_Transform> LocationList = Context.GetComponentView<FDataFragment_Transform>();
 			const TConstArrayView<FMassRepresentationLODFragment> VisualizationLODList = Context.GetComponentView<FMassRepresentationLODFragment>();
@@ -95,7 +95,7 @@ void UMassCrowdVisualizationLODProcessor::Execute(UMassEntitySubsystem& EntitySu
 	{
 		TRACE_CPUPROFILER_EVENT_SCOPE(TEXT("ShowISMUnderSpecifiedRange"))
 
-		auto ShowISMUnderSpecifiedRange = [this](const FLWComponentSystemExecutionContext& Context)
+		auto ShowISMUnderSpecifiedRange = [this](const FMassExecutionContext& Context)
 		{
 			const TConstArrayView<FDataFragment_Transform> LocationList = Context.GetComponentView<FDataFragment_Transform>();
 			const TConstArrayView<FMassRepresentationFragment> RepresentationFragmentList = Context.GetComponentView<FMassRepresentationFragment>();
diff --git a/Engine/Plugins/AI/MassCrowd/Source/MassCrowd/Private/MassCrowdVisualizationProcessor.cpp b/Engine/Plugins/AI/MassCrowd/Source/MassCrowd/Private/MassCrowdVisualizationProcessor.cpp
--- a/Engine/Plugins/AI/MassCrowd/Source/MassCrowd/Private/MassCrowdVisualizationProcessor.cpp
+++ b/Engine/Plugins/AI/MassCrowd/Source/MassCrowd/Private/MassCrowdVisualizationProcessor.cpp
@@ -35,9 +35,9 @@ void UMassCrowdVisualizationProcessor::ConfigureQueries()
 	Super::ConfigureQueries();
 }
 
-void UMassCrowdVisualizationProcessor::Execute(UMassEntitySubsystem& EntitySubsystem, FLWComponentSystemExecutionContext& Context)
+void UMassCrowdVisualizationProcessor::Execute(UMassEntitySubsystem& EntitySubsystem, FMassExecutionContext& Context)
 {
-	EntityQuery.ForEachEntityChunk(EntitySubsystem, Context, [this](FLWComponentSystemExecutionContext& Context)
+	EntityQuery.ForEachEntityChunk(EntitySubsystem, Context, [this](FMassExecutionContext& Context)
 	{
 		UpdateVisualization(Context);
 	});
@@ -51,7 +51,7 @@ void UMassCrowdVisualizationProcessor::Execute(UMassEntitySubsystem& EntitySubsy
 
 		TRACE_CPUPROFILER_EVENT_SCOPE(TEXT("DebugDisplayVisualType"))
 
-		EntityQuery.ForEachEntityChunk(EntitySubsystem, Context, [this](const FLWComponentSystemExecutionContext& Context)
+		EntityQuery.ForEachEntityChunk(EntitySubsystem, Context, [this](const FMassExecutionContext& Context)
 		{
 			const TConstArrayView<FMassRepresentationFragment> VisualizationList = Context.GetComponentView<FMassRepresentationFragment>();
 			const TConstArrayView<FDataFragment_Actor> ActorList = Context.GetComponentView<FDataFragment_Actor>();
diff --git a/Engine/Plugins/AI/MassCrowd/Source/MassCrowd/Private/Tasks/MassCrowdClaimWaitSlotTask.cpp b/Engine/Plugins/AI/MassCrowd/Source/MassCrowd/Private/Tasks/MassCrowdClaimWaitSlotTask.cpp
--- a/Engine/Plugins/AI/MassCrowd/Source/MassCrowd/Private/Tasks/MassCrowdClaimWaitSlotTask.cpp
+++ b/Engine/Plugins/AI/MassCrowd/Source/MassCrowd/Private/Tasks/MassCrowdClaimWaitSlotTask.cpp
@@ -17,7 +17,7 @@ EStateTreeRunStatus FMassCrowdClaimWaitSlotTask::EnterState(FStateTreeExecutionC
 	WaitSlotLocation = nullptr;
 
 	const FMassStateTreeExecutionContext& MassContext = static_cast<FMassStateTreeExecutionContext&>(Context);
-	const FLWEntity Entity = MassContext.GetEntity();
+	const FMassEntityHandle Entity = MassContext.GetEntity();
 	
 	const FMassZoneGraphLaneLocationFragment& LaneLocation = Context.GetExternalItem(LocationHandle).Get<FMassZoneGraphLaneLocationFragment>();
 	const FMassMoveTargetFragment& MoveTarget = Context.GetExternalItem(MoveTargetHandle).Get<FMassMoveTargetFragment>();
@@ -61,7 +61,7 @@ void FMassCrowdClaimWaitSlotTask::ExitState(FStateTreeExecutionContext& Context,
 	}
 
 	const FMassStateTreeExecutionContext& MassContext = static_cast<FMassStateTreeExecutionContext&>(Context);
-	const FLWEntity Entity = MassContext.GetEntity();
+	const FMassEntityHandle Entity = MassContext.GetEntity();
 	
 	UMassCrowdSubsystem& CrowdSubsystem = Context.GetExternalItem(CrowdSubsystemHandle).GetMutable<UMassCrowdSubsystem>();
 
diff --git a/Engine/Plugins/AI/MassCrowd/Source/MassCrowd/Public/MassCrowdBubble.h b/Engine/Plugins/AI/MassCrowd/Source/MassCrowd/Public/MassCrowdBubble.h
--- a/Engine/Plugins/AI/MassCrowd/Source/MassCrowd/Public/MassCrowdBubble.h
+++ b/Engine/Plugins/AI/MassCrowd/Source/MassCrowd/Public/MassCrowdBubble.h
@@ -39,7 +39,7 @@ protected:
 	virtual void PostReplicatedAdd(const TArrayView<int32> AddedIndices, int32 FinalSize) override;
 	virtual void PostReplicatedChange(const TArrayView<int32> ChangedIndices, int32 FinalSize) override;
 
-	void PostReplicatedChangeEntity(const FEntityView& EntityView, const FReplicatedCrowdAgent& Item) const;
+	void PostReplicatedChangeEntity(const FMassEntityView& EntityView, const FReplicatedCrowdAgent& Item) const;
 #endif //UE_REPLICATION_COMPILE_CLIENT_CODE
 
 #if WITH_MASSGAMEPLAY_DEBUG && WITH_EDITOR
diff --git a/Engine/Plugins/AI/MassCrowd/Source/MassCrowd/Public/MassCrowdFragments.h b/Engine/Plugins/AI/MassCrowd/Source/MassCrowd/Public/MassCrowdFragments.h
--- a/Engine/Plugins/AI/MassCrowd/Source/MassCrowd/Public/MassCrowdFragments.h
+++ b/Engine/Plugins/AI/MassCrowd/Source/MassCrowd/Public/MassCrowdFragments.h
@@ -14,7 +14,7 @@
  * Should not contain any data, this is purely a tag
  */
 USTRUCT()
-struct MASSCROWD_API FTagFragment_MassCrowd : public FComponentTag
+struct MASSCROWD_API FTagFragment_MassCrowd : public FMassTag
 {
 	GENERATED_BODY()
 };
@@ -23,7 +23,7 @@ struct MASSCROWD_API FTagFragment_MassCrowd : public FComponentTag
  * Data fragment to store the last lane the agent was tracked on.
  */
 USTRUCT()
-struct MASSCROWD_API FMassCrowdLaneTrackingFragment : public FLWComponentData
+struct MASSCROWD_API FMassCrowdLaneTrackingFragment : public FMassFragment
 {
 	GENERATED_BODY()
 	FZoneGraphLaneHandle TrackedLaneHandle;
diff --git a/Engine/Plugins/AI/MassCrowd/Source/MassCrowd/Public/MassCrowdNavigationProcessor.h b/Engine/Plugins/AI/MassCrowd/Source/MassCrowd/Public/MassCrowdNavigationProcessor.h
--- a/Engine/Plugins/AI/MassCrowd/Source/MassCrowd/Public/MassCrowdNavigationProcessor.h
+++ b/Engine/Plugins/AI/MassCrowd/Source/MassCrowd/Public/MassCrowdNavigationProcessor.h
@@ -23,7 +23,7 @@ public:
 protected:
 	virtual void Initialize(UObject& Owner) override;
 	virtual void ConfigureQueries() override;
-	virtual void SignalEntities(UMassEntitySubsystem& EntitySubsystem, FLWComponentSystemExecutionContext& Context, FMassSignalNameLookup& EntitySignals) override;
+	virtual void SignalEntities(UMassEntitySubsystem& EntitySubsystem, FMassExecutionContext& Context, FMassSignalNameLookup& EntitySignals) override;
 
 	UPROPERTY(Transient)
 	UMassCrowdSubsystem* MassCrowdSubsystem;
@@ -41,12 +41,12 @@ public:
 protected:
 	virtual void Initialize(UObject& Owner) override;
 	virtual void ConfigureQueries() override;
-	virtual void Execute(UMassEntitySubsystem& EntitySubsystem, FLWComponentSystemExecutionContext& Context) override;
+	virtual void Execute(UMassEntitySubsystem& EntitySubsystem, FMassExecutionContext& Context) override;
 
 	UPROPERTY(Transient)
 	UMassCrowdSubsystem* MassCrowdSubsystem;
 
-	FLWComponentQuery EntityQuery;
+	FMassEntityQuery EntityQuery;
 };
 
 UCLASS()
diff --git a/Engine/Plugins/AI/MassCrowd/Source/MassCrowd/Public/MassCrowdReplicationProcessor.h b/Engine/Plugins/AI/MassCrowd/Source/MassCrowd/Public/MassCrowdReplicationProcessor.h
--- a/Engine/Plugins/AI/MassCrowd/Source/MassCrowd/Public/MassCrowdReplicationProcessor.h
+++ b/Engine/Plugins/AI/MassCrowd/Source/MassCrowd/Public/MassCrowdReplicationProcessor.h
@@ -22,7 +22,7 @@ public:
 protected:
 	virtual void ConfigureQueries() override;
 	virtual void Initialize(UObject& Owner) override;
-	virtual void Execute(UMassEntitySubsystem& EntitySubsystem, FLWComponentSystemExecutionContext& Context) override;
+	virtual void Execute(UMassEntitySubsystem& EntitySubsystem, FMassExecutionContext& Context) override;
 
-	void ProcessClientReplication(UMassEntitySubsystem& EntitySubsystem, FLWComponentSystemExecutionContext& Context);
+	void ProcessClientReplication(UMassEntitySubsystem& EntitySubsystem, FMassExecutionContext& Context);
 };
\ No newline at end of file
diff --git a/Engine/Plugins/AI/MassCrowd/Source/MassCrowd/Public/MassCrowdRepresentationProcessor.h b/Engine/Plugins/AI/MassCrowd/Source/MassCrowd/Public/MassCrowdRepresentationProcessor.h
--- a/Engine/Plugins/AI/MassCrowd/Source/MassCrowd/Public/MassCrowdRepresentationProcessor.h
+++ b/Engine/Plugins/AI/MassCrowd/Source/MassCrowd/Public/MassCrowdRepresentationProcessor.h
@@ -26,7 +26,7 @@ protected:
 	 */
 	virtual void Initialize(UObject& Owner) override;
 
-	/** Configure the owned FLWComponentQuery instances to express processor's requirements */
+	/** Configure the owned FMassEntityQuery instances to express processor's requirements */
 	virtual void ConfigureQueries() override;
 
 	/**
@@ -36,7 +36,7 @@ protected:
 	 * @param EntityIdx is the entity index currently processing
 	 * @param Context is the current light weight component execution context
 	 */
-	virtual void SetActorEnabled(const EActorEnabledType EnabledType, AActor& Actor, const int32 EntityIdx, FLWComponentSystemExecutionContext& Context) override;
+	virtual void SetActorEnabled(const EActorEnabledType EnabledType, AActor& Actor, const int32 EntityIdx, FMassExecutionContext& Context) override;
 
 	/** 
 	 * Returns an actor of the template type and setup fragments values from it
@@ -58,17 +58,17 @@ protected:
 	 * @param Actor is the actual actor to teleport
 	 * @param Context is the current light weight component execution context
 	 */
-	virtual void TeleportActor(const FTransform& Transform, AActor& Actor, FLWComponentSystemExecutionContext& Context) override;
+	virtual void TeleportActor(const FTransform& Transform, AActor& Actor, FMassExecutionContext& Context) override;
 
 	/**
 	 * Initializes velocity for each entity when just switch to use movement component, must be called within a ForEachEntityChunk
 	 * @param EntitySubsystem The subsystem in which all entities must be initialized
 	 * @param Context of the execution from the entity sub system
 	 */
-	void InitializeVelocity(UMassEntitySubsystem& EntitySubsystem, FLWComponentSystemExecutionContext& Context);
+	void InitializeVelocity(UMassEntitySubsystem& EntitySubsystem, FMassExecutionContext& Context);
 
 	/** 
 	 * A dedicated query for processing entities owning a FDataFragment_CharacterMovementComponentWrapper
 	 */
-	FLWComponentQuery CharacterMovementEntitiesQuery_Conditional;
+	FMassEntityQuery CharacterMovementEntitiesQuery_Conditional;
 };
\ No newline at end of file
diff --git a/Engine/Plugins/AI/MassCrowd/Source/MassCrowd/Public/MassCrowdServerRepresentationLODProcessor.h b/Engine/Plugins/AI/MassCrowd/Source/MassCrowd/Public/MassCrowdServerRepresentationLODProcessor.h
--- a/Engine/Plugins/AI/MassCrowd/Source/MassCrowd/Public/MassCrowdServerRepresentationLODProcessor.h
+++ b/Engine/Plugins/AI/MassCrowd/Source/MassCrowd/Public/MassCrowdServerRepresentationLODProcessor.h
@@ -22,7 +22,7 @@ protected:
 
 	virtual void ConfigureQueries() override;
 	virtual void Initialize(UObject& InOwner) override;
-	virtual void Execute(UMassEntitySubsystem& EntitySubsystem, FLWComponentSystemExecutionContext& Context) override;
+	virtual void Execute(UMassEntitySubsystem& EntitySubsystem, FMassExecutionContext& Context) override;
 
 	/** Distance where each LOD becomes relevant */
 	UPROPERTY(EditAnywhere, Category = "Mass|LOD", config)
@@ -38,5 +38,5 @@ protected:
 
 	TMassLODCalculator<FLODDefaultLogic> LODCalculator;
 
-	FLWComponentQuery EntityQuery;
+	FMassEntityQuery EntityQuery;
 };
diff --git a/Engine/Plugins/AI/MassCrowd/Source/MassCrowd/Public/MassCrowdServerRepresentationProcessor.h b/Engine/Plugins/AI/MassCrowd/Source/MassCrowd/Public/MassCrowdServerRepresentationProcessor.h
--- a/Engine/Plugins/AI/MassCrowd/Source/MassCrowd/Public/MassCrowdServerRepresentationProcessor.h
+++ b/Engine/Plugins/AI/MassCrowd/Source/MassCrowd/Public/MassCrowdServerRepresentationProcessor.h
@@ -25,5 +25,5 @@ protected:
 	 * @param EntitySubsystem is the system to execute the lambdas on each entity chunk
 	 * @param Context is the execution context to be passed when executing the lambdas
 	 */
-	virtual void Execute(UMassEntitySubsystem& EntitySubsystem, FLWComponentSystemExecutionContext& Context) override;
+	virtual void Execute(UMassEntitySubsystem& EntitySubsystem, FMassExecutionContext& Context) override;
 };
\ No newline at end of file
diff --git a/Engine/Plugins/AI/MassCrowd/Source/MassCrowd/Public/MassCrowdSubsystem.h b/Engine/Plugins/AI/MassCrowd/Source/MassCrowd/Public/MassCrowdSubsystem.h
--- a/Engine/Plugins/AI/MassCrowd/Source/MassCrowd/Public/MassCrowdSubsystem.h
+++ b/Engine/Plugins/AI/MassCrowd/Source/MassCrowd/Public/MassCrowdSubsystem.h
@@ -138,7 +138,7 @@ public:
 	 * @param OutSlotDirection Facing direction associated to the acquire slot.
 	 * @return Index of the slot, or INDEX_NONE if no slots are available or if the lane is not a waiting lane.
 	 */
-	int32 AcquireWaitingSlot(const FLWEntity Entity, const FVector& EntityPosition, const FZoneGraphLaneHandle LaneHandle,
+	int32 AcquireWaitingSlot(const FMassEntityHandle Entity, const FVector& EntityPosition, const FZoneGraphLaneHandle LaneHandle,
 							 FVector& OutSlotPosition, FVector& OutSlotDirection);
 	
 	/**
@@ -147,7 +147,7 @@ public:
 	 * @param LaneHandle A handle to a lane with waiting data.
 	 * @param SlotIndex Index of the previously acquired slot.
 	 */
-	void ReleaseWaitingSlot(const FLWEntity Entity, const FZoneGraphLaneHandle LaneHandle, const int32 SlotIndex);
+	void ReleaseWaitingSlot(const FMassEntityHandle Entity, const FZoneGraphLaneHandle LaneHandle, const int32 SlotIndex);
 
 	/**
 	 * Callback from the lane tracker processor to indicates a mass entity changing lane.
@@ -155,7 +155,7 @@ public:
 	 * @param PreviousLaneHandle Last frame lane handle (can be invalid)
 	 * @param CurrentLaneHandle Current frame lane handle (can be invalid)
 	 */
-	void OnEntityLaneChanged(const FLWEntity& Entity, const FZoneGraphLaneHandle PreviousLaneHandle, const FZoneGraphLaneHandle CurrentLaneHandle);
+	void OnEntityLaneChanged(const FMassEntityHandle& Entity, const FZoneGraphLaneHandle PreviousLaneHandle, const FZoneGraphLaneHandle CurrentLaneHandle);
 
 	/** Returns the weight for lane selection that is associated to the given lane based on its density tag. */
 	float GetDensityWeight(const FZoneGraphLaneHandle LaneHandle, const FZoneGraphTagMask LaneTagMask) const;
@@ -193,7 +193,7 @@ protected:
 	 * @param LaneIndex Index of the lane
 	 * @param TrackingData Runtime state associated to the lane
 	 */
-	void OnEnterTrackedLane(const FLWEntity& Entity, const int32 LaneIndex, FCrowdTrackingLaneData& TrackingData);
+	void OnEnterTrackedLane(const FMassEntityHandle& Entity, const int32 LaneIndex, FCrowdTrackingLaneData& TrackingData);
 
 	/**
 	 * Callback to keep count of entities currently on a given lane.
@@ -201,7 +201,7 @@ protected:
 	 * @param LaneIndex Index of the lane
 	 * @param TrackingData Runtime state associated to the lane
 	 */
-	void OnExitTrackedLane(const FLWEntity& Entity, const int32 LaneIndex, FCrowdTrackingLaneData& TrackingData);
+	void OnExitTrackedLane(const FMassEntityHandle& Entity, const int32 LaneIndex, FCrowdTrackingLaneData& TrackingData);
 
 	/**
 	 * Creates and initializes the occupancy data of a lane.
diff --git a/Engine/Plugins/AI/MassCrowd/Source/MassCrowd/Public/MassCrowdVisualizationLODProcessor.h b/Engine/Plugins/AI/MassCrowd/Source/MassCrowd/Public/MassCrowdVisualizationLODProcessor.h
--- a/Engine/Plugins/AI/MassCrowd/Source/MassCrowd/Public/MassCrowdVisualizationLODProcessor.h
+++ b/Engine/Plugins/AI/MassCrowd/Source/MassCrowd/Public/MassCrowdVisualizationLODProcessor.h
@@ -18,5 +18,5 @@ public:
 
 protected:
 	virtual void ConfigureQueries() override;
-	virtual void Execute(UMassEntitySubsystem& EntitySubsystem, FLWComponentSystemExecutionContext& Context) override;
+	virtual void Execute(UMassEntitySubsystem& EntitySubsystem, FMassExecutionContext& Context) override;
 };
diff --git a/Engine/Plugins/AI/MassCrowd/Source/MassCrowd/Public/MassCrowdVisualizationProcessor.h b/Engine/Plugins/AI/MassCrowd/Source/MassCrowd/Public/MassCrowdVisualizationProcessor.h
--- a/Engine/Plugins/AI/MassCrowd/Source/MassCrowd/Public/MassCrowdVisualizationProcessor.h
+++ b/Engine/Plugins/AI/MassCrowd/Source/MassCrowd/Public/MassCrowdVisualizationProcessor.h
@@ -18,7 +18,7 @@ public:
 
 protected:
 
-	/** Configure the owned FLWComponentQuery instances to express processor's requirements */
+	/** Configure the owned FMassEntityQuery instances to express processor's requirements */
 	virtual void ConfigureQueries() override;
 
 	/**
@@ -26,7 +26,7 @@ protected:
 	 * @param EntitySubsystem is the system to execute the lambdas on each entity chunk
 	 * @param Context is the execution context to be passed when executing the lambdas
 	 */
-	virtual void Execute(UMassEntitySubsystem& EntitySubsystem, FLWComponentSystemExecutionContext& Context) override;
+	virtual void Execute(UMassEntitySubsystem& EntitySubsystem, FMassExecutionContext& Context) override;
 };
 
 UCLASS()
