[StateTree] fixed 'NumTickedStateTree' stat

diff --git a/Engine/Plugins/AI/MassAI/Source/MassAIBehavior/Private/MassStateTreeProcessors.cpp b/Engine/Plugins/AI/MassAI/Source/MassAIBehavior/Private/MassStateTreeProcessors.cpp
--- a/Engine/Plugins/AI/MassAI/Source/MassAIBehavior/Private/MassStateTreeProcessors.cpp
+++ b/Engine/Plugins/AI/MassAI/Source/MassAIBehavior/Private/MassStateTreeProcessors.cpp
@@ -77,7 +77,7 @@ bool SetExternalSubsystems(FMassStateTreeExecutionContext& Context)
 	return bFoundAllSubsystems;
 }
 
-void ProcessChunk(
+void ForEachEntityInChunk(
 	FMassStateTreeExecutionContext& StateTreeContext,
 	UMassStateTreeSubsystem& MassStateTreeSubsystem,
 	const TFunctionRef<void(FMassStateTreeExecutionContext&, FStateTreeItemView)> ForEachEntityCallback)
@@ -166,7 +166,7 @@ void UMassStateTreeFragmentInitializer::Execute(UMassEntitySubsystem& EntitySubs
 		Context,
 		[this, &StateTreeContext, &MassStateTreeSubsystem, &EntitiesToSignal](const FMassExecutionContext& Context)
 		{
-			UE::MassBehavior::ProcessChunk(
+			UE::MassBehavior::ForEachEntityInChunk(
 				StateTreeContext,
 				*MassStateTreeSubsystem,
 				[](FMassStateTreeExecutionContext& StateTreeExecutionContext, FStateTreeItemView Storage)
@@ -212,7 +212,7 @@ void UMassStateTreeFragmentDestructor::Execute(UMassEntitySubsystem& EntitySubsy
 		Context,
 		[this, &StateTreeContext, &MassStateTreeSubsystem](FMassExecutionContext&)
 		{
-			UE::MassBehavior::ProcessChunk(
+			UE::MassBehavior::ForEachEntityInChunk(
 				StateTreeContext,
 				*MassStateTreeSubsystem,
 				[](FMassStateTreeExecutionContext& StateTreeExecutionContext, FStateTreeItemView Storage)
@@ -293,16 +293,16 @@ void UMassStateTreeProcessor::SignalEntities(UMassEntitySubsystem& EntitySubsyst
 		Context,
 		[this, &StateTreeContext, TimeDelta, TimeInSeconds, &EntitiesToSignal](FMassExecutionContext& Context)
 		{
+			// Keep stats regarding the amount of tree instances ticked per frame
+			CSV_CUSTOM_STAT(StateTreeProcessor, NumTickedStateTree, Context.GetEntitiesNum(), ECsvCustomStatOp::Accumulate);
+
 			TArrayView<FMassStateTreeFragment> StateTreeList = Context.GetMutableComponentView<FMassStateTreeFragment>();
 
-			UE::MassBehavior::ProcessChunk(
+			UE::MassBehavior::ForEachEntityInChunk(
 				StateTreeContext,
 				*MassStateTreeSubsystem,
 				[&StateTreeList, TimeDelta, TimeInSeconds, &EntitiesToSignal](FMassStateTreeExecutionContext& StateTreeExecutionContext, const FStateTreeItemView Storage)
 				{
-					// Keep stats regarding the amount of tree instances ticked per frame
-					CSV_CUSTOM_STAT(StateTreeProcessor, NumTickedStateTree, StateTreeExecutionContext.GetEntitySubsystemExecutionContext().GetEntitiesNum(), ECsvCustomStatOp::Accumulate);
-
 					// Compute adjusted delta time
 					TOptional<float>& LastUpdate = StateTreeList[StateTreeExecutionContext.GetEntityIndex()].LastUpdateTimeInSeconds;
 					const float AdjustedTimeDelta = LastUpdate.IsSet() ? TimeDelta + (TimeInSeconds - LastUpdate.GetValue()) : TimeDelta;
