[MassEntity] Renamed FMassExecutionContext::GetEntitiesNum to GetNumEntities and FMassEntityQuery::GetMatchingEntitiesNum to GetNumMatchingEntities

diff --git a/Engine/Plugins/AI/MassAI/Source/MassAIBehavior/Private/MassLookAtProcessors.cpp b/Engine/Plugins/AI/MassAI/Source/MassAIBehavior/Private/MassLookAtProcessors.cpp
--- a/Engine/Plugins/AI/MassAI/Source/MassAIBehavior/Private/MassLookAtProcessors.cpp
+++ b/Engine/Plugins/AI/MassAI/Source/MassAIBehavior/Private/MassLookAtProcessors.cpp
@@ -125,7 +125,7 @@ void UMassLookAtProcessor::Execute(UMassEntitySubsystem& EntitySubsystem, FMassE
 
 	EntityQuery_Conditional.ForEachEntityChunk(EntitySubsystem, Context, [this, &EntitySubsystem, CurrentTime](FMassExecutionContext& Context)
 		{
-			const int32 NumEntities = Context.GetEntitiesNum();
+			const int32 NumEntities = Context.GetNumEntities();
 			const TArrayView<FMassLookAtFragment> LookAtList = Context.GetMutableComponentView<FMassLookAtFragment>();
 			const TConstArrayView<FDataFragment_Transform> TransformList = Context.GetComponentView<FDataFragment_Transform>();
 			const TConstArrayView<FMassMoveTargetFragment> MoveTargetList = Context.GetComponentView<FMassMoveTargetFragment>();
diff --git a/Engine/Plugins/AI/MassAI/Source/MassAIBehavior/Private/MassStateTreeProcessors.cpp b/Engine/Plugins/AI/MassAI/Source/MassAIBehavior/Private/MassStateTreeProcessors.cpp
--- a/Engine/Plugins/AI/MassAI/Source/MassAIBehavior/Private/MassStateTreeProcessors.cpp
+++ b/Engine/Plugins/AI/MassAI/Source/MassAIBehavior/Private/MassStateTreeProcessors.cpp
@@ -86,7 +86,7 @@ void ForEachEntityInChunk(
 	const TConstArrayView<FMassStateTreeFragment> StateTreeList = Context.GetComponentView<FMassStateTreeFragment>();
 
 	// Assuming that all the entities share same StateTree, because they all should have the same storage fragment.
-	const int32 NumEntities = Context.GetEntitiesNum();
+	const int32 NumEntities = Context.GetNumEntities();
 	check(NumEntities > 0);
 	UStateTree* StateTree = MassStateTreeSubsystem.GetRegisteredStateTreeAsset(StateTreeList[0].StateTreeHandle);
 
@@ -294,7 +294,7 @@ void UMassStateTreeProcessor::SignalEntities(UMassEntitySubsystem& EntitySubsyst
 		[this, &StateTreeContext, TimeDelta, TimeInSeconds, &EntitiesToSignal](FMassExecutionContext& Context)
 		{
 			// Keep stats regarding the amount of tree instances ticked per frame
-			CSV_CUSTOM_STAT(StateTreeProcessor, NumTickedStateTree, Context.GetEntitiesNum(), ECsvCustomStatOp::Accumulate);
+			CSV_CUSTOM_STAT(StateTreeProcessor, NumTickedStateTree, Context.GetNumEntities(), ECsvCustomStatOp::Accumulate);
 
 			TArrayView<FMassStateTreeFragment> StateTreeList = Context.GetMutableComponentView<FMassStateTreeFragment>();
 
diff --git a/Engine/Plugins/AI/MassAI/Source/MassAIBehavior/Private/MassZoneGraphAnnotationProcessors.cpp b/Engine/Plugins/AI/MassAI/Source/MassAIBehavior/Private/MassZoneGraphAnnotationProcessors.cpp
--- a/Engine/Plugins/AI/MassAI/Source/MassAIBehavior/Private/MassZoneGraphAnnotationProcessors.cpp
+++ b/Engine/Plugins/AI/MassAI/Source/MassAIBehavior/Private/MassZoneGraphAnnotationProcessors.cpp
@@ -35,7 +35,7 @@ void UMassZoneGraphAnnotationTagsInitializer::Execute(UMassEntitySubsystem& Enti
 {
 	EntityQuery.ForEachEntityChunk(EntitySubsystem, Context, [this](FMassExecutionContext& Context)
 	{
-		const int32 NumEntities = Context.GetEntitiesNum();
+		const int32 NumEntities = Context.GetNumEntities();
 		const TArrayView<FMassZoneGraphAnnotationTagsFragment> AnnotationTagsList = Context.GetMutableComponentView<FMassZoneGraphAnnotationTagsFragment>();
 		const TConstArrayView<FMassZoneGraphLaneLocationFragment> LaneLocationList = Context.GetComponentView<FMassZoneGraphLaneLocationFragment>();
 
@@ -98,7 +98,7 @@ void UMassZoneGraphAnnotationTagUpdateProcessor::Execute(UMassEntitySubsystem& E
 			return;
 		}
 
-		const int32 NumEntities = Context.GetEntitiesNum();
+		const int32 NumEntities = Context.GetNumEntities();
 		const TArrayView<FMassZoneGraphAnnotationTagsFragment> AnnotationTagsList = Context.GetMutableComponentView<FMassZoneGraphAnnotationTagsFragment>();
 		const TConstArrayView<FMassZoneGraphLaneLocationFragment> LaneLocationList = Context.GetComponentView<FMassZoneGraphLaneLocationFragment>();
 
@@ -145,7 +145,7 @@ void UMassZoneGraphAnnotationTagUpdateProcessor::SignalEntities(UMassEntitySubsy
 
 	EntityQuery.ForEachEntityChunk(EntitySubsystem, Context, [this](FMassExecutionContext& Context)
 	{
-		const int32 NumEntities = Context.GetEntitiesNum();
+		const int32 NumEntities = Context.GetNumEntities();
 		const TConstArrayView<FMassZoneGraphLaneLocationFragment> LaneLocationList = Context.GetComponentView<FMassZoneGraphLaneLocationFragment>();
 		const TArrayView<FMassZoneGraphAnnotationTagsFragment> AnnotationTagsList = Context.GetMutableComponentView<FMassZoneGraphAnnotationTagsFragment>();
 
diff --git a/Engine/Plugins/AI/MassAI/Source/MassAIDebug/Private/GameplayDebuggerCategory_Mass.cpp b/Engine/Plugins/AI/MassAI/Source/MassAIDebug/Private/GameplayDebuggerCategory_Mass.cpp
--- a/Engine/Plugins/AI/MassAI/Source/MassAIDebug/Private/GameplayDebuggerCategory_Mass.cpp
+++ b/Engine/Plugins/AI/MassAI/Source/MassAIDebug/Private/GameplayDebuggerCategory_Mass.cpp
@@ -318,7 +318,7 @@ void FGameplayDebuggerCategory_Mass::CollectData(APlayerController* OwnerPC, AAc
 		
 			EntityQuery.ForEachEntityChunk(*EntitySystem, Context, [this, Debugger, MassStateTreeSubsystem, EntitySystem, OwnerPC, ViewLocation, ViewDirection, CurrentTime](FMassExecutionContext& Context)
 			{
-				const int32 NumEntities = Context.GetEntitiesNum();
+				const int32 NumEntities = Context.GetNumEntities();
 				const TConstArrayView<FMassStateTreeFragment> StateTreeList = Context.GetComponentView<FMassStateTreeFragment>();
 				const TConstArrayView<FDataFragment_Transform> TransformList = Context.GetComponentView<FDataFragment_Transform>();
 				const TConstArrayView<FDataFragment_AgentRadius> RadiusList = Context.GetComponentView<FDataFragment_AgentRadius>();
@@ -527,7 +527,7 @@ void FGameplayDebuggerCategory_Mass::CollectData(APlayerController* OwnerPC, AAc
 			FMassExecutionContext Context(0.f);
 			EntityColliderQuery.ForEachEntityChunk(*EntitySystem, Context, [this, Debugger, EntitySystem, OwnerPC, ViewLocation, ViewDirection](FMassExecutionContext& Context)
 			{
-				const int32 NumEntities = Context.GetEntitiesNum();
+				const int32 NumEntities = Context.GetNumEntities();
 				const TConstArrayView<FDataFragment_Transform> TransformList = Context.GetComponentView<FDataFragment_Transform>();
 				const TConstArrayView<FMassAvoidanceColliderFragment> CollidersList = Context.GetComponentView<FMassAvoidanceColliderFragment>();
 	
diff --git a/Engine/Plugins/AI/MassAI/Source/MassAIDebug/Private/MassDebugStateTreeProcessor.cpp b/Engine/Plugins/AI/MassAI/Source/MassAIDebug/Private/MassDebugStateTreeProcessor.cpp
--- a/Engine/Plugins/AI/MassAI/Source/MassAIDebug/Private/MassDebugStateTreeProcessor.cpp
+++ b/Engine/Plugins/AI/MassAI/Source/MassAIDebug/Private/MassDebugStateTreeProcessor.cpp
@@ -53,7 +53,7 @@ void UMassDebugStateTreeProcessor::Execute(UMassEntitySubsystem& EntitySubsystem
 	EntityQuery.ForEachEntityChunk(EntitySubsystem, Context, [this, Debugger, MassStateTreeSubsystem, &EntitySubsystem](FMassExecutionContext& Context)
 		{
 			const FMassEntityHandle SelectedEntity = Debugger->GetSelectedEntity();
-			const int32 NumEntities = Context.GetEntitiesNum();
+			const int32 NumEntities = Context.GetNumEntities();
 			const TConstArrayView<FMassStateTreeFragment> StateTreeList = Context.GetComponentView<FMassStateTreeFragment>();
 			const TConstArrayView<FDataFragment_Transform> TransformList = Context.GetComponentView<FDataFragment_Transform>();
 			const UStateTree* StateTree = MassStateTreeSubsystem->GetRegisteredStateTreeAsset(StateTreeList[0].StateTreeHandle);
diff --git a/Engine/Plugins/AI/MassAI/Source/MassAIMovement/Private/AssignRandomNavLocationProcessor.cpp b/Engine/Plugins/AI/MassAI/Source/MassAIMovement/Private/AssignRandomNavLocationProcessor.cpp
--- a/Engine/Plugins/AI/MassAI/Source/MassAIMovement/Private/AssignRandomNavLocationProcessor.cpp
+++ b/Engine/Plugins/AI/MassAI/Source/MassAIMovement/Private/AssignRandomNavLocationProcessor.cpp
@@ -32,7 +32,7 @@ void UAssignRandomNavLocationProcessor::Execute(UMassEntitySubsystem& EntitySubs
 		{
 			TArrayView<FDataFragment_Transform> LocationList = Context.GetMutableComponentView<FDataFragment_Transform>();
 			TArrayView<FDataFragment_NavLocation> NavLocationList = Context.GetMutableComponentView<FDataFragment_NavLocation>();
-			const int32 NumEntities = Context.GetEntitiesNum();
+			const int32 NumEntities = Context.GetNumEntities();
 			TArray<FNavLocation> Locations;
 			Locations.AddDefaulted(NumEntities);
 			for (FNavLocation& Location : Locations)
diff --git a/Engine/Plugins/AI/MassAI/Source/MassAIMovement/Private/MassAvoidanceProcessors.cpp b/Engine/Plugins/AI/MassAI/Source/MassAIMovement/Private/MassAvoidanceProcessors.cpp
--- a/Engine/Plugins/AI/MassAI/Source/MassAIMovement/Private/MassAvoidanceProcessors.cpp
+++ b/Engine/Plugins/AI/MassAI/Source/MassAIMovement/Private/MassAvoidanceProcessors.cpp
@@ -495,7 +495,7 @@ void UMassAvoidanceProcessor::Execute(UMassEntitySubsystem& EntitySubsystem, FMa
 
 	EntityQuery.ForEachEntityChunk(EntitySubsystem, Context, [=, &EntitySubsystem](FMassExecutionContext& Ctx)
 	{
-		const int32 NumEntities = Ctx.GetEntitiesNum();
+		const int32 NumEntities = Ctx.GetNumEntities();
 		const float DistanceCutOffSquare = FMath::Square(UE::MassAvoidance::Tweakables::AgentDetectionDistance);
 
 		const TArrayView<FMassSteeringFragment> SteeringList = Ctx.GetMutableComponentView<FMassSteeringFragment>();
@@ -1178,7 +1178,7 @@ void UMassStandingAvoidanceProcessor::Execute(UMassEntitySubsystem& EntitySubsys
 	// Avoidance while standing
 	EntityQuery.ForEachEntityChunk(EntitySubsystem, Context, [&EntitySubsystem, MovementSubsystem, World](FMassExecutionContext& Context)
 	{
-		const int32 NumEntities = Context.GetEntitiesNum();
+		const int32 NumEntities = Context.GetNumEntities();
 		const float DeltaTime = Context.GetDeltaTimeSeconds();
 		const float DistanceCutOffSquare = FMath::Square(UE::MassAvoidance::Tweakables::AgentDetectionDistance);
 
@@ -1436,7 +1436,7 @@ void UMassAvoidanceObstacleProcessor::Execute(UMassEntitySubsystem& EntitySubsys
 	// can't be ParallelFor due to MovementSubsystem->GetGridMutable().Move not being thread-safe
 	AddToGridEntityQuery.ForEachEntityChunk(EntitySubsystem, Context, [this, MovementSubsystem, &EntitySubsystem](FMassExecutionContext& Context)
 		{
-			const int32 NumEntities = Context.GetEntitiesNum();
+			const int32 NumEntities = Context.GetNumEntities();
 
 			TConstArrayView<FDataFragment_Transform> LocationList = Context.GetComponentView<FDataFragment_Transform>();
 			TConstArrayView<FDataFragment_AgentRadius> RadiiList = Context.GetComponentView<FDataFragment_AgentRadius>();
@@ -1466,7 +1466,7 @@ void UMassAvoidanceObstacleProcessor::Execute(UMassEntitySubsystem& EntitySubsys
 
 	UpdateGridEntityQuery.ForEachEntityChunk(EntitySubsystem, Context, [this, MovementSubsystem, &EntitySubsystem](FMassExecutionContext& Context)
 		{
-			const int32 NumEntities = Context.GetEntitiesNum();
+			const int32 NumEntities = Context.GetNumEntities();
 
 			TConstArrayView<FDataFragment_Transform> LocationList = Context.GetComponentView<FDataFragment_Transform>();
 			TConstArrayView<FDataFragment_AgentRadius> RadiiList = Context.GetComponentView<FDataFragment_AgentRadius>();
@@ -1502,7 +1502,7 @@ void UMassAvoidanceObstacleProcessor::Execute(UMassEntitySubsystem& EntitySubsys
 
 	RemoveFromGridEntityQuery.ForEachEntityChunk(EntitySubsystem, Context, [this, MovementSubsystem](FMassExecutionContext& Context)
 	{
-		const int32 NumEntities = Context.GetEntitiesNum();
+		const int32 NumEntities = Context.GetNumEntities();
 
 		TArrayView<FMassAvoidanceObstacleGridCellLocationFragment> AvoidanceObstacleCellLocationList = Context.GetMutableComponentView<FMassAvoidanceObstacleGridCellLocationFragment>();
 		for (int32 EntityIndex = 0; EntityIndex < NumEntities; ++EntityIndex)
@@ -1566,7 +1566,7 @@ void UMassNavigationBoundaryProcessor::Execute(UMassEntitySubsystem& EntitySubsy
 
 	EntityQuery.ForEachEntityChunk(EntitySubsystem, Context, [this, NavData](FMassExecutionContext& Context)
 		{
-			const int32 NumEntities = Context.GetEntitiesNum();
+			const int32 NumEntities = Context.GetNumEntities();
 			TConstArrayView<FDataFragment_Transform> LocationList = Context.GetComponentView<FDataFragment_Transform>();
 			TConstArrayView<FDataFragment_NavLocation> NavLocationList = Context.GetComponentView<FDataFragment_NavLocation>();
 			TConstArrayView<FMassEdgeDetectionParamsFragment> EdgeDetectionParamsList = Context.GetComponentView<FMassEdgeDetectionParamsFragment>();
@@ -1648,7 +1648,7 @@ void UMassLaneBoundaryProcessor::Execute(UMassEntitySubsystem& EntitySubsystem,
 
 	EntityQuery.ForEachEntityChunk(EntitySubsystem, Context, [this, World, ZoneGraphSubsystem](FMassExecutionContext& Context)
 		{
-			const int32 NumEntities = Context.GetEntitiesNum();
+			const int32 NumEntities = Context.GetNumEntities();
 
 			TConstArrayView<FDataFragment_Transform> LocationList = Context.GetComponentView<FDataFragment_Transform>();
 			TArrayView<FMassNavigationEdgesFragment> EdgesList = Context.GetMutableComponentView<FMassNavigationEdgesFragment>();
@@ -1918,7 +1918,7 @@ void UMassLaneCacheBoundaryProcessor::Execute(UMassEntitySubsystem& EntitySubsys
 
 	EntityQuery.ForEachEntityChunk(EntitySubsystem, Context, [this, World](FMassExecutionContext& Context)
 	{
-		const int32 NumEntities = Context.GetEntitiesNum();
+		const int32 NumEntities = Context.GetNumEntities();
 
 		TConstArrayView<FMassZoneGraphCachedLaneFragment> CachedLaneList = Context.GetComponentView<FMassZoneGraphCachedLaneFragment>();
 		TConstArrayView<FMassZoneGraphLaneLocationFragment> LaneLocationList = Context.GetComponentView<FMassZoneGraphLaneLocationFragment>();
diff --git a/Engine/Plugins/AI/MassAI/Source/MassAIMovement/Private/MassZoneGraphMovementProcessors.cpp b/Engine/Plugins/AI/MassAI/Source/MassAIMovement/Private/MassZoneGraphMovementProcessors.cpp
--- a/Engine/Plugins/AI/MassAI/Source/MassAIMovement/Private/MassZoneGraphMovementProcessors.cpp
+++ b/Engine/Plugins/AI/MassAI/Source/MassAIMovement/Private/MassZoneGraphMovementProcessors.cpp
@@ -100,7 +100,7 @@ void UMassZoneGraphLocationInitializer::Execute(UMassEntitySubsystem& EntitySubs
 
 	EntityQuery.ForEachEntityChunk(EntitySubsystem, Context, [this, Settings](FMassExecutionContext& Context)
 	{
-		const int32 NumEntities = Context.GetEntitiesNum();
+		const int32 NumEntities = Context.GetNumEntities();
 		const TArrayView<FMassZoneGraphLaneLocationFragment> LaneLocationList = Context.GetMutableComponentView<FMassZoneGraphLaneLocationFragment>();
 		const TArrayView<FMassMoveTargetFragment> MoveTargetList = Context.GetMutableComponentView<FMassMoveTargetFragment>();
 		const TConstArrayView<FMassMovementConfigFragment> MovementConfigList = Context.GetComponentView<FMassMovementConfigFragment>();
@@ -207,7 +207,7 @@ void UMassZoneGraphPathFollowProcessor::Execute(UMassEntitySubsystem& EntitySubs
 
 	EntityQuery_Conditional.ForEachEntityChunk(EntitySubsystem, Context, [this, &EntitiesToSignalPathDone, &EntitiesToSignalLaneChanged](FMassExecutionContext& Context)
 		{
-		const int32 NumEntities = Context.GetEntitiesNum();
+		const int32 NumEntities = Context.GetNumEntities();
 		const TArrayView<FMassZoneGraphShortPathFragment> ShortPathList = Context.GetMutableComponentView<FMassZoneGraphShortPathFragment>();
 		const TArrayView<FMassZoneGraphLaneLocationFragment> LaneLocationList = Context.GetMutableComponentView<FMassZoneGraphLaneLocationFragment>();
 		const TArrayView<FMassMoveTargetFragment> MoveTargetList = Context.GetMutableComponentView<FMassMoveTargetFragment>();
@@ -497,7 +497,7 @@ void UMassZoneGraphSteeringProcessor::Execute(UMassEntitySubsystem& EntitySubsys
 	
 	EntityQuery.ForEachEntityChunk(EntitySubsystem, Context, [this, Settings](FMassExecutionContext& Context)
 	{
-		const int32 NumEntities = Context.GetEntitiesNum();
+		const int32 NumEntities = Context.GetNumEntities();
 		const TArrayView<FMassMoveTargetFragment> MoveTargetList = Context.GetMutableComponentView<FMassMoveTargetFragment>();
 		const TConstArrayView<FMassMovementConfigFragment> MovementConfigList = Context.GetComponentView<FMassMovementConfigFragment>();
 		const TConstArrayView<FMassVelocityFragment> VelocityList = Context.GetComponentView<FMassVelocityFragment>();
diff --git a/Engine/Plugins/AI/MassAI/Source/MassAIMovement/Private/MovementProcessor.cpp b/Engine/Plugins/AI/MassAI/Source/MassAIMovement/Private/MovementProcessor.cpp
--- a/Engine/Plugins/AI/MassAI/Source/MassAIMovement/Private/MovementProcessor.cpp
+++ b/Engine/Plugins/AI/MassAI/Source/MassAIMovement/Private/MovementProcessor.cpp
@@ -126,7 +126,7 @@ void UMassProcessor_Movement::Execute(UMassEntitySubsystem& EntitySubsystem, FMa
 
 	EntityQuery.ForEachEntityChunk(EntitySubsystem, Context, [this, TimeDelta](FMassExecutionContext& Context)
 		{
-			const int32 NumEntities = Context.GetEntitiesNum();
+			const int32 NumEntities = Context.GetNumEntities();
 			const TArrayView<FDataFragment_Transform> LocationList = Context.GetMutableComponentView<FDataFragment_Transform>();
 			const TArrayView<FMassVelocityFragment> VelocityList = Context.GetMutableComponentView<FMassVelocityFragment>();
 			for (int32 i = 0; i < NumEntities; ++i)
@@ -186,7 +186,7 @@ void UMassAvoidanceObstacleRemoverFragmentDestructor::Execute(UMassEntitySubsyst
 
 	EntityQuery.ForEachEntityChunk(EntitySubsystem, Context, [this, MovementSubsystem](FMassExecutionContext& Context)
 		{
-			const int32 NumEntities = Context.GetEntitiesNum();
+			const int32 NumEntities = Context.GetNumEntities();
 			const TArrayView<FMassAvoidanceObstacleGridCellLocationFragment> AvoidanceObstacleCellLocationList = Context.GetMutableComponentView<FMassAvoidanceObstacleGridCellLocationFragment>();
 
 			for (int32 i = 0; i < NumEntities; ++i)
diff --git a/Engine/Plugins/AI/MassAI/Source/MassAIMovement/Private/NavigationProcessor.cpp b/Engine/Plugins/AI/MassAI/Source/MassAIMovement/Private/NavigationProcessor.cpp
--- a/Engine/Plugins/AI/MassAI/Source/MassAIMovement/Private/NavigationProcessor.cpp
+++ b/Engine/Plugins/AI/MassAI/Source/MassAIMovement/Private/NavigationProcessor.cpp
@@ -112,7 +112,7 @@ void UMassApplyVelocityMoveTargetProcessor::Execute(UMassEntitySubsystem& Entity
 
 		HighResEntityQuery.ForEachEntityChunk(EntitySubsystem, Context, [this, Settings, TimeDelta](FMassExecutionContext& Context)
 		{
-			const int32 NumEntities = Context.GetEntitiesNum();
+			const int32 NumEntities = Context.GetNumEntities();
 			constexpr float ZDamperHalfLife = 0.2f;
 
 			const float OrientationEndOfPathHeadingAnticipation = FMath::Max(KINDA_SMALL_NUMBER, Settings->OrientationEndOfPathHeadingAnticipation);
@@ -205,7 +205,7 @@ void UMassApplyVelocityMoveTargetProcessor::Execute(UMassEntitySubsystem& Entity
 
 		LowResEntityQuery_Conditional.ForEachEntityChunk(EntitySubsystem, Context, [this](FMassExecutionContext& Context)
 			{
-				const int32 NumEntities = Context.GetEntitiesNum();
+				const int32 NumEntities = Context.GetNumEntities();
 
 				const TArrayView<FDataFragment_Transform> LocationList = Context.GetMutableComponentView<FDataFragment_Transform>();
 				const TConstArrayView<FMassMoveTargetFragment> MoveTargetList = Context.GetComponentView<FMassMoveTargetFragment>();
@@ -245,7 +245,7 @@ void UMassDynamicObstacleProcessor::Execute(UMassEntitySubsystem& EntitySubsyste
 {
 	EntityQuery_Conditional.ForEachEntityChunk(EntitySubsystem, Context, [this](FMassExecutionContext& Context)
 		{
-			const int32 NumEntities = Context.GetEntitiesNum();
+			const int32 NumEntities = Context.GetNumEntities();
 
 			const TConstArrayView<FDataFragment_Transform> LocationList = Context.GetComponentView<FDataFragment_Transform>();
 			const TConstArrayView<FDataFragment_AgentRadius> RadiusList = Context.GetComponentView<FDataFragment_AgentRadius>();
diff --git a/Engine/Plugins/AI/MassAI/Source/MassAIMovement/Private/SnapToNavigationProcessor.cpp b/Engine/Plugins/AI/MassAI/Source/MassAIMovement/Private/SnapToNavigationProcessor.cpp
--- a/Engine/Plugins/AI/MassAI/Source/MassAIMovement/Private/SnapToNavigationProcessor.cpp
+++ b/Engine/Plugins/AI/MassAI/Source/MassAIMovement/Private/SnapToNavigationProcessor.cpp
@@ -35,7 +35,7 @@ void USnapToNavigationProcessor::Execute(UMassEntitySubsystem& EntitySubsystem,
 
 	EntityQuery.ForEachEntityChunk(EntitySubsystem, Context, [this, NavData](FMassExecutionContext& Context)
 		{
-			const int32 NumEntities = Context.GetEntitiesNum();
+			const int32 NumEntities = Context.GetNumEntities();
 			const TArrayView<FDataFragment_Transform> LocationList = Context.GetMutableComponentView<FDataFragment_Transform>();
 			const TArrayView<FDataFragment_NavLocation> NavLocationList = Context.GetMutableComponentView<FDataFragment_NavLocation>();
 
diff --git a/Engine/Plugins/AI/MassCrowd/Source/MassCrowd/Private/MassCrowdNavigationProcessor.cpp b/Engine/Plugins/AI/MassCrowd/Source/MassCrowd/Private/MassCrowdNavigationProcessor.cpp
--- a/Engine/Plugins/AI/MassCrowd/Source/MassCrowd/Private/MassCrowdNavigationProcessor.cpp
+++ b/Engine/Plugins/AI/MassCrowd/Source/MassCrowd/Private/MassCrowdNavigationProcessor.cpp
@@ -47,7 +47,7 @@ void UMassCrowdLaneTrackingSignalProcessor::SignalEntities(UMassEntitySubsystem&
 
 	EntityQuery.ForEachEntityChunk(EntitySubsystem, Context, [this](FMassExecutionContext& Context)
 	{
-		const int32 NumEntities = Context.GetEntitiesNum();
+		const int32 NumEntities = Context.GetNumEntities();
 		const TConstArrayView<FMassZoneGraphLaneLocationFragment> LaneLocationList = Context.GetComponentView<FMassZoneGraphLaneLocationFragment>();
 		const TArrayView<FMassCrowdLaneTrackingFragment> LaneTrackingList = Context.GetMutableComponentView<FMassCrowdLaneTrackingFragment>();
 
@@ -90,7 +90,7 @@ void UMassCrowdLaneTrackingDestructor::Execute(UMassEntitySubsystem& EntitySubsy
 {
 	EntityQuery.ForEachEntityChunk(EntitySubsystem, Context, [this](const FMassExecutionContext& Context)
 	{
-		const int32 NumEntities = Context.GetEntitiesNum();
+		const int32 NumEntities = Context.GetNumEntities();
 		const TConstArrayView<FMassCrowdLaneTrackingFragment> LaneTrackingList = Context.GetComponentView<FMassCrowdLaneTrackingFragment>();
 
 		for (int32 EntityIndex = 0; EntityIndex < NumEntities; ++EntityIndex)
diff --git a/Engine/Plugins/AI/MassCrowd/Source/MassCrowd/Private/MassCrowdRepresentationProcessor.cpp b/Engine/Plugins/AI/MassCrowd/Source/MassCrowd/Private/MassCrowdRepresentationProcessor.cpp
--- a/Engine/Plugins/AI/MassCrowd/Source/MassCrowd/Private/MassCrowdRepresentationProcessor.cpp
+++ b/Engine/Plugins/AI/MassCrowd/Source/MassCrowd/Private/MassCrowdRepresentationProcessor.cpp
@@ -46,7 +46,7 @@ void UMassCrowdRepresentationProcessor::InitializeVelocity(UMassEntitySubsystem&
 			const TConstArrayView<FMassVelocityFragment> VelocityList = Context.GetComponentView<FMassVelocityFragment>();
 			const TArrayView<FDataFragment_CharacterMovementComponentWrapper> CharMoveWrapperList = Context.GetMutableComponentView<FDataFragment_CharacterMovementComponentWrapper>();
 
-			const int32 NumEntities = Context.GetEntitiesNum();
+			const int32 NumEntities = Context.GetNumEntities();
 			for (int32 EntityIdx = 0; EntityIdx < NumEntities; EntityIdx++)
 			{
 				const FMassRepresentationFragment& Visualization = VisualizationList[EntityIdx];
diff --git a/Engine/Plugins/AI/MassCrowd/Source/MassCrowd/Private/MassCrowdVisualizationLODProcessor.cpp b/Engine/Plugins/AI/MassCrowd/Source/MassCrowd/Private/MassCrowdVisualizationLODProcessor.cpp
--- a/Engine/Plugins/AI/MassCrowd/Source/MassCrowd/Private/MassCrowdVisualizationLODProcessor.cpp
+++ b/Engine/Plugins/AI/MassCrowd/Source/MassCrowd/Private/MassCrowdVisualizationLODProcessor.cpp
@@ -100,7 +100,7 @@ void UMassCrowdVisualizationLODProcessor::Execute(UMassEntitySubsystem& EntitySu
 			const TConstArrayView<FDataFragment_Transform> LocationList = Context.GetComponentView<FDataFragment_Transform>();
 			const TConstArrayView<FMassRepresentationFragment> RepresentationFragmentList = Context.GetComponentView<FMassRepresentationFragment>();
 			const TConstArrayView<FMassRepresentationLODFragment> RepresentationLODFragmentList = Context.GetComponentView<FMassRepresentationLODFragment>();
-			const int32 NumEntities = Context.GetEntitiesNum();
+			const int32 NumEntities = Context.GetNumEntities();
 			const float SpecifiedRangeSquaredCentimeters = FMath::Square(UE::MassCrowd::bDebugShowISMUnderSpecifiedRange * 100);
 			for (int EntityIdx = 0; EntityIdx < NumEntities; EntityIdx++)
 			{
diff --git a/Engine/Plugins/AI/MassCrowd/Source/MassCrowd/Private/MassCrowdVisualizationProcessor.cpp b/Engine/Plugins/AI/MassCrowd/Source/MassCrowd/Private/MassCrowdVisualizationProcessor.cpp
--- a/Engine/Plugins/AI/MassCrowd/Source/MassCrowd/Private/MassCrowdVisualizationProcessor.cpp
+++ b/Engine/Plugins/AI/MassCrowd/Source/MassCrowd/Private/MassCrowdVisualizationProcessor.cpp
@@ -57,7 +57,7 @@ void UMassCrowdVisualizationProcessor::Execute(UMassEntitySubsystem& EntitySubsy
 			const TConstArrayView<FDataFragment_Actor> ActorList = Context.GetComponentView<FDataFragment_Actor>();
 			const TConstArrayView<FDataFragment_Transform> EntityLocationList = Context.GetComponentView<FDataFragment_Transform>();
 
-			const int32 NumEntities = Context.GetEntitiesNum();
+			const int32 NumEntities = Context.GetNumEntities();
 			for (int EntityIdx = 0; EntityIdx < NumEntities; EntityIdx++)
 			{
 				const FDataFragment_Transform& EntityLocation = EntityLocationList[EntityIdx];
diff --git a/Engine/Plugins/Runtime/MassEntity/Source/MassEntity/Private/MassEntityQuery.cpp b/Engine/Plugins/Runtime/MassEntity/Source/MassEntity/Private/MassEntityQuery.cpp
--- a/Engine/Plugins/Runtime/MassEntity/Source/MassEntity/Private/MassEntityQuery.cpp
+++ b/Engine/Plugins/Runtime/MassEntity/Source/MassEntity/Private/MassEntityQuery.cpp
@@ -145,7 +145,7 @@ void FMassEntityQuery::ForEachEntityChunk(UMassEntitySubsystem& EntitySubsystem,
 		ExecutionContext.SetRequirements(Requirements, ChunkRequirements);
 		ExecutionContext.GetChunkCollection().GetArchetype().DataPtr->ExecuteFunction(ExecutionContext, ExecuteFunction, {}, ExecutionContext.GetChunkCollection());
 #if WITH_MASSENTITY_DEBUG
-		NumEntitiesToProcess = ExecutionContext.GetChunkCollection().GetArchetype().DataPtr->GetNumEntities();
+		NumEntitiesToProcess = ExecutionContext.GetNumEntities();
 #endif
 	}
 	else
@@ -161,7 +161,7 @@ void FMassEntityQuery::ForEachEntityChunk(UMassEntitySubsystem& EntitySubsystem,
 			Archetype.DataPtr->ExecuteFunction(ExecutionContext, ExecuteFunction, ArchetypeComponentMapping[i], ChunkCondition);
 			ExecutionContext.ClearComponentViews();
 #if WITH_MASSENTITY_DEBUG
-			NumEntitiesToProcess += Archetype.DataPtr->GetNumEntities();
+			NumEntitiesToProcess += ExecutionContext.GetNumEntities();
 #endif
 		}
 	}
@@ -242,7 +242,7 @@ void FMassEntityQuery::ParallelForEachEntityChunk(UMassEntitySubsystem& EntitySu
 	ExecutionContext.FlushDeferred(EntitySubsystem);
 }
 
-int32 FMassEntityQuery::GetMatchingEntitiesNum(UMassEntitySubsystem& InEntitySubsystem)
+int32 FMassEntityQuery::GetNumMatchingEntities(UMassEntitySubsystem& InEntitySubsystem)
 {
 	CacheArchetypes(InEntitySubsystem);
 	int32 TotalEntities = 0;
diff --git a/Engine/Plugins/Runtime/MassEntity/Source/MassEntity/Public/MassEntityQuery.h b/Engine/Plugins/Runtime/MassEntity/Source/MassEntity/Public/MassEntityQuery.h
--- a/Engine/Plugins/Runtime/MassEntity/Source/MassEntity/Public/MassEntityQuery.h
+++ b/Engine/Plugins/Runtime/MassEntity/Source/MassEntity/Public/MassEntityQuery.h
@@ -293,7 +293,7 @@ public:
 	 * Note that the function is not const because calling it can result in re-caching of ValidArchetypes 
 	 * @return the number of entities this given query would process if called "now"
 	 */
-	int32 GetMatchingEntitiesNum(UMassEntitySubsystem& InEntitySubsystem);
+	int32 GetNumMatchingEntities(UMassEntitySubsystem& InEntitySubsystem);
 
 	/**
 	 * Checks if any of ValidArchetypes has any entities.
diff --git a/Engine/Plugins/Runtime/MassEntity/Source/MassEntity/Public/MassEntitySubsystem.h b/Engine/Plugins/Runtime/MassEntity/Source/MassEntity/Public/MassEntitySubsystem.h
--- a/Engine/Plugins/Runtime/MassEntity/Source/MassEntity/Public/MassEntitySubsystem.h
+++ b/Engine/Plugins/Runtime/MassEntity/Source/MassEntity/Public/MassEntitySubsystem.h
@@ -250,7 +250,7 @@ public:
 
 	FScopedProcessing NewProcessingScope() { return FScopedProcessing(ProcessingScopeCount); }
 	bool IsProcessing() const { return ProcessingScopeCount > 0; }
-	FMassCommandBuffer& Defer() { return *DeferredCommandBuffer.Get(); }
+	FMassCommandBuffer& Defer() const { return *DeferredCommandBuffer.Get(); }
 
 #if WITH_MASSENTITY_DEBUG
 	void DebugPrintEntity(int32 Index, FOutputDevice& Ar, const TCHAR* InPrefix = TEXT("")) const;
@@ -393,10 +393,10 @@ public:
 	}
 
 	TSharedPtr<FMassCommandBuffer> GetSharedDeferredCommandBuffer() const { return DeferredCommandBuffer; }
-	FMassCommandBuffer& Defer() { checkSlow(DeferredCommandBuffer.IsValid()); return *DeferredCommandBuffer.Get(); }
+	FMassCommandBuffer& Defer() const { checkSlow(DeferredCommandBuffer.IsValid()); return *DeferredCommandBuffer.Get(); }
 
 	TConstArrayView<FMassEntityHandle> GetEntities() const { return EntityListView; }
-	int32 GetEntitiesNum() const { return EntityListView.Num(); }
+	int32 GetNumEntities() const { return EntityListView.Num(); }
 
 	FMassEntityHandle GetEntity(const int32 Index) const
 	{
@@ -415,14 +415,14 @@ public:
 	int32 GetChunkSerialModificationNumber() const { return ChunkSerialModificationNumber; }
 
 	template<typename T>
-	T& GetMutableChunkComponent()
+	T& GetMutableChunkComponent() const
 	{
 		static_assert(TIsDerivedFrom<T, FMassChunkFragment>::IsDerived, "Given struct doesn't represent a valid chunk fragment type. Make sure to inherit from FMassChunkFragment or one of its child-types.");
 
 		const UScriptStruct* Type = T::StaticStruct();
 		const FChunkFragmentView* FoundChunkComponentData = ChunkComponents.FindByPredicate([Type](const FChunkFragmentView& Element) { return Element.Requirement.StructType == Type; });
 		checkf(FoundChunkComponentData, TEXT("Chunk Component requirement not found: %s"), *T::StaticStruct()->GetName());
-		return FoundChunkComponentData->ChunkComponentView.template GetMutable<T>();
+		return FoundChunkComponentData->ChunkComponentView.GetMutable<T>();
 	}
 
 	template<typename T>
@@ -433,7 +433,7 @@ public:
 		const UScriptStruct* Type = T::StaticStruct();
 		const FChunkFragmentView* FoundChunkComponentData = ChunkComponents.FindByPredicate([Type](const FChunkFragmentView& Element) { return Element.Requirement.StructType == Type; } );
 		checkf(FoundChunkComponentData, TEXT("Chunk Component requirement not found: %s"), *T::StaticStruct()->GetName());
-		return FoundChunkComponentData->ChunkComponentView.template Get<T>();
+		return FoundChunkComponentData->ChunkComponentView.Get<T>();
 	}
 
 	template<typename TFragment>
diff --git a/Engine/Plugins/Runtime/MassEntity/Source/MassEntityTestSuite/Private/LWComponentTestFarmPlot.cpp b/Engine/Plugins/Runtime/MassEntity/Source/MassEntityTestSuite/Private/LWComponentTestFarmPlot.cpp
--- a/Engine/Plugins/Runtime/MassEntity/Source/MassEntityTestSuite/Private/LWComponentTestFarmPlot.cpp
+++ b/Engine/Plugins/Runtime/MassEntity/Source/MassEntityTestSuite/Private/LWComponentTestFarmPlot.cpp
@@ -26,7 +26,7 @@ void UFarmHarvestTimerSetIcon::Execute(UMassEntitySubsystem& EntitySubsystem, FM
 
 	EntityQuery.ForEachEntityChunk(EntitySubsystem, Context, [this](FMassExecutionContext& Context) {
 
-		const int32 NumEntities = Context.GetEntitiesNum();
+		const int32 NumEntities = Context.GetNumEntities();
 		TConstArrayView<FFarmGridCellData> GridCoordList = Context.GetComponentView<FFarmGridCellData>();
 		TArrayView<FFarmVisualComponent> VisualList = Context.GetMutableComponentView<FFarmVisualComponent>();
 
diff --git a/Engine/Plugins/Runtime/MassEntity/Source/MassEntityTestSuite/Private/MassProcessorTest.cpp b/Engine/Plugins/Runtime/MassEntity/Source/MassEntityTestSuite/Private/MassProcessorTest.cpp
--- a/Engine/Plugins/Runtime/MassEntity/Source/MassEntityTestSuite/Private/MassProcessorTest.cpp
+++ b/Engine/Plugins/Runtime/MassEntity/Source/MassEntityTestSuite/Private/MassProcessorTest.cpp
@@ -25,7 +25,7 @@ int32 SimpleProcessorRun(UMassEntitySubsystem& EntitySubsystem)
 		check(Processor);
 		Processor->TestGetQuery().ForEachEntityChunk(InEntitySubsystem, Context, [Processor, &EntityProcessedCount](FMassExecutionContext& Context)
 			{
-				EntityProcessedCount += Context.GetEntitiesNum();
+				EntityProcessedCount += Context.GetNumEntities();
 			});
 	};
 
diff --git a/Engine/Plugins/Runtime/MassEntity/Source/MassEntityTestSuite/Private/MassQueryTest.cpp b/Engine/Plugins/Runtime/MassEntity/Source/MassEntityTestSuite/Private/MassQueryTest.cpp
--- a/Engine/Plugins/Runtime/MassEntity/Source/MassEntityTestSuite/Private/MassQueryTest.cpp
+++ b/Engine/Plugins/Runtime/MassEntity/Source/MassEntityTestSuite/Private/MassQueryTest.cpp
@@ -84,7 +84,7 @@ struct FQueryTest_FragmentViewBinding : FEntityTestBase
 				{
 					TArrayView<FTestFragment_Float> Floats = Context.GetMutableComponentView<FTestFragment_Float>();
 
-					for (int32 i = 0; i < Context.GetEntitiesNum(); ++i)
+					for (int32 i = 0; i < Context.GetNumEntities(); ++i)
 					{
 						Floats[i].Value = 13.f;
 					}
@@ -117,10 +117,10 @@ struct FQueryTest_ExecuteSingleArchetype : FEntityTestBase
 		FMassEntityQuery Query({ FTestFragment_Float::StaticStruct() });
 		Query.ForEachEntityChunk(*EntitySubsystem, ExecContext, [&TotalProcessed](FMassExecutionContext& Context)
 			{
-				TotalProcessed += Context.GetEntitiesNum();
+				TotalProcessed += Context.GetNumEntities();
 				TArrayView<FTestFragment_Float> Floats = Context.GetMutableComponentView<FTestFragment_Float>();
 				
-				for (int32 i = 0; i < Context.GetEntitiesNum(); ++i)
+				for (int32 i = 0; i < Context.GetNumEntities(); ++i)
 				{
 					Floats[i].Value = 13.f;
 				}
@@ -161,10 +161,10 @@ struct FQueryTest_ExecuteMultipleArchetypes : FEntityTestBase
 		FMassEntityQuery Query({ FTestFragment_Float::StaticStruct() });
 		Query.ForEachEntityChunk(*EntitySubsystem, ExecContext, [&TotalProcessed](FMassExecutionContext& Context)
 			{
-				TotalProcessed += Context.GetEntitiesNum();
+				TotalProcessed += Context.GetNumEntities();
 				TArrayView<FTestFragment_Float> Floats = Context.GetMutableComponentView<FTestFragment_Float>();
 
-				for (int32 i = 0; i < Context.GetEntitiesNum(); ++i)
+				for (int32 i = 0; i < Context.GetNumEntities(); ++i)
 				{
 					Floats[i].Value = 13.f;
 				}
@@ -217,10 +217,10 @@ struct FQueryTest_ExecuteSparse : FEntityTestBase
 		FMassEntityQuery().AddRequirement<FTestFragment_Float>(EMassFragmentAccess::ReadWrite)
 			.ForEachEntityChunk(FArchetypeChunkCollection(FloatsArchetype, EntitiesToProcess), *EntitySubsystem, ExecContext, [&TotalProcessed](FMassExecutionContext& Context)
 			{
-				TotalProcessed += Context.GetEntitiesNum();
+				TotalProcessed += Context.GetNumEntities();
 				TArrayView<FTestFragment_Float> Floats = Context.GetMutableComponentView<FTestFragment_Float>();
 
-				for (int32 i = 0; i < Context.GetEntitiesNum(); ++i)
+				for (int32 i = 0; i < Context.GetNumEntities(); ++i)
 				{
 					Floats[i].Value = 13.f;
 				}
diff --git a/Engine/Plugins/Runtime/MassEntity/Source/MassEntityTestSuite/Public/LWComponentTestFarmPlot.h b/Engine/Plugins/Runtime/MassEntity/Source/MassEntityTestSuite/Public/LWComponentTestFarmPlot.h
--- a/Engine/Plugins/Runtime/MassEntity/Source/MassEntityTestSuite/Public/LWComponentTestFarmPlot.h
+++ b/Engine/Plugins/Runtime/MassEntity/Source/MassEntityTestSuite/Public/LWComponentTestFarmPlot.h
@@ -154,7 +154,7 @@ public:
 		QUICK_SCOPE_CYCLE_COUNTER(UFarmHarvestTimerSystem_Flowers_Run);
 		
 		EntityQuery.ForEachEntityChunk(EntitySubsystem, Context, [this](FMassExecutionContext& Context) {
-			const int32 NumEntities = Context.GetEntitiesNum();
+			const int32 NumEntities = Context.GetNumEntities();
 			const float WellWateredThreshold = 0.25f;
 			TArrayView<FHarvestTimerComponent> TimerList = Context.GetMutableComponentView<FHarvestTimerComponent>();
 			TConstArrayView<FFarmWaterComponent> WaterList = Context.GetComponentView<FFarmWaterComponent>();
@@ -198,7 +198,7 @@ public:
 		
 		EntityQuery.ForEachEntityChunk(EntitySubsystem, Context, [this](FMassExecutionContext& Context) {
 			
-			const int32 NumEntities = Context.GetEntitiesNum();
+			const int32 NumEntities = Context.GetNumEntities();
 			const float WellWateredThreshold = 0.25f;
 			TArrayView<FHarvestTimerComponent> TimerList = Context.GetMutableComponentView<FHarvestTimerComponent>();
 			TConstArrayView<FFarmWaterComponent> WaterList = Context.GetMutableComponentView<FFarmWaterComponent>();
@@ -232,7 +232,7 @@ public:
 		QUICK_SCOPE_CYCLE_COUNTER(UFarmHarvestTimerExpired_Run);
 		
 		EntityQuery.ForEachEntityChunk(EntitySubsystem, Context, [this](FMassExecutionContext& Context) {
-			const int32 NumEntities = Context.GetEntitiesNum();
+			const int32 NumEntities = Context.GetNumEntities();
 			TConstArrayView<FHarvestTimerComponent> TimerList = Context.GetComponentView<FHarvestTimerComponent>();
 
 			for (int32 i = 0; i < NumEntities; ++i)
diff --git a/Engine/Plugins/Runtime/MassGameplay/Source/MassActors/Private/Translators/MassCharacterMovementTranslators.cpp b/Engine/Plugins/Runtime/MassGameplay/Source/MassActors/Private/Translators/MassCharacterMovementTranslators.cpp
--- a/Engine/Plugins/Runtime/MassGameplay/Source/MassActors/Private/Translators/MassCharacterMovementTranslators.cpp
+++ b/Engine/Plugins/Runtime/MassGameplay/Source/MassActors/Private/Translators/MassCharacterMovementTranslators.cpp
@@ -36,7 +36,7 @@ void UMassCharacterMovementToMassTranslator::Execute(UMassEntitySubsystem& Entit
 		const TArrayView<FDataFragment_Transform> LocationList = Context.GetMutableComponentView<FDataFragment_Transform>();
 		const TArrayView<FMassVelocityFragment> VelocityList = Context.GetMutableComponentView<FMassVelocityFragment>();
 
-		const int32 NumEntities = Context.GetEntitiesNum();
+		const int32 NumEntities = Context.GetNumEntities();
 		
 		for (int32 i = 0; i < NumEntities; ++i)
 		{
@@ -75,7 +75,7 @@ void UMassCharacterMovementToActorTranslator::Execute(UMassEntitySubsystem& Enti
 		const TArrayView<FDataFragment_CharacterMovementComponentWrapper> ComponentList = Context.GetMutableComponentView<FDataFragment_CharacterMovementComponentWrapper>();
 		const TConstArrayView<FMassVelocityFragment> VelocityList = Context.GetComponentView<FMassVelocityFragment>();
 		
-		const int32 NumEntities = Context.GetEntitiesNum();
+		const int32 NumEntities = Context.GetNumEntities();
 
 		for (int32 i = 0; i < NumEntities; ++i)
 		{
@@ -111,7 +111,7 @@ void UMassCharacterOrientationToMassTranslator::Execute(UMassEntitySubsystem& En
 		const TConstArrayView<FDataFragment_CharacterMovementComponentWrapper> ComponentList = Context.GetComponentView<FDataFragment_CharacterMovementComponentWrapper>();
 		const TArrayView<FDataFragment_Transform> LocationList = Context.GetMutableComponentView<FDataFragment_Transform>();
 
-		const int32 NumEntities = Context.GetEntitiesNum();
+		const int32 NumEntities = Context.GetNumEntities();
 		
 		for (int32 i = 0; i < NumEntities; ++i)
 		{
@@ -151,7 +151,7 @@ void UMassCharacterOrientationToActorTranslator::Execute(UMassEntitySubsystem& E
 		const TArrayView<FDataFragment_CharacterMovementComponentWrapper> ComponentList = Context.GetMutableComponentView<FDataFragment_CharacterMovementComponentWrapper>();
 		const TConstArrayView<FDataFragment_Transform> TransformList = Context.GetComponentView<FDataFragment_Transform>();
 		
-		const int32 NumEntities = Context.GetEntitiesNum();
+		const int32 NumEntities = Context.GetNumEntities();
 
 		for (int32 i = 0; i < NumEntities; ++i)
 		{
@@ -186,7 +186,7 @@ void UMassFragmentInitializer_NavLocation::Execute(UMassEntitySubsystem& EntityS
 	EntityQuery.ForEachEntityChunk(EntitySubsystem, Context, [this](FMassExecutionContext& Context)
 		{
 			const TArrayView<FDataFragment_NavLocation> NavLocationList = Context.GetMutableComponentView<FDataFragment_NavLocation>();
-			const int32 NumEntities = Context.GetEntitiesNum();
+			const int32 NumEntities = Context.GetNumEntities();
 			for (int32 i = 0; i < NumEntities; ++i)
 			{
 				NavLocationList[i].NodeRef = INVALID_NAVNODEREF;
@@ -237,7 +237,7 @@ void UMassFragmentInitializer_Transform::Execute(UMassEntitySubsystem& EntitySub
 		int32 NumRequiredSpawnTransforms = 0;
 		EntityQuery.ForEachEntityChunk(EntitySubsystem, Context, [&NumRequiredSpawnTransforms](const FMassExecutionContext& Context)
 			{
-				NumRequiredSpawnTransforms += Context.GetEntitiesNum();
+				NumRequiredSpawnTransforms += Context.GetNumEntities();
 			});
 
 		const int32 NumToAdd = NumRequiredSpawnTransforms - NumSpawnTransforms;
@@ -257,7 +257,7 @@ void UMassFragmentInitializer_Transform::Execute(UMassEntitySubsystem& EntitySub
 		EntityQuery.ForEachEntityChunk(EntitySubsystem, Context, [&Transforms, this](FMassExecutionContext& Context)
 			{
 				const TArrayView<FDataFragment_Transform> LocationList = Context.GetMutableComponentView<FDataFragment_Transform>();
-				const int32 NumEntities = Context.GetEntitiesNum();
+				const int32 NumEntities = Context.GetNumEntities();
 				for (int32 i = 0; i < NumEntities; ++i)
 				{
 					const int32 AuxIndex = FMath::RandRange(0, Transforms.Num() - 1);
diff --git a/Engine/Plugins/Runtime/MassGameplay/Source/MassActors/Private/Translators/MassSceneComponentLocationTranslator.cpp b/Engine/Plugins/Runtime/MassGameplay/Source/MassActors/Private/Translators/MassSceneComponentLocationTranslator.cpp
--- a/Engine/Plugins/Runtime/MassGameplay/Source/MassActors/Private/Translators/MassSceneComponentLocationTranslator.cpp
+++ b/Engine/Plugins/Runtime/MassGameplay/Source/MassActors/Private/Translators/MassSceneComponentLocationTranslator.cpp
@@ -29,7 +29,7 @@ void UMassSceneComponentLocationToMassTranslator::Execute(UMassEntitySubsystem&
 		const TConstArrayView<FMassSceneComponentWrapperFragment> ComponentList = Context.GetComponentView<FMassSceneComponentWrapperFragment>();
 		const TArrayView<FDataFragment_Transform> LocationList = Context.GetMutableComponentView<FDataFragment_Transform>();
 
-		const int32 NumEntities = Context.GetEntitiesNum();
+		const int32 NumEntities = Context.GetNumEntities();
 		for (int32 i = 0; i < NumEntities; ++i)
 		{
 			if (const USceneComponent* AsComponent = ComponentList[i].Component.Get())
@@ -65,7 +65,7 @@ void UMassSceneComponentLocationToActorTranslator::Execute(UMassEntitySubsystem&
 			const TConstArrayView<FMassSceneComponentWrapperFragment> ComponentList = Context.GetComponentView<FMassSceneComponentWrapperFragment>();
 			const TArrayView<FDataFragment_Transform> LocationList = Context.GetMutableComponentView<FDataFragment_Transform>();
 
-			const int32 NumEntities = Context.GetEntitiesNum();
+			const int32 NumEntities = Context.GetNumEntities();
 			for (int32 i = 0; i < NumEntities; ++i)
 			{
 				if (USceneComponent* AsComponent = ComponentList[i].Component.Get())
diff --git a/Engine/Plugins/Runtime/MassGameplay/Source/MassGameplayDebug/Private/DebugVisLocationProcessor.cpp b/Engine/Plugins/Runtime/MassGameplay/Source/MassGameplayDebug/Private/DebugVisLocationProcessor.cpp
--- a/Engine/Plugins/Runtime/MassGameplay/Source/MassGameplayDebug/Private/DebugVisLocationProcessor.cpp
+++ b/Engine/Plugins/Runtime/MassGameplay/Source/MassGameplayDebug/Private/DebugVisLocationProcessor.cpp
@@ -39,7 +39,7 @@ void UDebugVisLocationProcessor::Execute(UMassEntitySubsystem& EntitySubsystem,
 	{
 		EntityQuery.ForEachEntityChunk(EntitySubsystem, Context, [this, &VisualDataISMCs](const FMassExecutionContext& Context)
 		{
-			const int32 NumEntities = Context.GetEntitiesNum();
+			const int32 NumEntities = Context.GetNumEntities();
 			const TConstArrayView<FDataFragment_Transform> LocationList = Context.GetComponentView<FDataFragment_Transform>();
 			const TConstArrayView<FSimDebugVisComponent> DebugVisList = Context.GetComponentView<FSimDebugVisComponent>();
 
@@ -104,7 +104,7 @@ void UMassProcessor_UpdateDebugVis::Execute(UMassEntitySubsystem& EntitySubsyste
 
 	EntityQuery.ForEachEntityChunk(EntitySubsystem, Context, [this, Debugger](FMassExecutionContext& Context)
 		{
-			const int32 NumEntities = Context.GetEntitiesNum();
+			const int32 NumEntities = Context.GetNumEntities();
 			const TConstArrayView<FDataFragment_Transform> LocationList = Context.GetComponentView<FDataFragment_Transform>();
 			const TArrayView<FDataFragment_DebugVis> DebugVisList = Context.GetMutableComponentView<FDataFragment_DebugVis>();
 			const TArrayView<FDataFragment_AgentRadius> RadiiList = Context.GetMutableComponentView<FDataFragment_AgentRadius>();
diff --git a/Engine/Plugins/Runtime/MassGameplay/Source/MassLOD/Public/MassLODCalculator.h b/Engine/Plugins/Runtime/MassGameplay/Source/MassLOD/Public/MassLODCalculator.h
--- a/Engine/Plugins/Runtime/MassGameplay/Source/MassLOD/Public/MassLODCalculator.h
+++ b/Engine/Plugins/Runtime/MassGameplay/Source/MassLOD/Public/MassLODCalculator.h
@@ -270,7 +270,7 @@ void TMassLODCalculator<FLODLogic>::CalculateLOD(FMassExecutionContext& Context,
 	}
 #endif // WITH_MASSGAMEPLAY_DEBUG
 
-	const int32 NumEntities = Context.GetEntitiesNum();
+	const int32 NumEntities = Context.GetNumEntities();
 	for (int EntityIdx = 0; EntityIdx < NumEntities; EntityIdx++)
 	{
 		// Calculate the LOD purely upon distances
@@ -447,7 +447,7 @@ template <typename FLODLogic>
 template< typename FMassLODResultInfo, typename FMassLODSourceInfo >
 void TMassLODCalculator<FLODLogic>::AdjustLODFromCount(FMassExecutionContext& Context, TArrayView<FMassLODResultInfo> LODList, TConstArrayView<FMassLODSourceInfo> ViewersInfoList)
 {
-	const int32 NumEntities = Context.GetEntitiesNum();
+	const int32 NumEntities = Context.GetNumEntities();
 	// Adjust LOD for each viewer and remember the new highest
 	for (int EntityIdx = 0; EntityIdx < NumEntities; EntityIdx++)
 	{
@@ -518,7 +518,7 @@ template <typename FLODLogic>
 template <typename FMassLODResultInfo>
 void TMassLODCalculator<FLODLogic>::ForceOffLOD(FMassExecutionContext& Context, TArrayView<FMassLODResultInfo> LODList)
 {
-	const int32 NumEntities = Context.GetEntitiesNum();
+	const int32 NumEntities = Context.GetNumEntities();
 	for (int EntityIdx = 0; EntityIdx < NumEntities; EntityIdx++)
 	{
 		FMassLODResultInfo& EntityLOD = LODList[EntityIdx];
@@ -539,7 +539,7 @@ template <typename FLODLogic>
 template< typename FMassLODResultInfo, typename TFragmentLocation >
 void TMassLODCalculator<FLODLogic>::DebugDisplayLOD(FMassExecutionContext& Context, TConstArrayView<FMassLODResultInfo> LODList, TConstArrayView<TFragmentLocation> LocationList, UWorld* World)
 {
-	const int32 NumEntities = Context.GetEntitiesNum();
+	const int32 NumEntities = Context.GetNumEntities();
 	for (int EntityIdx = 0; EntityIdx < NumEntities; EntityIdx++)
 	{
 		const TFragmentLocation& EntityLocation = LocationList[EntityIdx];
diff --git a/Engine/Plugins/Runtime/MassGameplay/Source/MassLOD/Public/MassLODCollector.h b/Engine/Plugins/Runtime/MassGameplay/Source/MassLOD/Public/MassLODCollector.h
--- a/Engine/Plugins/Runtime/MassGameplay/Source/MassLOD/Public/MassLODCollector.h
+++ b/Engine/Plugins/Runtime/MassGameplay/Source/MassLOD/Public/MassLODCollector.h
@@ -73,7 +73,7 @@ template <typename FLODLogic>
 template< typename FMassTransform, typename FMassLODSourceInfo >
 void TMassLODCollector<FLODLogic>::CollectLODInfo(FMassExecutionContext& Context, TConstArrayView<FMassTransform> LocationList, TArrayView<FMassLODSourceInfo> ViewersInfoList)
 {
-	const int32 NumEntities = Context.GetEntitiesNum();
+	const int32 NumEntities = Context.GetNumEntities();
 	for (int EntityIdx = 0; EntityIdx < NumEntities; EntityIdx++)
 	{
 		const FMassTransform& EntityLocation = LocationList[EntityIdx];
diff --git a/Engine/Plugins/Runtime/MassGameplay/Source/MassLOD/Public/MassLODTickRateController.h b/Engine/Plugins/Runtime/MassGameplay/Source/MassLOD/Public/MassLODTickRateController.h
--- a/Engine/Plugins/Runtime/MassGameplay/Source/MassLOD/Public/MassLODTickRateController.h
+++ b/Engine/Plugins/Runtime/MassGameplay/Source/MassLOD/Public/MassLODTickRateController.h
@@ -127,7 +127,7 @@ bool TMassLODTickRateController<FVariableTickChunkFragment, FLODLogic>::UpdateTi
 
 	if (bWasChunkTicked)
 	{
-		const int32 NumEntities = Context.GetEntitiesNum();
+		const int32 NumEntities = Context.GetNumEntities();
 		for (int32 Index = 0; Index < NumEntities; ++Index)
 		{
 			FMassLODResultInfo& EntityLOD = LODList[Index];
diff --git a/Engine/Plugins/Runtime/MassGameplay/Source/MassReplication/Private/MassReplicationTypes.cpp b/Engine/Plugins/Runtime/MassGameplay/Source/MassReplication/Private/MassReplicationTypes.cpp
--- a/Engine/Plugins/Runtime/MassGameplay/Source/MassReplication/Private/MassReplicationTypes.cpp
+++ b/Engine/Plugins/Runtime/MassGameplay/Source/MassReplication/Private/MassReplicationTypes.cpp
@@ -44,7 +44,7 @@ void UMassNetworkIDFragmentInitializer::Execute(UMassEntitySubsystem& EntitySubs
 		EntityQuery.ForEachEntityChunk(EntitySubsystem, Context, [this](FMassExecutionContext& Context)
 			{
 				const TArrayView<FMassNetworkIDFragment> NetworkIDList = Context.GetMutableComponentView<FMassNetworkIDFragment>();
-				const int32 NumEntities = Context.GetEntitiesNum();
+				const int32 NumEntities = Context.GetNumEntities();
 
 				for (int32 Idx = 0; Idx < NumEntities; ++Idx)
 				{
diff --git a/Engine/Plugins/Runtime/MassGameplay/Source/MassReplication/Public/MassClientBubbleHandler.h b/Engine/Plugins/Runtime/MassGameplay/Source/MassReplication/Public/MassClientBubbleHandler.h
--- a/Engine/Plugins/Runtime/MassGameplay/Source/MassReplication/Public/MassClientBubbleHandler.h
+++ b/Engine/Plugins/Runtime/MassGameplay/Source/MassReplication/Public/MassClientBubbleHandler.h
@@ -560,7 +560,7 @@ void TClientBubbleHandlerBase<AgentArrayItem>::PostReplicatedAddEntitiesHelper(c
 
 				const TArrayView<FMassNetworkIDFragment> NetworkIDList = Context.GetMutableComponentView<FMassNetworkIDFragment>();
 
-				for (int32 i = 0; i < Context.GetEntitiesNum(); ++i)
+				for (int32 i = 0; i < Context.GetNumEntities(); ++i)
 				{
 					const typename AgentArrayItem::FReplicatedAgentType& AgentSpawn = *AgentsSpawn[AgentsSpawnIdx];
 					const FMassEntityHandle Entity = Context.GetEntity(i);
diff --git a/Engine/Plugins/Runtime/MassGameplay/Source/MassReplication/Public/MassReplicationProcessorBase.h b/Engine/Plugins/Runtime/MassGameplay/Source/MassReplication/Public/MassReplicationProcessorBase.h
--- a/Engine/Plugins/Runtime/MassGameplay/Source/MassReplication/Public/MassReplicationProcessorBase.h
+++ b/Engine/Plugins/Runtime/MassGameplay/Source/MassReplication/Public/MassReplicationProcessorBase.h
@@ -95,7 +95,7 @@ void UMassReplicationProcessorBase::CalculateClientReplication(FMassExecutionCon
 	check(LODManager);
 	check(World);
 
-	const int32 NumEntities = Context.GetEntitiesNum();
+	const int32 NumEntities = Context.GetNumEntities();
 
 	TConstArrayView<FMassNetworkIDFragment> NetworkIDList = Context.GetComponentView<FMassNetworkIDFragment>();
 	TArrayView<FMassReplicationViewerLODFragment> ViewerLODList = Context.GetMutableComponentView<FMassReplicationViewerLODFragment>();
diff --git a/Engine/Plugins/Runtime/MassGameplay/Source/MassRepresentation/Private/MassRepresentationProcessor.cpp b/Engine/Plugins/Runtime/MassGameplay/Source/MassRepresentation/Private/MassRepresentationProcessor.cpp
--- a/Engine/Plugins/Runtime/MassGameplay/Source/MassRepresentation/Private/MassRepresentationProcessor.cpp
+++ b/Engine/Plugins/Runtime/MassGameplay/Source/MassRepresentation/Private/MassRepresentationProcessor.cpp
@@ -66,7 +66,7 @@ void UMassRepresentationProcessor::UpdateRepresentation(FMassExecutionContext& C
 	const TConstArrayView<FMassRepresentationLODFragment> RepresentationLODList = Context.GetComponentView<FMassRepresentationLODFragment>();
 	const TArrayView<FDataFragment_Actor> ActorList = Context.GetMutableComponentView<FDataFragment_Actor>();
 
-	const int32 NumEntities = Context.GetEntitiesNum();
+	const int32 NumEntities = Context.GetNumEntities();
 	for (int32 EntityIdx = 0; EntityIdx < NumEntities; EntityIdx++)
 	{
 		const FMassEntityHandle MassAgent = Context.GetEntity(EntityIdx);
@@ -332,7 +332,7 @@ void UMassRepresentationProcessor::UpdateVisualization(FMassExecutionContext& Co
 	const TArrayView<FMassRepresentationFragment> RepresentationList = Context.GetMutableComponentView<FMassRepresentationFragment>();
 	const TConstArrayView<FMassRepresentationLODFragment> RepresentationLODList = Context.GetComponentView<FMassRepresentationLODFragment>();
 
-	const int32 NumEntities = Context.GetEntitiesNum();
+	const int32 NumEntities = Context.GetNumEntities();
 	for (int32 EntityIdx = 0; EntityIdx < NumEntities; EntityIdx++)
 	{
 		const FMassEntityHandle Entity = Context.GetEntity(EntityIdx);
@@ -477,7 +477,7 @@ void UMassRepresentationFragmentDestructor::Execute(UMassEntitySubsystem& Entity
 		const TArrayView<FMassRepresentationFragment> RepresentationList = Context.GetMutableComponentView<FMassRepresentationFragment>();
 		const TArrayView<FDataFragment_Actor> ActorList = Context.GetMutableComponentView<FDataFragment_Actor>();
 
-		const int32 NumEntities = Context.GetEntitiesNum();
+		const int32 NumEntities = Context.GetNumEntities();
 		for (int32 i = 0; i < NumEntities; ++i)
 		{
 			FMassRepresentationFragment& Representation = RepresentationList[i];
diff --git a/Engine/Plugins/Runtime/MassGameplay/Source/MassRepresentation/Private/MassUpdateISMProcessor.cpp b/Engine/Plugins/Runtime/MassGameplay/Source/MassRepresentation/Private/MassUpdateISMProcessor.cpp
--- a/Engine/Plugins/Runtime/MassGameplay/Source/MassRepresentation/Private/MassUpdateISMProcessor.cpp
+++ b/Engine/Plugins/Runtime/MassGameplay/Source/MassRepresentation/Private/MassUpdateISMProcessor.cpp
@@ -53,7 +53,7 @@ void UMassUpdateISMProcessor::Execute(UMassEntitySubsystem& EntitySubsystem, FMa
 			const TArrayView<FMassRepresentationFragment> RepresentationList = Context.GetMutableComponentView<FMassRepresentationFragment>();
 			const TConstArrayView<FMassRepresentationLODFragment> RepresentationLODList = Context.GetComponentView<FMassRepresentationLODFragment>();
 
-			const int32 NumEntities = Context.GetEntitiesNum();
+			const int32 NumEntities = Context.GetNumEntities();
 			for (int32 EntityIdx = 0; EntityIdx < NumEntities; EntityIdx++)
 			{
 				const FDataFragment_Transform& TransformFragment = TransformList[EntityIdx];
diff --git a/Engine/Plugins/Runtime/MassGameplay/Source/MassSmartObjects/Private/MassSmartObjectProcessor.cpp b/Engine/Plugins/Runtime/MassGameplay/Source/MassSmartObjects/Private/MassSmartObjectProcessor.cpp
--- a/Engine/Plugins/Runtime/MassGameplay/Source/MassSmartObjects/Private/MassSmartObjectProcessor.cpp
+++ b/Engine/Plugins/Runtime/MassGameplay/Source/MassSmartObjects/Private/MassSmartObjectProcessor.cpp
@@ -95,7 +95,7 @@ void UMassProcessor_SmartObjectCandidatesFinder::Execute(UMassEntitySubsystem& E
 	{
 		const FSmartObjectOctree& Octree = SmartObjectSubsystem->GetOctree();
 
-		const int32 NumEntities = Context.GetEntitiesNum();
+		const int32 NumEntities = Context.GetNumEntities();
 		EntitiesToSignal.Reserve(EntitiesToSignal.Num() + NumEntities);
 
 		const TConstArrayView<FMassSmartObjectWorldLocationRequestFragment> RequestList = Context.GetComponentView<FMassSmartObjectWorldLocationRequestFragment>();
@@ -154,7 +154,7 @@ void UMassProcessor_SmartObjectCandidatesFinder::Execute(UMassEntitySubsystem& E
 	{
 		LaneRequestQuery.ForEachEntityChunk(EntitySubsystem, Context, [this, SmartObjectBehavior, &Filter, SmartObjectTag, &EntitiesToSignal, &BeginRequestProcessing, &EndRequestProcessing](FMassExecutionContext& Context)
 		{
-			const int32 NumEntities = Context.GetEntitiesNum();
+			const int32 NumEntities = Context.GetNumEntities();
 			EntitiesToSignal.Reserve(EntitiesToSignal.Num() + NumEntities);
 
 			TConstArrayView<FMassSmartObjectLaneLocationRequestFragment> RequestList = Context.GetComponentView<FMassSmartObjectLaneLocationRequestFragment>();
@@ -289,7 +289,7 @@ void UMassProcessor_SmartObjectTimedBehavior::Execute(UMassEntitySubsystem& Enti
 
 	EntityQuery.ForEachEntityChunk(EntitySubsystem, Context, [this, &ToRelease](FMassExecutionContext& Context)
 	{
-		const int32 NumEntities = Context.GetEntitiesNum();
+		const int32 NumEntities = Context.GetNumEntities();
 		const TArrayView<FDataFragment_SmartObjectUser> UserList = Context.GetMutableComponentView<FDataFragment_SmartObjectUser>();
 
 		for (int32 i = 0; i < NumEntities; ++i)
