Renamed all Component-containing functions and member variable names in Mass, replaced with "Fragment". Updated a lot of comments and logged text as well.

diff --git a/Engine/Plugins/AI/MassAI/Source/MassAIBehavior/Private/MassLookAtProcessors.cpp b/Engine/Plugins/AI/MassAI/Source/MassAIBehavior/Private/MassLookAtProcessors.cpp
--- a/Engine/Plugins/AI/MassAI/Source/MassAIBehavior/Private/MassLookAtProcessors.cpp
+++ b/Engine/Plugins/AI/MassAI/Source/MassAIBehavior/Private/MassLookAtProcessors.cpp
@@ -126,12 +126,12 @@ void UMassLookAtProcessor::Execute(UMassEntitySubsystem& EntitySubsystem, FMassE
 	EntityQuery_Conditional.ForEachEntityChunk(EntitySubsystem, Context, [this, &EntitySubsystem, CurrentTime](FMassExecutionContext& Context)
 		{
 			const int32 NumEntities = Context.GetNumEntities();
-			const TArrayView<FMassLookAtFragment> LookAtList = Context.GetMutableComponentView<FMassLookAtFragment>();
-			const TConstArrayView<FDataFragment_Transform> TransformList = Context.GetComponentView<FDataFragment_Transform>();
-			const TConstArrayView<FMassMoveTargetFragment> MoveTargetList = Context.GetComponentView<FMassMoveTargetFragment>();
-			const TConstArrayView<FMassZoneGraphLaneLocationFragment> ZoneGraphLocationList = Context.GetComponentView<FMassZoneGraphLaneLocationFragment>();
-			const TConstArrayView<FMassZoneGraphShortPathFragment> ShortPathList = Context.GetComponentView<FMassZoneGraphShortPathFragment>();
-			const TArrayView<FMassLookAtTrajectoryFragment> LookAtTrajectoryList = Context.GetMutableComponentView<FMassLookAtTrajectoryFragment>();
+			const TArrayView<FMassLookAtFragment> LookAtList = Context.GetMutableFragmentView<FMassLookAtFragment>();
+			const TConstArrayView<FDataFragment_Transform> TransformList = Context.GetFragmentView<FDataFragment_Transform>();
+			const TConstArrayView<FMassMoveTargetFragment> MoveTargetList = Context.GetFragmentView<FMassMoveTargetFragment>();
+			const TConstArrayView<FMassZoneGraphLaneLocationFragment> ZoneGraphLocationList = Context.GetFragmentView<FMassZoneGraphLaneLocationFragment>();
+			const TConstArrayView<FMassZoneGraphShortPathFragment> ShortPathList = Context.GetFragmentView<FMassZoneGraphShortPathFragment>();
+			const TArrayView<FMassLookAtTrajectoryFragment> LookAtTrajectoryList = Context.GetMutableFragmentView<FMassLookAtTrajectoryFragment>();
 
 			for (int32 i = 0; i < NumEntities; ++i)
 			{
@@ -280,7 +280,7 @@ void UMassLookAtProcessor::FindNewGazeTarget(const UMassEntitySubsystem& EntityS
 			}
 
 			// TargetTag is added through the LookAtTargetTrait and Transform was added with it
-			const FDataFragment_Transform& TargetTransform = EntityView.GetComponentData<FDataFragment_Transform>();
+			const FDataFragment_Transform& TargetTransform = EntityView.GetFragmentData<FDataFragment_Transform>();
 			const FVector TargetLocation = TargetTransform.GetTransform().GetLocation();
 			FVector Direction = (TargetLocation - Location).GetSafeNormal();
 			Direction = Transform.InverseTransformVector(Direction);
@@ -295,7 +295,7 @@ void UMassLookAtProcessor::FindNewGazeTarget(const UMassEntitySubsystem& EntityS
 			}
 
 			// Allow to pick entities out of view if they are moving towards us.
-			if (const FMassVelocityFragment* Velocity = EntityView.GetComponentDataPtr<FMassVelocityFragment>())
+			if (const FMassVelocityFragment* Velocity = EntityView.GetFragmentDataPtr<FMassVelocityFragment>())
 			{
 				const FVector MoveDirection = Transform.InverseTransformVector(Velocity->Value.GetSafeNormal());
 				
@@ -354,7 +354,7 @@ void UMassLookAtProcessor::UpdateLookAtTrackedEntity(const UMassEntitySubsystem&
 	// Update direction toward target
 	if (EntitySubsystem.IsEntityValid(LookAt.TrackedEntity))
 	{
-		if (const FDataFragment_Transform* TargetTransform = EntitySubsystem.GetComponentDataPtr<FDataFragment_Transform>(LookAt.TrackedEntity))
+		if (const FDataFragment_Transform* TargetTransform = EntitySubsystem.GetFragmentDataPtr<FDataFragment_Transform>(LookAt.TrackedEntity))
 		{
 			const FVector AgentPosition = Transform.GetLocation();
 			const FVector NewGlobalDirection = (TargetTransform->GetTransform().GetLocation() - AgentPosition).GetSafeNormal();
@@ -378,7 +378,7 @@ bool UMassLookAtProcessor::UpdateGazeTrackedEntity(const UMassEntitySubsystem& E
 	// Update direction toward gaze target
 	if (LookAt.GazeTrackedEntity.IsSet() && EntitySubsystem.IsEntityValid(LookAt.GazeTrackedEntity))
 	{
-		if (const FDataFragment_Transform* TargetTransform = EntitySubsystem.GetComponentDataPtr<FDataFragment_Transform>(LookAt.GazeTrackedEntity))
+		if (const FDataFragment_Transform* TargetTransform = EntitySubsystem.GetFragmentDataPtr<FDataFragment_Transform>(LookAt.GazeTrackedEntity))
 		{
 			const FVector AgentPosition = Transform.GetLocation();
 			const FVector NewGlobalDirection = (TargetTransform->GetTransform().GetLocation() - AgentPosition).GetSafeNormal();
diff --git a/Engine/Plugins/AI/MassAI/Source/MassAIBehavior/Private/MassStateTreeProcessors.cpp b/Engine/Plugins/AI/MassAI/Source/MassAIBehavior/Private/MassStateTreeProcessors.cpp
--- a/Engine/Plugins/AI/MassAI/Source/MassAIBehavior/Private/MassStateTreeProcessors.cpp
+++ b/Engine/Plugins/AI/MassAI/Source/MassAIBehavior/Private/MassStateTreeProcessors.cpp
@@ -27,7 +27,7 @@ bool SetExternalFragments(FMassStateTreeExecutionContext& Context, const UMassEn
 		if (ItemDesc.Struct && ItemDesc.Struct->IsChildOf(FMassFragment::StaticStruct()))
 		{
 			const UScriptStruct* ScriptStruct = Cast<const UScriptStruct>(ItemDesc.Struct);
-			FStructView Fragment = EntityView.GetComponentDataStruct(ScriptStruct);
+			FStructView Fragment = EntityView.GetFragmentDataStruct(ScriptStruct);
 			if (Fragment.IsValid())
 			{
 				Context.SetExternalItem(ItemDesc.Handle, FStateTreeItemView(Fragment));
@@ -83,7 +83,7 @@ void ForEachEntityInChunk(
 	const TFunctionRef<void(FMassStateTreeExecutionContext&, FStateTreeItemView)> ForEachEntityCallback)
 {
 	const FMassExecutionContext& Context = StateTreeContext.GetEntitySubsystemExecutionContext();
-	const TConstArrayView<FMassStateTreeFragment> StateTreeList = Context.GetComponentView<FMassStateTreeFragment>();
+	const TConstArrayView<FMassStateTreeFragment> StateTreeList = Context.GetFragmentView<FMassStateTreeFragment>();
 
 	// Assuming that all the entities share same StateTree, because they all should have the same storage fragment.
 	const int32 NumEntities = Context.GetNumEntities();
@@ -131,7 +131,7 @@ void ForEachEntityInChunk(
 			}
 		}
 
-		ForEachEntityCallback(StateTreeContext, StateTreeContext.GetEntitySubsystem().GetComponentDataStruct(Entity, StorageScriptStruct));
+		ForEachEntityCallback(StateTreeContext, StateTreeContext.GetEntitySubsystem().GetFragmentDataStruct(Entity, StorageScriptStruct));
 	}
 }
 
@@ -296,7 +296,7 @@ void UMassStateTreeProcessor::SignalEntities(UMassEntitySubsystem& EntitySubsyst
 			// Keep stats regarding the amount of tree instances ticked per frame
 			CSV_CUSTOM_STAT(StateTreeProcessor, NumTickedStateTree, Context.GetNumEntities(), ECsvCustomStatOp::Accumulate);
 
-			TArrayView<FMassStateTreeFragment> StateTreeList = Context.GetMutableComponentView<FMassStateTreeFragment>();
+			TArrayView<FMassStateTreeFragment> StateTreeList = Context.GetMutableFragmentView<FMassStateTreeFragment>();
 
 			UE::MassBehavior::ForEachEntityInChunk(
 				StateTreeContext,
diff --git a/Engine/Plugins/AI/MassAI/Source/MassAIBehavior/Private/MassZoneGraphAnnotationFragments.cpp b/Engine/Plugins/AI/MassAI/Source/MassAIBehavior/Private/MassZoneGraphAnnotationFragments.cpp
--- a/Engine/Plugins/AI/MassAI/Source/MassAIBehavior/Private/MassZoneGraphAnnotationFragments.cpp
+++ b/Engine/Plugins/AI/MassAI/Source/MassAIBehavior/Private/MassZoneGraphAnnotationFragments.cpp
@@ -22,7 +22,7 @@ static const float OffLODMaxUpdateInterval = 2.10f;
 
 bool FMassZoneGraphAnnotationVariableTickChunkFragment::UpdateChunk(FMassExecutionContext& Context)
 {
-	FMassZoneGraphAnnotationVariableTickChunkFragment& ChunkFrag = Context.GetMutableChunkComponent<FMassZoneGraphAnnotationVariableTickChunkFragment>();
+	FMassZoneGraphAnnotationVariableTickChunkFragment& ChunkFrag = Context.GetMutableChunkFragment<FMassZoneGraphAnnotationVariableTickChunkFragment>();
 	ChunkFrag.TimeUntilNextTick -= Context.GetDeltaTimeSeconds();
 	if (ChunkFrag.TimeUntilNextTick <= 0.0f)
 	{
diff --git a/Engine/Plugins/AI/MassAI/Source/MassAIBehavior/Private/MassZoneGraphAnnotationProcessors.cpp b/Engine/Plugins/AI/MassAI/Source/MassAIBehavior/Private/MassZoneGraphAnnotationProcessors.cpp
--- a/Engine/Plugins/AI/MassAI/Source/MassAIBehavior/Private/MassZoneGraphAnnotationProcessors.cpp
+++ b/Engine/Plugins/AI/MassAI/Source/MassAIBehavior/Private/MassZoneGraphAnnotationProcessors.cpp
@@ -36,8 +36,8 @@ void UMassZoneGraphAnnotationTagsInitializer::Execute(UMassEntitySubsystem& Enti
 	EntityQuery.ForEachEntityChunk(EntitySubsystem, Context, [this](FMassExecutionContext& Context)
 	{
 		const int32 NumEntities = Context.GetNumEntities();
-		const TArrayView<FMassZoneGraphAnnotationTagsFragment> AnnotationTagsList = Context.GetMutableComponentView<FMassZoneGraphAnnotationTagsFragment>();
-		const TConstArrayView<FMassZoneGraphLaneLocationFragment> LaneLocationList = Context.GetComponentView<FMassZoneGraphLaneLocationFragment>();
+		const TArrayView<FMassZoneGraphAnnotationTagsFragment> AnnotationTagsList = Context.GetMutableFragmentView<FMassZoneGraphAnnotationTagsFragment>();
+		const TConstArrayView<FMassZoneGraphLaneLocationFragment> LaneLocationList = Context.GetFragmentView<FMassZoneGraphLaneLocationFragment>();
 
 		for (int32 EntityIndex = 0; EntityIndex < NumEntities; ++EntityIndex)
 		{
@@ -99,8 +99,8 @@ void UMassZoneGraphAnnotationTagUpdateProcessor::Execute(UMassEntitySubsystem& E
 		}
 
 		const int32 NumEntities = Context.GetNumEntities();
-		const TArrayView<FMassZoneGraphAnnotationTagsFragment> AnnotationTagsList = Context.GetMutableComponentView<FMassZoneGraphAnnotationTagsFragment>();
-		const TConstArrayView<FMassZoneGraphLaneLocationFragment> LaneLocationList = Context.GetComponentView<FMassZoneGraphLaneLocationFragment>();
+		const TArrayView<FMassZoneGraphAnnotationTagsFragment> AnnotationTagsList = Context.GetMutableFragmentView<FMassZoneGraphAnnotationTagsFragment>();
+		const TConstArrayView<FMassZoneGraphLaneLocationFragment> LaneLocationList = Context.GetFragmentView<FMassZoneGraphLaneLocationFragment>();
 
 		for (int32 EntityIndex = 0; EntityIndex < NumEntities; ++EntityIndex)
 		{
@@ -146,8 +146,8 @@ void UMassZoneGraphAnnotationTagUpdateProcessor::SignalEntities(UMassEntitySubsy
 	EntityQuery.ForEachEntityChunk(EntitySubsystem, Context, [this](FMassExecutionContext& Context)
 	{
 		const int32 NumEntities = Context.GetNumEntities();
-		const TConstArrayView<FMassZoneGraphLaneLocationFragment> LaneLocationList = Context.GetComponentView<FMassZoneGraphLaneLocationFragment>();
-		const TArrayView<FMassZoneGraphAnnotationTagsFragment> AnnotationTagsList = Context.GetMutableComponentView<FMassZoneGraphAnnotationTagsFragment>();
+		const TConstArrayView<FMassZoneGraphLaneLocationFragment> LaneLocationList = Context.GetFragmentView<FMassZoneGraphLaneLocationFragment>();
+		const TArrayView<FMassZoneGraphAnnotationTagsFragment> AnnotationTagsList = Context.GetMutableFragmentView<FMassZoneGraphAnnotationTagsFragment>();
 
 		for (int32 EntityIndex = 0; EntityIndex < NumEntities; ++EntityIndex)
 		{
diff --git a/Engine/Plugins/AI/MassAI/Source/MassAIDebug/Private/GameplayDebuggerCategory_Mass.cpp b/Engine/Plugins/AI/MassAI/Source/MassAIDebug/Private/GameplayDebuggerCategory_Mass.cpp
--- a/Engine/Plugins/AI/MassAI/Source/MassAIDebug/Private/GameplayDebuggerCategory_Mass.cpp
+++ b/Engine/Plugins/AI/MassAI/Source/MassAIDebug/Private/GameplayDebuggerCategory_Mass.cpp
@@ -274,7 +274,7 @@ void FGameplayDebuggerCategory_Mass::CollectData(APlayerController* OwnerPC, AAc
 
 				}
 
-				FDataFragment_Transform& TransformFragment = EntitySystem->GetComponentDataChecked<FDataFragment_Transform>(CachedEntity);
+				FDataFragment_Transform& TransformFragment = EntitySystem->GetFragmentDataChecked<FDataFragment_Transform>(CachedEntity);
 				const float CapsuleRadius = 50.f;
 				AddShape(FGameplayDebuggerShape::MakeCapsule(TransformFragment.GetTransform().GetLocation() + 2.f * CapsuleRadius * FVector::UpVector, CapsuleRadius, CapsuleRadius * 2.f, FColor::Orange));
 			}
@@ -319,17 +319,17 @@ void FGameplayDebuggerCategory_Mass::CollectData(APlayerController* OwnerPC, AAc
 			EntityQuery.ForEachEntityChunk(*EntitySystem, Context, [this, Debugger, MassStateTreeSubsystem, EntitySystem, OwnerPC, ViewLocation, ViewDirection, CurrentTime](FMassExecutionContext& Context)
 			{
 				const int32 NumEntities = Context.GetNumEntities();
-				const TConstArrayView<FMassStateTreeFragment> StateTreeList = Context.GetComponentView<FMassStateTreeFragment>();
-				const TConstArrayView<FDataFragment_Transform> TransformList = Context.GetComponentView<FDataFragment_Transform>();
-				const TConstArrayView<FDataFragment_AgentRadius> RadiusList = Context.GetComponentView<FDataFragment_AgentRadius>();
-				const TConstArrayView<FMassSteeringFragment> SteeringList = Context.GetComponentView<FMassSteeringFragment>();
-				const TConstArrayView<FMassSteeringGhostFragment> GhostList = Context.GetComponentView<FMassSteeringGhostFragment>();
-				const TConstArrayView<FMassVelocityFragment> VelocityList = Context.GetComponentView<FMassVelocityFragment>();
-				const TConstArrayView<FMassMoveTargetFragment> MoveTargetList = Context.GetComponentView<FMassMoveTargetFragment>();
-				const TConstArrayView<FMassLookAtFragment> LookAtList = Context.GetComponentView<FMassLookAtFragment>();
-				const TConstArrayView<FMassLookAtTrajectoryFragment> LookAtTrajectoryList = Context.GetComponentView<FMassLookAtTrajectoryFragment>();
-				const TConstArrayView<FMassSimulationLODFragment> SimLODList = Context.GetComponentView<FMassSimulationLODFragment>();
-				const TConstArrayView<FMassZoneGraphShortPathFragment> ShortPathList = Context.GetComponentView<FMassZoneGraphShortPathFragment>();
+				const TConstArrayView<FMassStateTreeFragment> StateTreeList = Context.GetFragmentView<FMassStateTreeFragment>();
+				const TConstArrayView<FDataFragment_Transform> TransformList = Context.GetFragmentView<FDataFragment_Transform>();
+				const TConstArrayView<FDataFragment_AgentRadius> RadiusList = Context.GetFragmentView<FDataFragment_AgentRadius>();
+				const TConstArrayView<FMassSteeringFragment> SteeringList = Context.GetFragmentView<FMassSteeringFragment>();
+				const TConstArrayView<FMassSteeringGhostFragment> GhostList = Context.GetFragmentView<FMassSteeringGhostFragment>();
+				const TConstArrayView<FMassVelocityFragment> VelocityList = Context.GetFragmentView<FMassVelocityFragment>();
+				const TConstArrayView<FMassMoveTargetFragment> MoveTargetList = Context.GetFragmentView<FMassMoveTargetFragment>();
+				const TConstArrayView<FMassLookAtFragment> LookAtList = Context.GetFragmentView<FMassLookAtFragment>();
+				const TConstArrayView<FMassLookAtTrajectoryFragment> LookAtTrajectoryList = Context.GetFragmentView<FMassLookAtTrajectoryFragment>();
+				const TConstArrayView<FMassSimulationLODFragment> SimLODList = Context.GetFragmentView<FMassSimulationLODFragment>();
+				const TConstArrayView<FMassZoneGraphShortPathFragment> ShortPathList = Context.GetFragmentView<FMassZoneGraphShortPathFragment>();
 
 				constexpr float MaxViewDistance = 25000.0f;
 				constexpr float MinViewDirDot = 0.707f; // 45 degrees
@@ -391,7 +391,7 @@ void FGameplayDebuggerCategory_Mass::CollectData(APlayerController* OwnerPC, AAc
 					bool bLookArrowDrawn = false;
 					if (LookAt.LookAtMode == EMassLookAtMode::LookAtEntity && EntitySystem->IsEntityValid(LookAt.TrackedEntity))
 					{
-						if (const FDataFragment_Transform* TargetTransform = EntitySystem->GetComponentDataPtr<FDataFragment_Transform>(LookAt.TrackedEntity))
+						if (const FDataFragment_Transform* TargetTransform = EntitySystem->GetFragmentDataPtr<FDataFragment_Transform>(LookAt.TrackedEntity))
 						{
 							FVector TargetPosition = TargetTransform->GetTransform().GetLocation();
 							TargetPosition.Z = BasePos.Z;
@@ -405,7 +405,7 @@ void FGameplayDebuggerCategory_Mass::CollectData(APlayerController* OwnerPC, AAc
 					
 					if (LookAt.bRandomGazeEntities && EntitySystem->IsEntityValid(LookAt.GazeTrackedEntity))
 					{
-						if (const FDataFragment_Transform* TargetTransform = EntitySystem->GetComponentDataPtr<FDataFragment_Transform>(LookAt.GazeTrackedEntity))
+						if (const FDataFragment_Transform* TargetTransform = EntitySystem->GetFragmentDataPtr<FDataFragment_Transform>(LookAt.GazeTrackedEntity))
 						{
 							FVector TargetPosition = TargetTransform->GetTransform().GetLocation();
 							TargetPosition.Z = BasePos.Z;
@@ -490,7 +490,7 @@ void FGameplayDebuggerCategory_Mass::CollectData(APlayerController* OwnerPC, AAc
 							StateTreeContext.Init(*OwnerPC, *StateTree, EStateTreeStorage::External);
 							StateTreeContext.SetEntity(Entity);
 							
-							FStructView Storage = EntitySystem->GetComponentDataStruct(Entity, StateTree->GetRuntimeStorageStruct());
+							FStructView Storage = EntitySystem->GetFragmentDataStruct(Entity, StateTree->GetRuntimeStorageStruct());
 							
 							Status += StateTreeContext.GetActiveStateName(Storage);
 							Status += TEXT("\n");
@@ -528,8 +528,8 @@ void FGameplayDebuggerCategory_Mass::CollectData(APlayerController* OwnerPC, AAc
 			EntityColliderQuery.ForEachEntityChunk(*EntitySystem, Context, [this, Debugger, EntitySystem, OwnerPC, ViewLocation, ViewDirection](FMassExecutionContext& Context)
 			{
 				const int32 NumEntities = Context.GetNumEntities();
-				const TConstArrayView<FDataFragment_Transform> TransformList = Context.GetComponentView<FDataFragment_Transform>();
-				const TConstArrayView<FMassAvoidanceColliderFragment> CollidersList = Context.GetComponentView<FMassAvoidanceColliderFragment>();
+				const TConstArrayView<FDataFragment_Transform> TransformList = Context.GetFragmentView<FDataFragment_Transform>();
+				const TConstArrayView<FMassAvoidanceColliderFragment> CollidersList = Context.GetFragmentView<FMassAvoidanceColliderFragment>();
 	
 				constexpr float MaxViewDistance = 25000.0f;
 				constexpr float MinViewDirDot = 0.707f; // 45 degrees
diff --git a/Engine/Plugins/AI/MassAI/Source/MassAIDebug/Private/MassDebugStateTreeProcessor.cpp b/Engine/Plugins/AI/MassAI/Source/MassAIDebug/Private/MassDebugStateTreeProcessor.cpp
--- a/Engine/Plugins/AI/MassAI/Source/MassAIDebug/Private/MassDebugStateTreeProcessor.cpp
+++ b/Engine/Plugins/AI/MassAI/Source/MassAIDebug/Private/MassDebugStateTreeProcessor.cpp
@@ -54,8 +54,8 @@ void UMassDebugStateTreeProcessor::Execute(UMassEntitySubsystem& EntitySubsystem
 		{
 			const FMassEntityHandle SelectedEntity = Debugger->GetSelectedEntity();
 			const int32 NumEntities = Context.GetNumEntities();
-			const TConstArrayView<FMassStateTreeFragment> StateTreeList = Context.GetComponentView<FMassStateTreeFragment>();
-			const TConstArrayView<FDataFragment_Transform> TransformList = Context.GetComponentView<FDataFragment_Transform>();
+			const TConstArrayView<FMassStateTreeFragment> StateTreeList = Context.GetFragmentView<FMassStateTreeFragment>();
+			const TConstArrayView<FDataFragment_Transform> TransformList = Context.GetFragmentView<FDataFragment_Transform>();
 			const UStateTree* StateTree = MassStateTreeSubsystem->GetRegisteredStateTreeAsset(StateTreeList[0].StateTreeHandle);
 
 			// Not reporting error since this processor is a debug tool 
@@ -74,7 +74,7 @@ void UMassDebugStateTreeProcessor::Execute(UMassEntitySubsystem& EntitySubsystem
 					StateTreeContext.SetEntity(Entity);
 					
 #if WITH_GAMEPLAY_DEBUGGER
-					const FStructView Storage = EntitySubsystem.GetComponentDataStruct(Entity, StateTree->GetRuntimeStorageStruct());
+					const FStructView Storage = EntitySubsystem.GetFragmentDataStruct(Entity, StateTree->GetRuntimeStorageStruct());
 					Debugger->AppendSelectedEntityInfo(StateTreeContext.GetDebugInfoString(Storage));
 #endif // WITH_GAMEPLAY_DEBUGGER
 				}
@@ -92,7 +92,7 @@ void UMassDebugStateTreeProcessor::Execute(UMassEntitySubsystem& EntitySubsystem
 					StateTreeContext.Init(*this, *StateTree, EStateTreeStorage::External);
 					StateTreeContext.SetEntity(Entity);
 
-					const FStructView Storage = EntitySubsystem.GetComponentDataStruct(Entity, StateTree->GetRuntimeStorageStruct());
+					const FStructView Storage = EntitySubsystem.GetFragmentDataStruct(Entity, StateTree->GetRuntimeStorageStruct());
 
 					// State
 					UE_VLOG_SEGMENT_THICK(this, LogStateTree, Log, Position, Position + FVector(0,0,50), EntityColor, 2.0f, TEXT("%s %s"),
diff --git a/Engine/Plugins/AI/MassAI/Source/MassAIMovement/Private/AssignRandomNavLocationProcessor.cpp b/Engine/Plugins/AI/MassAI/Source/MassAIMovement/Private/AssignRandomNavLocationProcessor.cpp
--- a/Engine/Plugins/AI/MassAI/Source/MassAIMovement/Private/AssignRandomNavLocationProcessor.cpp
+++ b/Engine/Plugins/AI/MassAI/Source/MassAIMovement/Private/AssignRandomNavLocationProcessor.cpp
@@ -30,8 +30,8 @@ void UAssignRandomNavLocationProcessor::Execute(UMassEntitySubsystem& EntitySubs
 
 	EntityQuery.ForEachEntityChunk(EntitySubsystem, Context, [this](FMassExecutionContext& Context)
 		{
-			TArrayView<FDataFragment_Transform> LocationList = Context.GetMutableComponentView<FDataFragment_Transform>();
-			TArrayView<FDataFragment_NavLocation> NavLocationList = Context.GetMutableComponentView<FDataFragment_NavLocation>();
+			TArrayView<FDataFragment_Transform> LocationList = Context.GetMutableFragmentView<FDataFragment_Transform>();
+			TArrayView<FDataFragment_NavLocation> NavLocationList = Context.GetMutableFragmentView<FDataFragment_NavLocation>();
 			const int32 NumEntities = Context.GetNumEntities();
 			TArray<FNavLocation> Locations;
 			Locations.AddDefaulted(NumEntities);
diff --git a/Engine/Plugins/AI/MassAI/Source/MassAIMovement/Private/MassAvoidanceProcessors.cpp b/Engine/Plugins/AI/MassAI/Source/MassAIMovement/Private/MassAvoidanceProcessors.cpp
--- a/Engine/Plugins/AI/MassAI/Source/MassAIMovement/Private/MassAvoidanceProcessors.cpp
+++ b/Engine/Plugins/AI/MassAI/Source/MassAIMovement/Private/MassAvoidanceProcessors.cpp
@@ -498,14 +498,14 @@ void UMassAvoidanceProcessor::Execute(UMassEntitySubsystem& EntitySubsystem, FMa
 		const int32 NumEntities = Ctx.GetNumEntities();
 		const float DistanceCutOffSquare = FMath::Square(UE::MassAvoidance::Tweakables::AgentDetectionDistance);
 
-		const TArrayView<FMassSteeringFragment> SteeringList = Ctx.GetMutableComponentView<FMassSteeringFragment>();
-		const TConstArrayView<FMassNavigationEdgesFragment> NavEdgesList = Ctx.GetComponentView<FMassNavigationEdgesFragment>();
-		const TConstArrayView<FDataFragment_Transform> LocationList = Ctx.GetComponentView<FDataFragment_Transform>();
-		const TConstArrayView<FMassVelocityFragment> VelocityList = Ctx.GetComponentView<FMassVelocityFragment>();
-		const TConstArrayView<FDataFragment_AgentRadius> RadiusList = Ctx.GetComponentView<FDataFragment_AgentRadius>();
-		const TConstArrayView<FMassSimulationLODFragment> SimLODList = Ctx.GetComponentView<FMassSimulationLODFragment>();
-		const TConstArrayView<FMassMoveTargetFragment> MoveTargetList = Ctx.GetComponentView<FMassMoveTargetFragment>();
-		const TConstArrayView<FMassMovementConfigFragment> MovementConfigList = Ctx.GetComponentView<FMassMovementConfigFragment>();
+		const TArrayView<FMassSteeringFragment> SteeringList = Ctx.GetMutableFragmentView<FMassSteeringFragment>();
+		const TConstArrayView<FMassNavigationEdgesFragment> NavEdgesList = Ctx.GetFragmentView<FMassNavigationEdgesFragment>();
+		const TConstArrayView<FDataFragment_Transform> LocationList = Ctx.GetFragmentView<FDataFragment_Transform>();
+		const TConstArrayView<FMassVelocityFragment> VelocityList = Ctx.GetFragmentView<FMassVelocityFragment>();
+		const TConstArrayView<FDataFragment_AgentRadius> RadiusList = Ctx.GetFragmentView<FDataFragment_AgentRadius>();
+		const TConstArrayView<FMassSimulationLODFragment> SimLODList = Ctx.GetFragmentView<FMassSimulationLODFragment>();
+		const TConstArrayView<FMassMoveTargetFragment> MoveTargetList = Ctx.GetFragmentView<FMassMoveTargetFragment>();
+		const TConstArrayView<FMassMovementConfigFragment> MovementConfigList = Ctx.GetFragmentView<FMassMovementConfigFragment>();
 
 		// Arrays used to store close agents
 		TArray<FMassAvoidanceObstacleItem, TFixedAllocator<UE::MassAvoidance::MaxAgentResults>> CloseEntities;
@@ -848,7 +848,7 @@ void UMassAvoidanceProcessor::Execute(UMassEntitySubsystem& EntitySubsystem, FMa
 				}
 				
 				// Skip too far
-				const FTransform& Transform = EntitySubsystem.GetComponentDataChecked<FDataFragment_Transform>(OtherEntity.Entity).GetTransform();
+				const FTransform& Transform = EntitySubsystem.GetFragmentDataChecked<FDataFragment_Transform>(OtherEntity.Entity).GetTransform();
 				const FVector OtherLocation = Transform.GetLocation();
 				
 				const float SqDist = FVector::DistSquared(AgentLocation, OtherLocation);
@@ -883,18 +883,18 @@ void UMassAvoidanceProcessor::Execute(UMassEntitySubsystem& EntitySubsystem, FMa
 				FSortingAgent& OtherAgent = ClosestAgents[Index];
 				FMassEntityView OtherEntityView(EntitySubsystem, OtherAgent.ObstacleItem.Entity);
 
-				const FMassVelocityFragment* OtherVelocityFragment = OtherEntityView.GetComponentDataPtr<FMassVelocityFragment>();
+				const FMassVelocityFragment* OtherVelocityFragment = OtherEntityView.GetFragmentDataPtr<FMassVelocityFragment>();
 				const FVector OtherVelocity = OtherVelocityFragment != nullptr ? OtherVelocityFragment->Value : FVector::ZeroVector; // Get velocity from FAvoidanceComponent
 				const bool bExtendToEdge = OtherEntityView.HasTag<FMassAvoidanceExtendToEdgeObstacleTag>();
 
 				// @todo: this is heavy fragment to access, see if we could handle this differently.
-				const FMassMoveTargetFragment* OtherMoveTarget = OtherEntityView.GetComponentDataPtr<FMassMoveTargetFragment>();
+				const FMassMoveTargetFragment* OtherMoveTarget = OtherEntityView.GetFragmentDataPtr<FMassMoveTargetFragment>();
 				const bool bOtherIsMoving = OtherMoveTarget ? OtherMoveTarget->GetCurrentAction() == EMassMovementAction::Move : true; // Assume moving if other does not have move target.
 				
 				// Check for colliders data
 				if (EnumHasAnyFlags(OtherAgent.ObstacleItem.ItemFlags, EMassAvoidanceObstacleItemFlags::HasColliderData))
 				{
-					const FMassAvoidanceColliderFragment* ColliderFragment = OtherEntityView.GetComponentDataPtr<FMassAvoidanceColliderFragment>();
+					const FMassAvoidanceColliderFragment* ColliderFragment = OtherEntityView.GetFragmentDataPtr<FMassAvoidanceColliderFragment>();
 					if (ColliderFragment)
 					{
 						if (ColliderFragment->Type == EMassColliderType::Circle)
@@ -934,7 +934,7 @@ void UMassAvoidanceProcessor::Execute(UMassEntitySubsystem& EntitySubsystem, FMa
 					FCollider Collider;
 					Collider.Location = OtherAgent.LocationCached;
 					Collider.Velocity = OtherVelocity;
-					Collider.Radius = OtherEntityView.GetComponentData<FDataFragment_AgentRadius>().Radius;
+					Collider.Radius = OtherEntityView.GetFragmentData<FDataFragment_AgentRadius>().Radius;
 					Collider.bExtendToEdge = bExtendToEdge;
 					Collider.bIsMoving = bOtherIsMoving;
 					Colliders.Add(Collider);
@@ -1182,10 +1182,10 @@ void UMassStandingAvoidanceProcessor::Execute(UMassEntitySubsystem& EntitySubsys
 		const float DeltaTime = Context.GetDeltaTimeSeconds();
 		const float DistanceCutOffSquare = FMath::Square(UE::MassAvoidance::Tweakables::AgentDetectionDistance);
 
-		const TArrayView<FMassSteeringGhostFragment> GhostList = Context.GetMutableComponentView<FMassSteeringGhostFragment>();
-		const TConstArrayView<FDataFragment_Transform> LocationList = Context.GetComponentView<FDataFragment_Transform>();
-		const TConstArrayView<FDataFragment_AgentRadius> RadiusList = Context.GetComponentView<FDataFragment_AgentRadius>();
-		const TConstArrayView<FMassMoveTargetFragment> MoveTargetList = Context.GetComponentView<FMassMoveTargetFragment>();
+		const TArrayView<FMassSteeringGhostFragment> GhostList = Context.GetMutableFragmentView<FMassSteeringGhostFragment>();
+		const TConstArrayView<FDataFragment_Transform> LocationList = Context.GetFragmentView<FDataFragment_Transform>();
+		const TConstArrayView<FDataFragment_AgentRadius> RadiusList = Context.GetFragmentView<FDataFragment_AgentRadius>();
+		const TConstArrayView<FMassMoveTargetFragment> MoveTargetList = Context.GetFragmentView<FMassMoveTargetFragment>();
 
 		// Arrays used to store close agents
 		TArray<FMassAvoidanceObstacleItem, TFixedAllocator<UE::MassAvoidance::MaxAgentResults>> CloseEntities;
@@ -1270,7 +1270,7 @@ void UMassStandingAvoidanceProcessor::Execute(UMassEntitySubsystem& EntitySubsys
 				}
 
 				// Skip too far
-				const FDataFragment_Transform& OtherTransform = EntitySubsystem.GetComponentDataChecked<FDataFragment_Transform>(OtherEntity.Entity);
+				const FDataFragment_Transform& OtherTransform = EntitySubsystem.GetFragmentDataChecked<FDataFragment_Transform>(OtherEntity.Entity);
 				const FVector OtherLocation = OtherTransform.GetTransform().GetLocation();
 				const float DistSq = FVector::DistSquared(AgentLocation, OtherLocation);
 				if (DistSq > DistanceCutOffSquare)
@@ -1290,8 +1290,8 @@ void UMassStandingAvoidanceProcessor::Execute(UMassEntitySubsystem& EntitySubsys
 				FSortingAgent& OtherAgent = ClosestAgents[Index];
 				FMassEntityView OtherEntityView(EntitySubsystem, OtherAgent.Entity);
 
-				const FMassVelocityFragment* OtherVelocityFragment = OtherEntityView.GetComponentDataPtr<FMassVelocityFragment>();
-				const float OtherRadius = OtherEntityView.GetComponentData<FDataFragment_AgentRadius>().Radius;
+				const FMassVelocityFragment* OtherVelocityFragment = OtherEntityView.GetFragmentDataPtr<FMassVelocityFragment>();
+				const float OtherRadius = OtherEntityView.GetFragmentData<FDataFragment_AgentRadius>().Radius;
 
 				const float TotalRadius = AgentRadius + OtherRadius;
 
@@ -1306,8 +1306,8 @@ void UMassStandingAvoidanceProcessor::Execute(UMassEntitySubsystem& EntitySubsys
 				constexpr float DirectionScaleStrength = 0.9f; // How strongly the direction scaling affects [0..1]
 
 				// @todo: this is heavy fragment to access, see if we could handle this differently.
-				const FMassMoveTargetFragment* OtherMoveTarget = OtherEntityView.GetComponentDataPtr<FMassMoveTargetFragment>();
-				const FMassSteeringGhostFragment* OtherGhost = OtherEntityView.GetComponentDataPtr<FMassSteeringGhostFragment>();
+				const FMassMoveTargetFragment* OtherMoveTarget = OtherEntityView.GetFragmentDataPtr<FMassMoveTargetFragment>();
+				const FMassSteeringGhostFragment* OtherGhost = OtherEntityView.GetFragmentDataPtr<FMassSteeringGhostFragment>();
 
 				const bool bOtherHasGhost = OtherMoveTarget != nullptr && OtherGhost != nullptr
 											&& OtherMoveTarget->GetCurrentAction() == EMassMovementAction::Stand
@@ -1438,9 +1438,9 @@ void UMassAvoidanceObstacleProcessor::Execute(UMassEntitySubsystem& EntitySubsys
 		{
 			const int32 NumEntities = Context.GetNumEntities();
 
-			TConstArrayView<FDataFragment_Transform> LocationList = Context.GetComponentView<FDataFragment_Transform>();
-			TConstArrayView<FDataFragment_AgentRadius> RadiiList = Context.GetComponentView<FDataFragment_AgentRadius>();
-			TArrayView<FMassAvoidanceObstacleGridCellLocationFragment> AvoidanceObstacleCellLocationList = Context.GetMutableComponentView<FMassAvoidanceObstacleGridCellLocationFragment>();
+			TConstArrayView<FDataFragment_Transform> LocationList = Context.GetFragmentView<FDataFragment_Transform>();
+			TConstArrayView<FDataFragment_AgentRadius> RadiiList = Context.GetFragmentView<FDataFragment_AgentRadius>();
+			TArrayView<FMassAvoidanceObstacleGridCellLocationFragment> AvoidanceObstacleCellLocationList = Context.GetMutableFragmentView<FMassAvoidanceObstacleGridCellLocationFragment>();
 
 			for (int32 EntityIndex = 0; EntityIndex < NumEntities; ++EntityIndex)
 			{
@@ -1451,7 +1451,7 @@ void UMassAvoidanceObstacleProcessor::Execute(UMassEntitySubsystem& EntitySubsys
 				FMassAvoidanceObstacleItem ObstacleItem;
 				ObstacleItem.Entity = Context.GetEntity(EntityIndex);
 				FMassEntityView EntityView(EntitySubsystem, ObstacleItem.Entity);
-				const FMassAvoidanceColliderFragment* Collider = EntityView.GetComponentDataPtr<FMassAvoidanceColliderFragment>();
+				const FMassAvoidanceColliderFragment* Collider = EntityView.GetFragmentDataPtr<FMassAvoidanceColliderFragment>();
 				if (Collider)
 				{
 					ObstacleItem.ItemFlags |= EMassAvoidanceObstacleItemFlags::HasColliderData;
@@ -1468,9 +1468,9 @@ void UMassAvoidanceObstacleProcessor::Execute(UMassEntitySubsystem& EntitySubsys
 		{
 			const int32 NumEntities = Context.GetNumEntities();
 
-			TConstArrayView<FDataFragment_Transform> LocationList = Context.GetComponentView<FDataFragment_Transform>();
-			TConstArrayView<FDataFragment_AgentRadius> RadiiList = Context.GetComponentView<FDataFragment_AgentRadius>();
-			TArrayView<FMassAvoidanceObstacleGridCellLocationFragment> AvoidanceObstacleCellLocationList = Context.GetMutableComponentView<FMassAvoidanceObstacleGridCellLocationFragment>();
+			TConstArrayView<FDataFragment_Transform> LocationList = Context.GetFragmentView<FDataFragment_Transform>();
+			TConstArrayView<FDataFragment_AgentRadius> RadiiList = Context.GetFragmentView<FDataFragment_AgentRadius>();
+			TArrayView<FMassAvoidanceObstacleGridCellLocationFragment> AvoidanceObstacleCellLocationList = Context.GetMutableFragmentView<FMassAvoidanceObstacleGridCellLocationFragment>();
 
 			for (int32 EntityIndex = 0; EntityIndex < NumEntities; ++EntityIndex)
 			{
@@ -1480,7 +1480,7 @@ void UMassAvoidanceObstacleProcessor::Execute(UMassEntitySubsystem& EntitySubsys
 				FMassAvoidanceObstacleItem ObstacleItem;
 				ObstacleItem.Entity = Context.GetEntity(EntityIndex);
 				FMassEntityView EntityView(EntitySubsystem, ObstacleItem.Entity);
-				const FMassAvoidanceColliderFragment* Collider = EntityView.GetComponentDataPtr<FMassAvoidanceColliderFragment>();
+				const FMassAvoidanceColliderFragment* Collider = EntityView.GetFragmentDataPtr<FMassAvoidanceColliderFragment>();
 				if (Collider)
 				{
 					ObstacleItem.ItemFlags |= EMassAvoidanceObstacleItemFlags::HasColliderData;
@@ -1504,7 +1504,7 @@ void UMassAvoidanceObstacleProcessor::Execute(UMassEntitySubsystem& EntitySubsys
 	{
 		const int32 NumEntities = Context.GetNumEntities();
 
-		TArrayView<FMassAvoidanceObstacleGridCellLocationFragment> AvoidanceObstacleCellLocationList = Context.GetMutableComponentView<FMassAvoidanceObstacleGridCellLocationFragment>();
+		TArrayView<FMassAvoidanceObstacleGridCellLocationFragment> AvoidanceObstacleCellLocationList = Context.GetMutableFragmentView<FMassAvoidanceObstacleGridCellLocationFragment>();
 		for (int32 EntityIndex = 0; EntityIndex < NumEntities; ++EntityIndex)
 		{
 			FMassAvoidanceObstacleItem ObstacleItem;
@@ -1567,10 +1567,10 @@ void UMassNavigationBoundaryProcessor::Execute(UMassEntitySubsystem& EntitySubsy
 	EntityQuery.ForEachEntityChunk(EntitySubsystem, Context, [this, NavData](FMassExecutionContext& Context)
 		{
 			const int32 NumEntities = Context.GetNumEntities();
-			TConstArrayView<FDataFragment_Transform> LocationList = Context.GetComponentView<FDataFragment_Transform>();
-			TConstArrayView<FDataFragment_NavLocation> NavLocationList = Context.GetComponentView<FDataFragment_NavLocation>();
-			TConstArrayView<FMassEdgeDetectionParamsFragment> EdgeDetectionParamsList = Context.GetComponentView<FMassEdgeDetectionParamsFragment>();
-			TArrayView<FMassNavigationEdgesFragment> EdgesList = Context.GetMutableComponentView<FMassNavigationEdgesFragment>();
+			TConstArrayView<FDataFragment_Transform> LocationList = Context.GetFragmentView<FDataFragment_Transform>();
+			TConstArrayView<FDataFragment_NavLocation> NavLocationList = Context.GetFragmentView<FDataFragment_NavLocation>();
+			TConstArrayView<FMassEdgeDetectionParamsFragment> EdgeDetectionParamsList = Context.GetFragmentView<FMassEdgeDetectionParamsFragment>();
+			TArrayView<FMassNavigationEdgesFragment> EdgesList = Context.GetMutableFragmentView<FMassNavigationEdgesFragment>();
 
 			const ARecastNavMesh* RecastNavMesh = Cast<const ARecastNavMesh>(NavData);
 			if (!RecastNavMesh)
@@ -1650,11 +1650,11 @@ void UMassLaneBoundaryProcessor::Execute(UMassEntitySubsystem& EntitySubsystem,
 		{
 			const int32 NumEntities = Context.GetNumEntities();
 
-			TConstArrayView<FDataFragment_Transform> LocationList = Context.GetComponentView<FDataFragment_Transform>();
-			TArrayView<FMassNavigationEdgesFragment> EdgesList = Context.GetMutableComponentView<FMassNavigationEdgesFragment>();
-			TArrayView<FMassLastUpdatePositionFragment> LastUpdatePositionList = Context.GetMutableComponentView<FMassLastUpdatePositionFragment>();
-			TConstArrayView<FMassZoneGraphLaneLocationFragment> LaneLocationList = Context.GetComponentView<FMassZoneGraphLaneLocationFragment>();
-			TArrayView<FMassAvoidanceBoundaryLastLaneHandleFragment> LastLaneHandleList = Context.GetMutableComponentView<FMassAvoidanceBoundaryLastLaneHandleFragment>();
+			TConstArrayView<FDataFragment_Transform> LocationList = Context.GetFragmentView<FDataFragment_Transform>();
+			TArrayView<FMassNavigationEdgesFragment> EdgesList = Context.GetMutableFragmentView<FMassNavigationEdgesFragment>();
+			TArrayView<FMassLastUpdatePositionFragment> LastUpdatePositionList = Context.GetMutableFragmentView<FMassLastUpdatePositionFragment>();
+			TConstArrayView<FMassZoneGraphLaneLocationFragment> LaneLocationList = Context.GetFragmentView<FMassZoneGraphLaneLocationFragment>();
+			TArrayView<FMassAvoidanceBoundaryLastLaneHandleFragment> LastLaneHandleList = Context.GetMutableFragmentView<FMassAvoidanceBoundaryLastLaneHandleFragment>();
 
 			TArray<FZoneGraphLinkedLane> LinkedLanes;
 			LinkedLanes.Reserve(4);
@@ -1920,11 +1920,11 @@ void UMassLaneCacheBoundaryProcessor::Execute(UMassEntitySubsystem& EntitySubsys
 	{
 		const int32 NumEntities = Context.GetNumEntities();
 
-		TConstArrayView<FMassZoneGraphCachedLaneFragment> CachedLaneList = Context.GetComponentView<FMassZoneGraphCachedLaneFragment>();
-		TConstArrayView<FMassZoneGraphLaneLocationFragment> LaneLocationList = Context.GetComponentView<FMassZoneGraphLaneLocationFragment>();
-		TConstArrayView<FMassMoveTargetFragment> MovementTargetList = Context.GetComponentView<FMassMoveTargetFragment>();
-		TArrayView<FMassLaneCacheBoundaryFragment> LaneCacheBoundaryList = Context.GetMutableComponentView<FMassLaneCacheBoundaryFragment>();
-		TArrayView<FMassNavigationEdgesFragment> EdgesList = Context.GetMutableComponentView<FMassNavigationEdgesFragment>();
+		TConstArrayView<FMassZoneGraphCachedLaneFragment> CachedLaneList = Context.GetFragmentView<FMassZoneGraphCachedLaneFragment>();
+		TConstArrayView<FMassZoneGraphLaneLocationFragment> LaneLocationList = Context.GetFragmentView<FMassZoneGraphLaneLocationFragment>();
+		TConstArrayView<FMassMoveTargetFragment> MovementTargetList = Context.GetFragmentView<FMassMoveTargetFragment>();
+		TArrayView<FMassLaneCacheBoundaryFragment> LaneCacheBoundaryList = Context.GetMutableFragmentView<FMassLaneCacheBoundaryFragment>();
+		TArrayView<FMassNavigationEdgesFragment> EdgesList = Context.GetMutableFragmentView<FMassNavigationEdgesFragment>();
 
 		TArray<FZoneGraphLinkedLane> LinkedLanes;
 		LinkedLanes.Reserve(4);	
diff --git a/Engine/Plugins/AI/MassAI/Source/MassAIMovement/Private/MassZoneGraphMovementProcessors.cpp b/Engine/Plugins/AI/MassAI/Source/MassAIMovement/Private/MassZoneGraphMovementProcessors.cpp
--- a/Engine/Plugins/AI/MassAI/Source/MassAIMovement/Private/MassZoneGraphMovementProcessors.cpp
+++ b/Engine/Plugins/AI/MassAI/Source/MassAIMovement/Private/MassZoneGraphMovementProcessors.cpp
@@ -101,10 +101,10 @@ void UMassZoneGraphLocationInitializer::Execute(UMassEntitySubsystem& EntitySubs
 	EntityQuery.ForEachEntityChunk(EntitySubsystem, Context, [this, Settings](FMassExecutionContext& Context)
 	{
 		const int32 NumEntities = Context.GetNumEntities();
-		const TArrayView<FMassZoneGraphLaneLocationFragment> LaneLocationList = Context.GetMutableComponentView<FMassZoneGraphLaneLocationFragment>();
-		const TArrayView<FMassMoveTargetFragment> MoveTargetList = Context.GetMutableComponentView<FMassMoveTargetFragment>();
-		const TConstArrayView<FMassMovementConfigFragment> MovementConfigList = Context.GetComponentView<FMassMovementConfigFragment>();
-		const TConstArrayView<FDataFragment_Transform> TransformList = Context.GetComponentView<FDataFragment_Transform>();
+		const TArrayView<FMassZoneGraphLaneLocationFragment> LaneLocationList = Context.GetMutableFragmentView<FMassZoneGraphLaneLocationFragment>();
+		const TArrayView<FMassMoveTargetFragment> MoveTargetList = Context.GetMutableFragmentView<FMassMoveTargetFragment>();
+		const TConstArrayView<FMassMovementConfigFragment> MovementConfigList = Context.GetFragmentView<FMassMovementConfigFragment>();
+		const TConstArrayView<FDataFragment_Transform> TransformList = Context.GetFragmentView<FDataFragment_Transform>();
 
 		// Get the default movement config.
 		FMassMovementConfigHandle CurrentConfigHandle;
@@ -208,10 +208,10 @@ void UMassZoneGraphPathFollowProcessor::Execute(UMassEntitySubsystem& EntitySubs
 	EntityQuery_Conditional.ForEachEntityChunk(EntitySubsystem, Context, [this, &EntitiesToSignalPathDone, &EntitiesToSignalLaneChanged](FMassExecutionContext& Context)
 		{
 		const int32 NumEntities = Context.GetNumEntities();
-		const TArrayView<FMassZoneGraphShortPathFragment> ShortPathList = Context.GetMutableComponentView<FMassZoneGraphShortPathFragment>();
-		const TArrayView<FMassZoneGraphLaneLocationFragment> LaneLocationList = Context.GetMutableComponentView<FMassZoneGraphLaneLocationFragment>();
-		const TArrayView<FMassMoveTargetFragment> MoveTargetList = Context.GetMutableComponentView<FMassMoveTargetFragment>();
-		const TConstArrayView<FMassSimulationLODFragment> SimLODList = Context.GetComponentView<FMassSimulationLODFragment>();
+		const TArrayView<FMassZoneGraphShortPathFragment> ShortPathList = Context.GetMutableFragmentView<FMassZoneGraphShortPathFragment>();
+		const TArrayView<FMassZoneGraphLaneLocationFragment> LaneLocationList = Context.GetMutableFragmentView<FMassZoneGraphLaneLocationFragment>();
+		const TArrayView<FMassMoveTargetFragment> MoveTargetList = Context.GetMutableFragmentView<FMassMoveTargetFragment>();
+		const TConstArrayView<FMassSimulationLODFragment> SimLODList = Context.GetFragmentView<FMassSimulationLODFragment>();
 
 		for (int32 EntityIndex = 0; EntityIndex < NumEntities; ++EntityIndex)
 		{
@@ -498,12 +498,12 @@ void UMassZoneGraphSteeringProcessor::Execute(UMassEntitySubsystem& EntitySubsys
 	EntityQuery.ForEachEntityChunk(EntitySubsystem, Context, [this, Settings](FMassExecutionContext& Context)
 	{
 		const int32 NumEntities = Context.GetNumEntities();
-		const TArrayView<FMassMoveTargetFragment> MoveTargetList = Context.GetMutableComponentView<FMassMoveTargetFragment>();
-		const TConstArrayView<FMassMovementConfigFragment> MovementConfigList = Context.GetComponentView<FMassMovementConfigFragment>();
-		const TConstArrayView<FMassVelocityFragment> VelocityList = Context.GetComponentView<FMassVelocityFragment>();
-		const TArrayView<FDataFragment_Transform> TransformList = Context.GetMutableComponentView<FDataFragment_Transform>();
-		const TArrayView<FMassSteeringFragment> SteeringList = Context.GetMutableComponentView<FMassSteeringFragment>();
-		const TArrayView<FMassSteeringGhostFragment> GhostList = Context.GetMutableComponentView<FMassSteeringGhostFragment>();
+		const TArrayView<FMassMoveTargetFragment> MoveTargetList = Context.GetMutableFragmentView<FMassMoveTargetFragment>();
+		const TConstArrayView<FMassMovementConfigFragment> MovementConfigList = Context.GetFragmentView<FMassMovementConfigFragment>();
+		const TConstArrayView<FMassVelocityFragment> VelocityList = Context.GetFragmentView<FMassVelocityFragment>();
+		const TArrayView<FDataFragment_Transform> TransformList = Context.GetMutableFragmentView<FDataFragment_Transform>();
+		const TArrayView<FMassSteeringFragment> SteeringList = Context.GetMutableFragmentView<FMassSteeringFragment>();
+		const TArrayView<FMassSteeringGhostFragment> GhostList = Context.GetMutableFragmentView<FMassSteeringGhostFragment>();
 
 		// @todo: make configurable
 		constexpr float StandDeadZoneRadius = 5.0f;
diff --git a/Engine/Plugins/AI/MassAI/Source/MassAIMovement/Private/MovementProcessor.cpp b/Engine/Plugins/AI/MassAI/Source/MassAIMovement/Private/MovementProcessor.cpp
--- a/Engine/Plugins/AI/MassAI/Source/MassAIMovement/Private/MovementProcessor.cpp
+++ b/Engine/Plugins/AI/MassAI/Source/MassAIMovement/Private/MovementProcessor.cpp
@@ -127,8 +127,8 @@ void UMassProcessor_Movement::Execute(UMassEntitySubsystem& EntitySubsystem, FMa
 	EntityQuery.ForEachEntityChunk(EntitySubsystem, Context, [this, TimeDelta](FMassExecutionContext& Context)
 		{
 			const int32 NumEntities = Context.GetNumEntities();
-			const TArrayView<FDataFragment_Transform> LocationList = Context.GetMutableComponentView<FDataFragment_Transform>();
-			const TArrayView<FMassVelocityFragment> VelocityList = Context.GetMutableComponentView<FMassVelocityFragment>();
+			const TArrayView<FDataFragment_Transform> LocationList = Context.GetMutableFragmentView<FDataFragment_Transform>();
+			const TArrayView<FMassVelocityFragment> VelocityList = Context.GetMutableFragmentView<FMassVelocityFragment>();
 			for (int32 i = 0; i < NumEntities; ++i)
 			{
 				LocationList[i].GetMutableTransform().AddToTranslation(VelocityList[i].Value * TimeDelta);
@@ -187,7 +187,7 @@ void UMassAvoidanceObstacleRemoverFragmentDestructor::Execute(UMassEntitySubsyst
 	EntityQuery.ForEachEntityChunk(EntitySubsystem, Context, [this, MovementSubsystem](FMassExecutionContext& Context)
 		{
 			const int32 NumEntities = Context.GetNumEntities();
-			const TArrayView<FMassAvoidanceObstacleGridCellLocationFragment> AvoidanceObstacleCellLocationList = Context.GetMutableComponentView<FMassAvoidanceObstacleGridCellLocationFragment>();
+			const TArrayView<FMassAvoidanceObstacleGridCellLocationFragment> AvoidanceObstacleCellLocationList = Context.GetMutableFragmentView<FMassAvoidanceObstacleGridCellLocationFragment>();
 
 			for (int32 i = 0; i < NumEntities; ++i)
 			{
@@ -216,7 +216,7 @@ void UVelocityFragmentInitializer::Execute(UMassEntitySubsystem& EntitySubsystem
 {
 	EntityQuery.ForEachEntityChunk(EntitySubsystem, Context, [this](FMassExecutionContext& Context)
 		{
-			const TArrayView<FMassVelocityFragment> VelocitiesList = Context.GetMutableComponentView<FMassVelocityFragment>();
+			const TArrayView<FMassVelocityFragment> VelocitiesList = Context.GetMutableFragmentView<FMassVelocityFragment>();
 			for (FMassVelocityFragment& VelocityFragment : VelocitiesList)
 			{
 				VelocityFragment.Value = FVector(FMath::FRandRange(-100.f, 100.f), FMath::FRandRange(-100.f, 100.f), 0);
diff --git a/Engine/Plugins/AI/MassAI/Source/MassAIMovement/Private/NavigationProcessor.cpp b/Engine/Plugins/AI/MassAI/Source/MassAIMovement/Private/NavigationProcessor.cpp
--- a/Engine/Plugins/AI/MassAI/Source/MassAIMovement/Private/NavigationProcessor.cpp
+++ b/Engine/Plugins/AI/MassAI/Source/MassAIMovement/Private/NavigationProcessor.cpp
@@ -120,10 +120,10 @@ void UMassApplyVelocityMoveTargetProcessor::Execute(UMassEntitySubsystem& Entity
 			constexpr float OrientationBlendWhileMoving = 0.4f;
 			constexpr float OrientationBlendWhileStanding = 0.95f;
 
-			const TConstArrayView<FMassMoveTargetFragment> MoveTargetList = Context.GetComponentView<FMassMoveTargetFragment>();
-			const TArrayView<FDataFragment_Transform> LocationList = Context.GetMutableComponentView<FDataFragment_Transform>();
-			const TArrayView<FMassSteeringFragment> SteeringList = Context.GetMutableComponentView<FMassSteeringFragment>();
-			const TArrayView<FMassVelocityFragment> VelocityList = Context.GetMutableComponentView<FMassVelocityFragment>();
+			const TConstArrayView<FMassMoveTargetFragment> MoveTargetList = Context.GetFragmentView<FMassMoveTargetFragment>();
+			const TArrayView<FDataFragment_Transform> LocationList = Context.GetMutableFragmentView<FDataFragment_Transform>();
+			const TArrayView<FMassSteeringFragment> SteeringList = Context.GetMutableFragmentView<FMassSteeringFragment>();
+			const TArrayView<FMassVelocityFragment> VelocityList = Context.GetMutableFragmentView<FMassVelocityFragment>();
 
 			for (int32 EntityIndex = 0; EntityIndex < NumEntities; ++EntityIndex)
 			{
@@ -207,8 +207,8 @@ void UMassApplyVelocityMoveTargetProcessor::Execute(UMassEntitySubsystem& Entity
 			{
 				const int32 NumEntities = Context.GetNumEntities();
 
-				const TArrayView<FDataFragment_Transform> LocationList = Context.GetMutableComponentView<FDataFragment_Transform>();
-				const TConstArrayView<FMassMoveTargetFragment> MoveTargetList = Context.GetComponentView<FMassMoveTargetFragment>();
+				const TArrayView<FDataFragment_Transform> LocationList = Context.GetMutableFragmentView<FDataFragment_Transform>();
+				const TConstArrayView<FMassMoveTargetFragment> MoveTargetList = Context.GetFragmentView<FMassMoveTargetFragment>();
 
 				for (int32 EntityIndex = 0; EntityIndex < NumEntities; ++EntityIndex)
 				{
@@ -247,9 +247,9 @@ void UMassDynamicObstacleProcessor::Execute(UMassEntitySubsystem& EntitySubsyste
 		{
 			const int32 NumEntities = Context.GetNumEntities();
 
-			const TConstArrayView<FDataFragment_Transform> LocationList = Context.GetComponentView<FDataFragment_Transform>();
-			const TConstArrayView<FDataFragment_AgentRadius> RadiusList = Context.GetComponentView<FDataFragment_AgentRadius>();
-			const TArrayView<FMassDynamicObstacleFragment> ObstacleDataList = Context.GetMutableComponentView<FMassDynamicObstacleFragment>();
+			const TConstArrayView<FDataFragment_Transform> LocationList = Context.GetFragmentView<FDataFragment_Transform>();
+			const TConstArrayView<FDataFragment_AgentRadius> RadiusList = Context.GetFragmentView<FDataFragment_AgentRadius>();
+			const TArrayView<FMassDynamicObstacleFragment> ObstacleDataList = Context.GetMutableFragmentView<FMassDynamicObstacleFragment>();
 
 			const FDateTime Now = FDateTime::UtcNow();
 
diff --git a/Engine/Plugins/AI/MassAI/Source/MassAIMovement/Private/SnapToNavigationProcessor.cpp b/Engine/Plugins/AI/MassAI/Source/MassAIMovement/Private/SnapToNavigationProcessor.cpp
--- a/Engine/Plugins/AI/MassAI/Source/MassAIMovement/Private/SnapToNavigationProcessor.cpp
+++ b/Engine/Plugins/AI/MassAI/Source/MassAIMovement/Private/SnapToNavigationProcessor.cpp
@@ -36,8 +36,8 @@ void USnapToNavigationProcessor::Execute(UMassEntitySubsystem& EntitySubsystem,
 	EntityQuery.ForEachEntityChunk(EntitySubsystem, Context, [this, NavData](FMassExecutionContext& Context)
 		{
 			const int32 NumEntities = Context.GetNumEntities();
-			const TArrayView<FDataFragment_Transform> LocationList = Context.GetMutableComponentView<FDataFragment_Transform>();
-			const TArrayView<FDataFragment_NavLocation> NavLocationList = Context.GetMutableComponentView<FDataFragment_NavLocation>();
+			const TArrayView<FDataFragment_Transform> LocationList = Context.GetMutableFragmentView<FDataFragment_Transform>();
+			const TArrayView<FDataFragment_NavLocation> NavLocationList = Context.GetMutableFragmentView<FDataFragment_NavLocation>();
 
 			TArray<FNavigationProjectionWork> Workload;
 			for (int32 Index = 0; Index < NumEntities; ++Index)
diff --git a/Engine/Plugins/AI/MassAI/Source/MassAIReplication/Private/MassReplicationPathHandlers.cpp b/Engine/Plugins/AI/MassAI/Source/MassAIReplication/Private/MassReplicationPathHandlers.cpp
--- a/Engine/Plugins/AI/MassAI/Source/MassAIReplication/Private/MassReplicationPathHandlers.cpp
+++ b/Engine/Plugins/AI/MassAI/Source/MassAIReplication/Private/MassReplicationPathHandlers.cpp
@@ -18,11 +18,11 @@ void FMassReplicationProcessorPathHandler::AddRequirements(FMassEntityQuery& InQ
 	InQuery.AddRequirement<FMassZoneGraphLaneLocationFragment>(EMassFragmentAccess::ReadOnly);
 }
 
-void FMassReplicationProcessorPathHandler::CacheComponentViews(FMassExecutionContext& ExecContext)
+void FMassReplicationProcessorPathHandler::CacheFragmentViews(FMassExecutionContext& ExecContext)
 {
-	PathRequestList = ExecContext.GetMutableComponentView<FMassZoneGraphPathRequestFragment>();
-	MoveTargetList = ExecContext.GetMutableComponentView<FMassMoveTargetFragment>();
-	LaneLocationList = ExecContext.GetMutableComponentView<FMassZoneGraphLaneLocationFragment>();
+	PathRequestList = ExecContext.GetMutableFragmentView<FMassZoneGraphPathRequestFragment>();
+	MoveTargetList = ExecContext.GetMutableFragmentView<FMassMoveTargetFragment>();
+	LaneLocationList = ExecContext.GetMutableFragmentView<FMassZoneGraphLaneLocationFragment>();
 }
 
 void FMassReplicationProcessorPathHandler::AddEntity(const int32 EntityIdx, FReplicatedAgentPathData& InOutReplicatedPathData) const
@@ -97,7 +97,7 @@ void FReplicatedAgentPathData::InitEntity(const UWorld& InWorld,
 void FReplicatedAgentPathData::ApplyToEntity(const UWorld& InWorld, const FMassEntityView& InEntityView) const
 {
 	const FMassEntityHandle Entity = InEntityView.GetEntity();
-	FMassMoveTargetFragment& MoveTarget = InEntityView.GetComponentData<FMassMoveTargetFragment>();
+	FMassMoveTargetFragment& MoveTarget = InEntityView.GetFragmentData<FMassMoveTargetFragment>();
 	if (MoveTarget.GetCurrentActionID() == ActionID)
 	{
 		return;
@@ -116,10 +116,10 @@ void FReplicatedAgentPathData::ApplyToEntity(const UWorld& InWorld, const FMassE
 
 	UE_VLOG(SimulationSubsystem, LogMassNavigation, Log, TEXT("Entity [%s] apply replicated data to entity"), *Entity.DebugGetDescription());
 
-	FMassZoneGraphShortPathFragment& ShortPath = InEntityView.GetComponentData<FMassZoneGraphShortPathFragment>();
-	FMassZoneGraphCachedLaneFragment& CachedLane = InEntityView.GetComponentData<FMassZoneGraphCachedLaneFragment>();
-	FMassZoneGraphLaneLocationFragment& LaneLocation = InEntityView.GetComponentData<FMassZoneGraphLaneLocationFragment>();
-	const FDataFragment_AgentRadius& AgentRadius = InEntityView.GetComponentData<FDataFragment_AgentRadius>();
+	FMassZoneGraphShortPathFragment& ShortPath = InEntityView.GetFragmentData<FMassZoneGraphShortPathFragment>();
+	FMassZoneGraphCachedLaneFragment& CachedLane = InEntityView.GetFragmentData<FMassZoneGraphCachedLaneFragment>();
+	FMassZoneGraphLaneLocationFragment& LaneLocation = InEntityView.GetFragmentData<FMassZoneGraphLaneLocationFragment>();
+	const FDataFragment_AgentRadius& AgentRadius = InEntityView.GetFragmentData<FDataFragment_AgentRadius>();
 
 	MoveTarget.CreateReplicatedAction(Action, ActionID, InWorld.GetTimeSeconds(), ActionServerStartTime);
 	MoveTarget.DesiredSpeed = DesiredSpeed;
diff --git a/Engine/Plugins/AI/MassAI/Source/MassAIReplication/Public/MassReplicationPathHandlers.h b/Engine/Plugins/AI/MassAI/Source/MassAIReplication/Public/MassReplicationPathHandlers.h
--- a/Engine/Plugins/AI/MassAI/Source/MassAIReplication/Public/MassReplicationPathHandlers.h
+++ b/Engine/Plugins/AI/MassAI/Source/MassAIReplication/Public/MassReplicationPathHandlers.h
@@ -95,8 +95,8 @@ public:
 	 * The following functions are used to configure the query and then set that data for path following.
 	 */
 	static void AddRequirementsForSpawnQuery(FMassEntityQuery& InQuery);
-	void CacheComponentViewsForSpawnQuery(FMassExecutionContext& InExecContext);
-	void ClearComponentViewsForSpawnQuery();
+	void CacheFragmentViewsForSpawnQuery(FMassExecutionContext& InExecContext);
+	void ClearFragmentViewsForSpawnQuery();
 
 	void SetSpawnedEntityData(const FMassEntityView& EntityView, const FReplicatedAgentPathData& ReplicatedPathData, const int32 EntityIdx) const;
 
@@ -150,17 +150,17 @@ void TMassClientBubblePathHandler<AgentArrayItem>::AddRequirementsForSpawnQuery(
 
 #if UE_REPLICATION_COMPILE_CLIENT_CODE
 template<typename AgentArrayItem>
-void TMassClientBubblePathHandler<AgentArrayItem>::CacheComponentViewsForSpawnQuery(FMassExecutionContext& InExecContext)
+void TMassClientBubblePathHandler<AgentArrayItem>::CacheFragmentViewsForSpawnQuery(FMassExecutionContext& InExecContext)
 {
-	PathRequestList = InExecContext.GetMutableComponentView<FMassZoneGraphPathRequestFragment>();
-	MoveTargetList = InExecContext.GetMutableComponentView<FMassMoveTargetFragment>();
-	LaneLocationList = InExecContext.GetMutableComponentView<FMassZoneGraphLaneLocationFragment>();
+	PathRequestList = InExecContext.GetMutableFragmentView<FMassZoneGraphPathRequestFragment>();
+	MoveTargetList = InExecContext.GetMutableFragmentView<FMassMoveTargetFragment>();
+	LaneLocationList = InExecContext.GetMutableFragmentView<FMassZoneGraphLaneLocationFragment>();
 }
 #endif // UE_REPLICATION_COMPILE_CLIENT_CODE
 
 #if UE_REPLICATION_COMPILE_CLIENT_CODE
 template<typename AgentArrayItem>
-void TMassClientBubblePathHandler<AgentArrayItem>::ClearComponentViewsForSpawnQuery()
+void TMassClientBubblePathHandler<AgentArrayItem>::ClearFragmentViewsForSpawnQuery()
 {
 	LaneLocationList = TArrayView<FMassZoneGraphLaneLocationFragment>();
 	MoveTargetList = TArrayView<FMassMoveTargetFragment>();
@@ -200,11 +200,11 @@ public:
 	static void AddRequirements(FMassEntityQuery& InQuery);
 
 	/** Cache any component views you want to, this will get called before we iterate through entities. */
-	void CacheComponentViews(FMassExecutionContext& ExecContext);
+	void CacheFragmentViews(FMassExecutionContext& ExecContext);
 
 	/**
 	 * Set the replicated path data when we are adding an entity to the client bubble.
-	 * @param EntityIdx the index of the entity in component views that have been cached.
+	 * @param EntityIdx the index of the entity in fragment views that have been cached.
 	 * @param InOUtReplicatedPathData the data to set.
 	 */
 	void AddEntity(const int32 EntityIdx, FReplicatedAgentPathData& InOUtReplicatedPathData) const;
@@ -212,7 +212,7 @@ public:
 	/**
 	 * Set the replicated path data when we are modifying an entity that already exists in the client bubble.
 	 * @param Handle to the agent in the TMassClientBubbleHandler (that TMassClientBubblePathHandler is a member variable of).
-	 * @param EntityIdx the index of the entity in component views that have been cached.
+	 * @param EntityIdx the index of the entity in fragment views that have been cached.
 	 * @param BubblePathHandler handler to actually set the data in the client bubble
 	 */
 	template<typename AgentArrayItem>
diff --git a/Engine/Plugins/AI/MassCrowd/Source/MassCrowd/Private/MassCrowdBubble.cpp b/Engine/Plugins/AI/MassCrowd/Source/MassCrowd/Private/MassCrowdBubble.cpp
--- a/Engine/Plugins/AI/MassCrowd/Source/MassCrowd/Private/MassCrowdBubble.cpp
+++ b/Engine/Plugins/AI/MassCrowd/Source/MassCrowd/Private/MassCrowdBubble.cpp
@@ -24,8 +24,8 @@ namespace UE::Mass::Crowd
 
 		const FMassEntityView EntityView(EntitySystem, Entity);
 
-		const FDataFragment_Transform& TransformFragment = EntityView.GetComponentData<FDataFragment_Transform>();
-		const FMassNetworkIDFragment& NetworkIDFragment = EntityView.GetComponentData<FMassNetworkIDFragment>();
+		const FDataFragment_Transform& TransformFragment = EntityView.GetFragmentData<FDataFragment_Transform>();
+		const FMassNetworkIDFragment& NetworkIDFragment = EntityView.GetFragmentData<FMassNetworkIDFragment>();
 
 		const FVector& Pos = TransformFragment.GetTransform().GetLocation();
 		const uint32 NetworkID = NetworkIDFragment.NetID.GetValue();
@@ -129,10 +129,10 @@ void FMassCrowdClientBubbleHandler::PostReplicatedAdd(const TArrayView<int32> Ad
 		TransformHandler.AddRequirementsForSpawnQuery(InQuery);
 	};
 
-	auto CacheComponentViewsForSpawnQuery = [this](FMassExecutionContext& InExecContext)
+	auto CacheFragmentViewsForSpawnQuery = [this](FMassExecutionContext& InExecContext)
 	{
-		PathHandler.CacheComponentViewsForSpawnQuery(InExecContext);
-		TransformHandler.CacheComponentViewsForSpawnQuery(InExecContext);
+		PathHandler.CacheFragmentViewsForSpawnQuery(InExecContext);
+		TransformHandler.CacheFragmentViewsForSpawnQuery(InExecContext);
 	};
 
 	auto SetSpawnedEntityData = [this](const FMassEntityView& EntityView, const FReplicatedCrowdAgent& ReplicatedEntity, const int32 EntityIdx)
@@ -146,10 +146,10 @@ void FMassCrowdClientBubbleHandler::PostReplicatedAdd(const TArrayView<int32> Ad
 		PostReplicatedChangeEntity(EntityView, Item);
 	};
 
-	PostReplicatedAddHelper(AddedIndices, AddRequirementsForSpawnQuery, CacheComponentViewsForSpawnQuery, SetSpawnedEntityData, SetModifiedEntityData);
+	PostReplicatedAddHelper(AddedIndices, AddRequirementsForSpawnQuery, CacheFragmentViewsForSpawnQuery, SetSpawnedEntityData, SetModifiedEntityData);
 
-	PathHandler.ClearComponentViewsForSpawnQuery();
-	TransformHandler.ClearComponentViewsForSpawnQuery();
+	PathHandler.ClearFragmentViewsForSpawnQuery();
+	TransformHandler.ClearFragmentViewsForSpawnQuery();
 }
 #endif //UE_REPLICATION_COMPILE_SERVER_CODE
 
diff --git a/Engine/Plugins/AI/MassCrowd/Source/MassCrowd/Private/MassCrowdNavigationProcessor.cpp b/Engine/Plugins/AI/MassCrowd/Source/MassCrowd/Private/MassCrowdNavigationProcessor.cpp
--- a/Engine/Plugins/AI/MassCrowd/Source/MassCrowd/Private/MassCrowdNavigationProcessor.cpp
+++ b/Engine/Plugins/AI/MassCrowd/Source/MassCrowd/Private/MassCrowdNavigationProcessor.cpp
@@ -48,8 +48,8 @@ void UMassCrowdLaneTrackingSignalProcessor::SignalEntities(UMassEntitySubsystem&
 	EntityQuery.ForEachEntityChunk(EntitySubsystem, Context, [this](FMassExecutionContext& Context)
 	{
 		const int32 NumEntities = Context.GetNumEntities();
-		const TConstArrayView<FMassZoneGraphLaneLocationFragment> LaneLocationList = Context.GetComponentView<FMassZoneGraphLaneLocationFragment>();
-		const TArrayView<FMassCrowdLaneTrackingFragment> LaneTrackingList = Context.GetMutableComponentView<FMassCrowdLaneTrackingFragment>();
+		const TConstArrayView<FMassZoneGraphLaneLocationFragment> LaneLocationList = Context.GetFragmentView<FMassZoneGraphLaneLocationFragment>();
+		const TArrayView<FMassCrowdLaneTrackingFragment> LaneTrackingList = Context.GetMutableFragmentView<FMassCrowdLaneTrackingFragment>();
 
 		for (int32 EntityIndex = 0; EntityIndex < NumEntities; ++EntityIndex)
 		{
@@ -91,7 +91,7 @@ void UMassCrowdLaneTrackingDestructor::Execute(UMassEntitySubsystem& EntitySubsy
 	EntityQuery.ForEachEntityChunk(EntitySubsystem, Context, [this](const FMassExecutionContext& Context)
 	{
 		const int32 NumEntities = Context.GetNumEntities();
-		const TConstArrayView<FMassCrowdLaneTrackingFragment> LaneTrackingList = Context.GetComponentView<FMassCrowdLaneTrackingFragment>();
+		const TConstArrayView<FMassCrowdLaneTrackingFragment> LaneTrackingList = Context.GetFragmentView<FMassCrowdLaneTrackingFragment>();
 
 		for (int32 EntityIndex = 0; EntityIndex < NumEntities; ++EntityIndex)
 		{
diff --git a/Engine/Plugins/AI/MassCrowd/Source/MassCrowd/Private/MassCrowdReplicationProcessor.cpp b/Engine/Plugins/AI/MassCrowd/Source/MassCrowd/Private/MassCrowdReplicationProcessor.cpp
--- a/Engine/Plugins/AI/MassCrowd/Source/MassCrowd/Private/MassCrowdReplicationProcessor.cpp
+++ b/Engine/Plugins/AI/MassCrowd/Source/MassCrowd/Private/MassCrowdReplicationProcessor.cpp
@@ -60,8 +60,8 @@ void UMassCrowdReplicationProcessor::Execute(UMassEntitySubsystem& EntitySubsyst
 	{
 		EntityQuery.ForEachEntityChunk(EntitySubsystem, Context, [this](FMassExecutionContext& Context)
 			{
-				const TConstArrayView<FDataFragment_Transform> TransformList = Context.GetComponentView<FDataFragment_Transform>();
-				const TConstArrayView<FMassReplicationViewerLODFragment> ViewerLODList = Context.GetComponentView<FMassReplicationViewerLODFragment>();
+				const TConstArrayView<FDataFragment_Transform> TransformList = Context.GetFragmentView<FDataFragment_Transform>();
+				const TConstArrayView<FMassReplicationViewerLODFragment> ViewerLODList = Context.GetFragmentView<FMassReplicationViewerLODFragment>();
 				LODCalculator.DebugDisplayLOD(Context, ViewerLODList, TransformList, World);
 			});
 	}
@@ -77,8 +77,8 @@ void UMassCrowdReplicationProcessor::ProcessClientReplication(UMassEntitySubsyst
 
 	auto CacheViewsCallback = [&Context, &PathHandler, &PositionYawHandler]()
 	{
-		PathHandler.CacheComponentViews(Context);
-		PositionYawHandler.CacheComponentViews(Context);
+		PathHandler.CacheFragmentViews(Context);
+		PositionYawHandler.CacheFragmentViews(Context);
 	};
 
 	auto AddEntityCallback = [this, &Context, &PathHandler, &PositionYawHandler](const int32 EntityIdx, FReplicatedCrowdAgent& InReplicatedAgent, const FMassClientHandle ClientHandle)->FMassReplicatedAgentHandle
diff --git a/Engine/Plugins/AI/MassCrowd/Source/MassCrowd/Private/MassCrowdRepresentationProcessor.cpp b/Engine/Plugins/AI/MassCrowd/Source/MassCrowd/Private/MassCrowdRepresentationProcessor.cpp
--- a/Engine/Plugins/AI/MassCrowd/Source/MassCrowd/Private/MassCrowdRepresentationProcessor.cpp
+++ b/Engine/Plugins/AI/MassCrowd/Source/MassCrowd/Private/MassCrowdRepresentationProcessor.cpp
@@ -42,9 +42,9 @@ void UMassCrowdRepresentationProcessor::InitializeVelocity(UMassEntitySubsystem&
 {
 	CharacterMovementEntitiesQuery_Conditional.ForEachEntityChunk(EntitySubsystem, Context, [this](FMassExecutionContext& Context)
 		{
-			const TConstArrayView<FMassRepresentationFragment> VisualizationList = Context.GetComponentView<FMassRepresentationFragment>();
-			const TConstArrayView<FMassVelocityFragment> VelocityList = Context.GetComponentView<FMassVelocityFragment>();
-			const TArrayView<FDataFragment_CharacterMovementComponentWrapper> CharMoveWrapperList = Context.GetMutableComponentView<FDataFragment_CharacterMovementComponentWrapper>();
+			const TConstArrayView<FMassRepresentationFragment> VisualizationList = Context.GetFragmentView<FMassRepresentationFragment>();
+			const TConstArrayView<FMassVelocityFragment> VelocityList = Context.GetFragmentView<FMassVelocityFragment>();
+			const TArrayView<FDataFragment_CharacterMovementComponentWrapper> CharMoveWrapperList = Context.GetMutableFragmentView<FDataFragment_CharacterMovementComponentWrapper>();
 
 			const int32 NumEntities = Context.GetNumEntities();
 			for (int32 EntityIdx = 0; EntityIdx < NumEntities; EntityIdx++)
diff --git a/Engine/Plugins/AI/MassCrowd/Source/MassCrowd/Private/MassCrowdServerRepresentationLODProcessor.cpp b/Engine/Plugins/AI/MassCrowd/Source/MassCrowd/Private/MassCrowdServerRepresentationLODProcessor.cpp
--- a/Engine/Plugins/AI/MassCrowd/Source/MassCrowd/Private/MassCrowdServerRepresentationLODProcessor.cpp
+++ b/Engine/Plugins/AI/MassCrowd/Source/MassCrowd/Private/MassCrowdServerRepresentationLODProcessor.cpp
@@ -56,8 +56,8 @@ void UMassCrowdServerRepresentationLODProcessor::Execute(UMassEntitySubsystem& E
 		
 		EntityQuery.ForEachEntityChunk(EntitySubsystem, Context, [this](FMassExecutionContext& Context)
 		{
-			const TConstArrayView<FMassLODInfoFragment> ViewersInfoList = Context.GetComponentView<FMassLODInfoFragment>();
-			const TArrayView<FMassRepresentationLODFragment> RepresentationLODFragments = Context.GetMutableComponentView<FMassRepresentationLODFragment>();
+			const TConstArrayView<FMassLODInfoFragment> ViewersInfoList = Context.GetFragmentView<FMassLODInfoFragment>();
+			const TArrayView<FMassRepresentationLODFragment> RepresentationLODFragments = Context.GetMutableFragmentView<FMassRepresentationLODFragment>();
 			LODCalculator.CalculateLOD(Context, RepresentationLODFragments, ViewersInfoList);
 		});
 	}
@@ -69,8 +69,8 @@ void UMassCrowdServerRepresentationLODProcessor::Execute(UMassEntitySubsystem& E
 		{
 			EntityQuery.ForEachEntityChunk(EntitySubsystem, Context, [this](FMassExecutionContext& Context)
 			{
-				const TConstArrayView<FMassLODInfoFragment> ViewersInfoList = Context.GetComponentView<FMassLODInfoFragment>();
-				const TArrayView<FMassRepresentationLODFragment> RepresentationLODFragments = Context.GetMutableComponentView<FMassRepresentationLODFragment>();
+				const TConstArrayView<FMassLODInfoFragment> ViewersInfoList = Context.GetFragmentView<FMassLODInfoFragment>();
+				const TArrayView<FMassRepresentationLODFragment> RepresentationLODFragments = Context.GetMutableFragmentView<FMassRepresentationLODFragment>();
 				LODCalculator.AdjustLODFromCount(Context, RepresentationLODFragments, ViewersInfoList);
 			});
 		}
@@ -83,8 +83,8 @@ void UMassCrowdServerRepresentationLODProcessor::Execute(UMassEntitySubsystem& E
 		
 		EntityQuery.ForEachEntityChunk(EntitySubsystem, Context, [this](FMassExecutionContext& Context)
 		{
-			const TConstArrayView<FDataFragment_Transform> LocationList = Context.GetComponentView<FDataFragment_Transform>();
-			const TConstArrayView<FMassRepresentationLODFragment> RepresentationLODFragments = Context.GetComponentView<FMassRepresentationLODFragment>();
+			const TConstArrayView<FDataFragment_Transform> LocationList = Context.GetFragmentView<FDataFragment_Transform>();
+			const TConstArrayView<FMassRepresentationLODFragment> RepresentationLODFragments = Context.GetFragmentView<FMassRepresentationLODFragment>();
 			LODCalculator.DebugDisplayLOD(Context, RepresentationLODFragments, LocationList, World);
 		});
 	}
diff --git a/Engine/Plugins/AI/MassCrowd/Source/MassCrowd/Private/MassCrowdVisualizationLODProcessor.cpp b/Engine/Plugins/AI/MassCrowd/Source/MassCrowd/Private/MassCrowdVisualizationLODProcessor.cpp
--- a/Engine/Plugins/AI/MassCrowd/Source/MassCrowd/Private/MassCrowdVisualizationLODProcessor.cpp
+++ b/Engine/Plugins/AI/MassCrowd/Source/MassCrowd/Private/MassCrowdVisualizationLODProcessor.cpp
@@ -82,8 +82,8 @@ void UMassCrowdVisualizationLODProcessor::Execute(UMassEntitySubsystem& EntitySu
 
 		auto DebugDisplayLOD = [this](FMassExecutionContext& Context)
 		{
-			const TConstArrayView<FDataFragment_Transform> LocationList = Context.GetComponentView<FDataFragment_Transform>();
-			const TConstArrayView<FMassRepresentationLODFragment> VisualizationLODList = Context.GetComponentView<FMassRepresentationLODFragment>();
+			const TConstArrayView<FDataFragment_Transform> LocationList = Context.GetFragmentView<FDataFragment_Transform>();
+			const TConstArrayView<FMassRepresentationLODFragment> VisualizationLODList = Context.GetFragmentView<FMassRepresentationLODFragment>();
 			LODCalculator.DebugDisplayLOD(Context, VisualizationLODList, LocationList, World);
 		};
 
@@ -97,9 +97,9 @@ void UMassCrowdVisualizationLODProcessor::Execute(UMassEntitySubsystem& EntitySu
 
 		auto ShowISMUnderSpecifiedRange = [this](const FMassExecutionContext& Context)
 		{
-			const TConstArrayView<FDataFragment_Transform> LocationList = Context.GetComponentView<FDataFragment_Transform>();
-			const TConstArrayView<FMassRepresentationFragment> RepresentationFragmentList = Context.GetComponentView<FMassRepresentationFragment>();
-			const TConstArrayView<FMassRepresentationLODFragment> RepresentationLODFragmentList = Context.GetComponentView<FMassRepresentationLODFragment>();
+			const TConstArrayView<FDataFragment_Transform> LocationList = Context.GetFragmentView<FDataFragment_Transform>();
+			const TConstArrayView<FMassRepresentationFragment> RepresentationFragmentList = Context.GetFragmentView<FMassRepresentationFragment>();
+			const TConstArrayView<FMassRepresentationLODFragment> RepresentationLODFragmentList = Context.GetFragmentView<FMassRepresentationLODFragment>();
 			const int32 NumEntities = Context.GetNumEntities();
 			const float SpecifiedRangeSquaredCentimeters = FMath::Square(UE::MassCrowd::bDebugShowISMUnderSpecifiedRange * 100);
 			for (int EntityIdx = 0; EntityIdx < NumEntities; EntityIdx++)
diff --git a/Engine/Plugins/AI/MassCrowd/Source/MassCrowd/Private/MassCrowdVisualizationProcessor.cpp b/Engine/Plugins/AI/MassCrowd/Source/MassCrowd/Private/MassCrowdVisualizationProcessor.cpp
--- a/Engine/Plugins/AI/MassCrowd/Source/MassCrowd/Private/MassCrowdVisualizationProcessor.cpp
+++ b/Engine/Plugins/AI/MassCrowd/Source/MassCrowd/Private/MassCrowdVisualizationProcessor.cpp
@@ -53,9 +53,9 @@ void UMassCrowdVisualizationProcessor::Execute(UMassEntitySubsystem& EntitySubsy
 
 		EntityQuery.ForEachEntityChunk(EntitySubsystem, Context, [this](const FMassExecutionContext& Context)
 		{
-			const TConstArrayView<FMassRepresentationFragment> VisualizationList = Context.GetComponentView<FMassRepresentationFragment>();
-			const TConstArrayView<FDataFragment_Actor> ActorList = Context.GetComponentView<FDataFragment_Actor>();
-			const TConstArrayView<FDataFragment_Transform> EntityLocationList = Context.GetComponentView<FDataFragment_Transform>();
+			const TConstArrayView<FMassRepresentationFragment> VisualizationList = Context.GetFragmentView<FMassRepresentationFragment>();
+			const TConstArrayView<FDataFragment_Actor> ActorList = Context.GetFragmentView<FDataFragment_Actor>();
+			const TConstArrayView<FDataFragment_Transform> EntityLocationList = Context.GetFragmentView<FDataFragment_Transform>();
 
 			const int32 NumEntities = Context.GetNumEntities();
 			for (int EntityIdx = 0; EntityIdx < NumEntities; EntityIdx++)
diff --git a/Engine/Plugins/AI/MassCrowd/Source/MassCrowd/Public/MassCrowdRepresentationProcessor.h b/Engine/Plugins/AI/MassCrowd/Source/MassCrowd/Public/MassCrowdRepresentationProcessor.h
--- a/Engine/Plugins/AI/MassCrowd/Source/MassCrowd/Public/MassCrowdRepresentationProcessor.h
+++ b/Engine/Plugins/AI/MassCrowd/Source/MassCrowd/Public/MassCrowdRepresentationProcessor.h
@@ -34,7 +34,7 @@ protected:
 	 * @param EnabledType is the type of enabling to do on this actor
 	 * @param Actor is the actual actor to perform enabling type on
 	 * @param EntityIdx is the entity index currently processing
-	 * @param Context is the current light weight component execution context
+	 * @param Context is the current Mass execution context
 	 */
 	virtual void SetActorEnabled(const EActorEnabledType EnabledType, AActor& Actor, const int32 EntityIdx, FMassExecutionContext& Context) override;
 
@@ -56,7 +56,7 @@ protected:
 	 * The destination will be adjusted to fit an existing capsule.
 	 * @param Transform is the new actor's transform 
 	 * @param Actor is the actual actor to teleport
-	 * @param Context is the current light weight component execution context
+	 * @param Context is the current Mass execution context
 	 */
 	virtual void TeleportActor(const FTransform& Transform, AActor& Actor, FMassExecutionContext& Context) override;
 
