Removed MassSimulationLODFragment handles from both MassZoneGraphPathFollowTask and MassZoneGraphStandTask as unused

diff --git a/Engine/Plugins/AI/MassAI/Source/MassAIBehavior/Public/Tasks/MassZoneGraphPathFollowTask.h b/Engine/Plugins/AI/MassAI/Source/MassAIBehavior/Public/Tasks/MassZoneGraphPathFollowTask.h
--- a/Engine/Plugins/AI/MassAI/Source/MassAIBehavior/Public/Tasks/MassZoneGraphPathFollowTask.h
+++ b/Engine/Plugins/AI/MassAI/Source/MassAIBehavior/Public/Tasks/MassZoneGraphPathFollowTask.h
@@ -94,9 +94,6 @@ protected:
 	UPROPERTY(meta=(BaseClass="ZoneGraphSubsystem"))
 	FStateTreeExternalItemHandle ZoneGraphSubsystemHandle;
 
-	UPROPERTY(meta = (BaseStruct = "MassSimulationLODFragment"))
-	FStateTreeExternalItemHandle SimulationLODHandle;
-
 	UPROPERTY(EditAnywhere, Category = Parameters, meta=(Struct="MassZoneGraphTargetLocation", Bindable))
 	FStateTreeResultRef TargetLocation;
 
diff --git a/Engine/Plugins/AI/MassAI/Source/MassAIBehavior/Public/Tasks/MassZoneGraphStandTask.h b/Engine/Plugins/AI/MassAI/Source/MassAIBehavior/Public/Tasks/MassZoneGraphStandTask.h
--- a/Engine/Plugins/AI/MassAI/Source/MassAIBehavior/Public/Tasks/MassZoneGraphStandTask.h
+++ b/Engine/Plugins/AI/MassAI/Source/MassAIBehavior/Public/Tasks/MassZoneGraphStandTask.h
@@ -38,9 +38,6 @@ protected:
 	UPROPERTY(meta=(BaseClass="MassSignalSubsystem"))
 	FStateTreeExternalItemHandle MassSignalSubsystemHandle;
 
-	UPROPERTY(meta = (BaseStruct = "MassSimulationLODFragment"))
-	FStateTreeExternalItemHandle SimulationLODHandle;
-
 	UPROPERTY(meta=(BaseStruct="MassMovementConfigFragment"))
 	FStateTreeExternalItemHandle MovementConfigHandle;
 
