Upped "Entity XXX Invalid movement config" logging's verbosity to Warning.

diff --git a/Engine/Plugins/AI/MassAI/Source/MassAIMovement/Private/MassAvoidanceProcessors.cpp b/Engine/Plugins/AI/MassAI/Source/MassAIMovement/Private/MassAvoidanceProcessors.cpp
--- a/Engine/Plugins/AI/MassAI/Source/MassAIMovement/Private/MassAvoidanceProcessors.cpp
+++ b/Engine/Plugins/AI/MassAI/Source/MassAIMovement/Private/MassAvoidanceProcessors.cpp
@@ -568,7 +568,7 @@ void UMassAvoidanceProcessor::Execute(UMassEntitySubsystem& EntitySubsystem, FMa
 			if (!CurrentMovementConfig)
 			{
 #if WITH_MASSGAMEPLAY_DEBUG && UNSAFE_FOR_MT
-				UE_VLOG(this, LogMassNavigation, Log, TEXT("%s Invalid movement config."), *Entity.DebugGetDescription());
+				UE_VLOG(this, LogMassNavigation, Warning, TEXT("%s Invalid movement config."), *Entity.DebugGetDescription());
 #endif
 				continue;
 			}
diff --git a/Engine/Plugins/AI/MassAI/Source/MassAIMovement/Private/MassZoneGraphMovementProcessors.cpp b/Engine/Plugins/AI/MassAI/Source/MassAIMovement/Private/MassZoneGraphMovementProcessors.cpp
--- a/Engine/Plugins/AI/MassAI/Source/MassAIMovement/Private/MassZoneGraphMovementProcessors.cpp
+++ b/Engine/Plugins/AI/MassAI/Source/MassAIMovement/Private/MassZoneGraphMovementProcessors.cpp
@@ -127,7 +127,7 @@ void UMassZoneGraphLocationInitializer::Execute(UMassEntitySubsystem& EntitySubs
 			{
 #if WITH_MASSGAMEPLAY_DEBUG && UNSAFE_FOR_MT
 				const FMassEntityHandle Entity = Context.GetEntity(EntityIndex);
-				UE_VLOG(this, LogMassNavigation, Log, TEXT("Entity [%s] Invalid movement config."), *Entity.DebugGetDescription());
+				UE_VLOG(this, LogMassNavigation, Warning, TEXT("Entity [%s] Invalid movement config."), *Entity.DebugGetDescription());
 #endif
 				continue;
 			}
