Added MoveTargetFragment default initializer processor, MassMoveTargetFragmentInitializer

diff --git a/Engine/Plugins/AI/MassAI/Source/MassAIMovement/Private/MassMoveTargetFragmentInitializer.cpp b/Engine/Plugins/AI/MassAI/Source/MassAIMovement/Private/MassMoveTargetFragmentInitializer.cpp
new file mode 100644
--- /dev/null
+++ b/Engine/Plugins/AI/MassAI/Source/MassAIMovement/Private/MassMoveTargetFragmentInitializer.cpp
@@ -0,0 +1,39 @@
+// Copyright Epic Games, Inc. All Rights Reserved.
+
+#include "MassMoveTargetFragmentInitializer.h"
+#include "MassCommonFragments.h"
+#include "MassAIMovementFragments.h"
+#include "MassSimulationLOD.h"
+
+
+UMassMoveTargetFragmentInitializer::UMassMoveTargetFragmentInitializer()
+{
+	FragmentType = FMassMoveTargetFragment::StaticStruct();
+}
+
+void UMassMoveTargetFragmentInitializer::ConfigureQueries()
+{
+	InitializerQuery.AddRequirement<FMassMoveTargetFragment>(EMassFragmentAccess::ReadWrite);
+	InitializerQuery.AddRequirement<FDataFragment_Transform>(EMassFragmentAccess::ReadOnly);
+}
+
+void UMassMoveTargetFragmentInitializer::Execute(UMassEntitySubsystem& EntitySubsystem, FMassExecutionContext& Context)
+{
+	InitializerQuery.ForEachEntityChunk(EntitySubsystem, Context, [](FMassExecutionContext& Context)
+	{
+		const int32 NumEntities = Context.GetNumEntities();
+		TArrayView<FMassMoveTargetFragment> MoveTargetList = Context.GetMutableFragmentView<FMassMoveTargetFragment>();
+		const TConstArrayView<FDataFragment_Transform> LocationList = Context.GetFragmentView<FDataFragment_Transform>();
+
+		for (int32 EntityIndex = 0; EntityIndex < NumEntities; ++EntityIndex)
+		{
+			FMassMoveTargetFragment& MoveTarget = MoveTargetList[EntityIndex];
+			const FDataFragment_Transform& Location = LocationList[EntityIndex];
+
+			MoveTarget.Center = Location.GetTransform().GetLocation();
+			MoveTarget.Forward = Location.GetTransform().GetRotation().Vector();
+			MoveTarget.DistanceToGoal = 0.0f;
+			MoveTarget.SlackRadius = 0.0f;
+		}
+	});
+}
diff --git a/Engine/Plugins/AI/MassAI/Source/MassAIMovement/Public/MassMoveTargetFragmentInitializer.h b/Engine/Plugins/AI/MassAI/Source/MassAIMovement/Public/MassMoveTargetFragmentInitializer.h
new file mode 100644
--- /dev/null
+++ b/Engine/Plugins/AI/MassAI/Source/MassAIMovement/Public/MassMoveTargetFragmentInitializer.h
@@ -0,0 +1,21 @@
+// Copyright Epic Games, Inc. All Rights Reserved.
+
+#pragma once 
+
+#include "MassTranslator.h"
+#include "MassMoveTargetFragmentInitializer.generated.h"
+
+
+UCLASS()
+class MASSAIMOVEMENT_API UMassMoveTargetFragmentInitializer : public UMassFragmentInitializer
+{
+	GENERATED_BODY()
+public:
+	UMassMoveTargetFragmentInitializer();
+	
+protected:
+	virtual void ConfigureQueries() override;
+	virtual void Execute(UMassEntitySubsystem& EntitySubsystem, FMassExecutionContext& Context) override;
+
+	FMassEntityQuery InitializerQuery;
+};
