Made UMassZoneGraphPathFollowProcessor's LOD dependency optional (i.e. works without it).
Also, GetOptional*ChunkFragment functions renamed to Get*ChunkFragmentPtr.

diff --git a/Engine/Plugins/AI/MassAI/Source/MassAIMovement/Private/MassZoneGraphMovementProcessors.cpp b/Engine/Plugins/AI/MassAI/Source/MassAIMovement/Private/MassZoneGraphMovementProcessors.cpp
--- a/Engine/Plugins/AI/MassAI/Source/MassAIMovement/Private/MassZoneGraphMovementProcessors.cpp
+++ b/Engine/Plugins/AI/MassAI/Source/MassAIMovement/Private/MassZoneGraphMovementProcessors.cpp
@@ -190,8 +190,9 @@ void UMassZoneGraphPathFollowProcessor::ConfigureQueries()
 	EntityQuery_Conditional.AddRequirement<FMassZoneGraphShortPathFragment>(EMassFragmentAccess::ReadWrite);
 	EntityQuery_Conditional.AddRequirement<FMassZoneGraphLaneLocationFragment>(EMassFragmentAccess::ReadWrite);
 	EntityQuery_Conditional.AddRequirement<FMassMoveTargetFragment>(EMassFragmentAccess::ReadWrite);
-	EntityQuery_Conditional.AddRequirement<FMassSimulationLODFragment>(EMassFragmentAccess::ReadOnly);
-	EntityQuery_Conditional.AddChunkRequirement<FMassSimulationVariableTickChunkFragment>(EMassFragmentAccess::ReadOnly);
+	EntityQuery_Conditional.AddRequirement<FMassSimulationLODFragment>(EMassFragmentAccess::ReadOnly, EMassFragmentPresence::Optional);
+
+	EntityQuery_Conditional.AddChunkRequirement<FMassSimulationVariableTickChunkFragment>(EMassFragmentAccess::ReadOnly, EMassFragmentPresence::Optional);
 	EntityQuery_Conditional.SetChunkFilter(&FMassSimulationVariableTickChunkFragment::ShouldTickChunkThisFrame);
 }
 
@@ -212,15 +213,16 @@ void UMassZoneGraphPathFollowProcessor::Execute(UMassEntitySubsystem& EntitySubs
 		const TArrayView<FMassZoneGraphLaneLocationFragment> LaneLocationList = Context.GetMutableFragmentView<FMassZoneGraphLaneLocationFragment>();
 		const TArrayView<FMassMoveTargetFragment> MoveTargetList = Context.GetMutableFragmentView<FMassMoveTargetFragment>();
 		const TConstArrayView<FMassSimulationLODFragment> SimLODList = Context.GetFragmentView<FMassSimulationLODFragment>();
+		const bool bHasLOD = (SimLODList.Num() > 0);
+		const float WorldDeltaTime = Context.GetDeltaTimeSeconds();
 
 		for (int32 EntityIndex = 0; EntityIndex < NumEntities; ++EntityIndex)
 		{
 			FMassZoneGraphShortPathFragment& ShortPath = ShortPathList[EntityIndex];
 			FMassZoneGraphLaneLocationFragment& LaneLocation = LaneLocationList[EntityIndex];
 			FMassMoveTargetFragment& MoveTarget = MoveTargetList[EntityIndex];
-			const FMassSimulationLODFragment& SimLOD = SimLODList[EntityIndex];
 			const FMassEntityHandle Entity = Context.GetEntity(EntityIndex);
-			const float DeltaTime = SimLOD.DeltaTime;
+			const float DeltaTime = bHasLOD ? SimLODList[EntityIndex].DeltaTime : WorldDeltaTime;
 
 			bool bDisplayDebug = false;
 #if WITH_MASSGAMEPLAY_DEBUG && UNSAFE_FOR_MT // this will result in bDisplayDebug == false and disabling of all the vlogs below
@@ -237,7 +239,7 @@ void UMassZoneGraphPathFollowProcessor::Execute(UMassEntitySubsystem& EntitySubs
 				const bool bWasDone = ShortPath.IsDone();
 
 				// Note: this should be in sync with the logic in apply velocity.
-				const bool bHasSteering = SimLOD.LOD != EMassLOD::Off;
+				const bool bHasSteering = (bHasLOD == false) || (SimLODList[EntityIndex].LOD != EMassLOD::Off);
 
 				if (!bHasSteering || !MoveTarget.bSteeringFallingBehind)
 				{
@@ -533,7 +535,7 @@ void UMassZoneGraphSteeringProcessor::Execute(UMassEntitySubsystem& EntitySubsys
 			if (!CurrentMovementConfig)
 			{
 #if WITH_MASSGAMEPLAY_DEBUG && UNSAFE_FOR_MT
-				UE_VLOG(this, LogMassNavigation, Log, TEXT("Entity [%s] Invalid movement config."), *Entity.DebugGetDescription());
+				UE_VLOG(this, LogMassNavigation, Warning, TEXT("Entity [%s] Invalid movement config."), *Entity.DebugGetDescription());
 #endif
 				continue;
 			}
diff --git a/Engine/Plugins/Runtime/MassEntity/Source/MassEntity/Public/MassEntitySubsystem.h b/Engine/Plugins/Runtime/MassEntity/Source/MassEntity/Public/MassEntitySubsystem.h
--- a/Engine/Plugins/Runtime/MassEntity/Source/MassEntity/Public/MassEntitySubsystem.h
+++ b/Engine/Plugins/Runtime/MassEntity/Source/MassEntity/Public/MassEntitySubsystem.h
@@ -416,7 +416,7 @@ public:
 	int32 GetChunkSerialModificationNumber() const { return ChunkSerialModificationNumber; }
 
 	template<typename T>
-	T* GetOptionalMutableChunkFragment() const
+	T* GetMutableChunkFragmentPtr() const
 	{
 		static_assert(TIsDerivedFrom<T, FMassChunkFragment>::IsDerived, "Given struct doesn't represent a valid chunk fragment type. Make sure to inherit from FMassChunkFragment or one of its child-types.");
 
@@ -428,13 +428,13 @@ public:
 	template<typename T>
 	T& GetMutableChunkFragment() const
 	{
-		T* ChunkFragment = GetOptionalMutableChunkFragment<T>();
+		T* ChunkFragment = GetMutableChunkFragmentPtr<T>();
 		checkf(ChunkFragment, TEXT("Chunk Fragment requirement not found: %s"), *T::StaticStruct()->GetName());
 		return *ChunkFragment;
 	}
 
 	template<typename T>
-	const T* GetOptionalChunkFragment() const
+	const T* GetChunkFragmentPtr() const
 	{
 		static_assert(TIsDerivedFrom<T, FMassChunkFragment>::IsDerived, "Given struct doesn't represent a valid chunk fragment type. Make sure to inherit from FMassChunkFragment or one of its child-types.");
 
@@ -446,7 +446,7 @@ public:
 	template<typename T>
 	const T& GetChunkFragment() const
 	{
-		const T* ChunkFragment = GetOptionalChunkFragment<T>();
+		const T* ChunkFragment = GetChunkFragmentPtr<T>();
 		checkf(ChunkFragment, TEXT("Chunk Fragment requirement not found: %s"), *T::StaticStruct()->GetName());
 		return *ChunkFragment;
 	}
diff --git a/Engine/Plugins/Runtime/MassGameplay/Source/MassLOD/Public/MassSimulationLOD.h b/Engine/Plugins/Runtime/MassGameplay/Source/MassLOD/Public/MassSimulationLOD.h
--- a/Engine/Plugins/Runtime/MassGameplay/Source/MassLOD/Public/MassSimulationLOD.h
+++ b/Engine/Plugins/Runtime/MassGameplay/Source/MassLOD/Public/MassSimulationLOD.h
@@ -73,10 +73,14 @@ struct FMassSimulationVariableTickChunkFragment : public FMassVariableTickChunkF
 {
 	GENERATED_BODY();
 
+	/**
+	 * @return if Context contains an instance of FMassSimulationVariableTickChunkFragment then the function returns
+	 *	the value of FMassVariableTickChunkFragment.ShouldTickThisFrame; otherwise it returns DefaultValue.
+	 */
 	static bool ShouldTickChunkThisFrame(const FMassExecutionContext& Context)
 	{
-		const FMassSimulationVariableTickChunkFragment& ChunkFragment = Context.GetChunkFragment<FMassSimulationVariableTickChunkFragment>();
-		return ChunkFragment.ShouldTickThisFrame();
+		const FMassSimulationVariableTickChunkFragment* ChunkFragment = Context.GetChunkFragmentPtr<FMassSimulationVariableTickChunkFragment>();
+		return ChunkFragment == nullptr || ChunkFragment->ShouldTickThisFrame();
 	}
 
 	static EMassLOD::Type GetChunkLOD(const FMassExecutionContext& Context)
