StateTree: Resolve external dependecies at load time instead of bake time
- added StateTree linker
- added Link() method to tasks and evaluators
- added templated item reference which makes getter code less error prone

diff --git a/Engine/Plugins/AI/MassAI/Source/MassAIBehavior/Private/Evaluators/MassComponentHitEvaluator.cpp b/Engine/Plugins/AI/MassAI/Source/MassAIBehavior/Private/Evaluators/MassComponentHitEvaluator.cpp
--- a/Engine/Plugins/AI/MassAI/Source/MassAIBehavior/Private/Evaluators/MassComponentHitEvaluator.cpp
+++ b/Engine/Plugins/AI/MassAI/Source/MassAIBehavior/Private/Evaluators/MassComponentHitEvaluator.cpp
@@ -6,10 +6,18 @@
 #include "MassStateTreeExecutionContext.h"
 #include "Engine/World.h"
 
+
+bool FMassComponentHitEvaluator::Link(FStateTreeLinker& Linker)
+{
+	Linker.LinkExternalItem(ComponentHitSubsystemHandle);
+
+	return true;
+}
+
 void FMassComponentHitEvaluator::Evaluate(FStateTreeExecutionContext &Context, const EStateTreeEvaluationType EvalType, const float DeltaTime)
 {
 	// Look for recent hits
-	UMassComponentHitSubsystem& HitSubsystem = Context.GetExternalItem(ComponentHitSubsystemHandle).GetMutable<UMassComponentHitSubsystem>();
+	UMassComponentHitSubsystem& HitSubsystem = Context.GetExternalItem(ComponentHitSubsystemHandle);
 	const FMassHitResult* HitResult = HitSubsystem.GetLastHit(static_cast<FMassStateTreeExecutionContext&>(Context).GetEntity());
 
 	// LastHitEntity is not reset intentionally, so that it's available the duration of the behavior reacting to it.
diff --git a/Engine/Plugins/AI/MassAI/Source/MassAIBehavior/Private/Evaluators/MassStateTreeSmartObjectEvaluator.cpp b/Engine/Plugins/AI/MassAI/Source/MassAIBehavior/Private/Evaluators/MassStateTreeSmartObjectEvaluator.cpp
--- a/Engine/Plugins/AI/MassAI/Source/MassAIBehavior/Private/Evaluators/MassStateTreeSmartObjectEvaluator.cpp
+++ b/Engine/Plugins/AI/MassAI/Source/MassAIBehavior/Private/Evaluators/MassStateTreeSmartObjectEvaluator.cpp
@@ -17,6 +17,18 @@
 //----------------------------------------------------------------------//
 // FMassStateTreeSmartObjectEvaluator
 //----------------------------------------------------------------------//
+
+bool FMassStateTreeSmartObjectEvaluator::Link(FStateTreeLinker& Linker)
+{
+	Linker.LinkExternalItem(SmartObjectSubsystemHandle);
+	Linker.LinkExternalItem(MassSignalSubsystemHandle);
+	Linker.LinkExternalItem(EntityTransformHandle);
+	Linker.LinkExternalItem(SmartObjectUserHandle);
+	Linker.LinkExternalItem(LocationHandle);
+
+	return true;
+}
+
 void FMassStateTreeSmartObjectEvaluator::EnterState(FStateTreeExecutionContext& Context, const EStateTreeStateChangeType ChangeType, const FStateTreeTransitionResult& Transition)
 {
 	if (ChangeType != EStateTreeStateChangeType::Changed)
@@ -37,7 +49,7 @@ void FMassStateTreeSmartObjectEvaluator::ExitState(FStateTreeExecutionContext& C
 	if (SearchRequestID.IsSet())
 	{
 		const FMassStateTreeExecutionContext& MassContext = static_cast<FMassStateTreeExecutionContext&>(Context);
-		USmartObjectSubsystem& SmartObjectSubsystem = Context.GetExternalItem(SmartObjectSubsystemHandle).GetMutable<USmartObjectSubsystem>();
+		USmartObjectSubsystem& SmartObjectSubsystem = Context.GetExternalItem(SmartObjectSubsystemHandle);
 		const FMassSmartObjectHandler MassSmartObjectHandler(MassContext.GetEntitySubsystem(), MassContext.GetEntitySubsystemExecutionContext(), SmartObjectSubsystem);
 		MassSmartObjectHandler.RemoveRequest(SearchRequestID);
 		SearchRequestID.Reset();
@@ -54,7 +66,7 @@ void FMassStateTreeSmartObjectEvaluator::Reset()
 
 void FMassStateTreeSmartObjectEvaluator::Evaluate(FStateTreeExecutionContext& Context, const EStateTreeEvaluationType EvalType,  const float DeltaTime)
 {
-	const FDataFragment_SmartObjectUser& SOUser = Context.GetExternalItem(SmartObjectUserHandle).Get<FDataFragment_SmartObjectUser>();
+	const FDataFragment_SmartObjectUser& SOUser = Context.GetExternalItem(SmartObjectUserHandle);
 
 	bCandidatesFound = false;
 	bClaimed = SOUser.GetClaimHandle().IsValid();
@@ -78,7 +90,7 @@ void FMassStateTreeSmartObjectEvaluator::Evaluate(FStateTreeExecutionContext& Co
 	}
 	NextUpdate = 0.f;
 
-	USmartObjectSubsystem& SmartObjectSubsystem = Context.GetExternalItem(SmartObjectSubsystemHandle).GetMutable<USmartObjectSubsystem>();
+	USmartObjectSubsystem& SmartObjectSubsystem = Context.GetExternalItem(SmartObjectSubsystemHandle);
 	const FMassStateTreeExecutionContext& MassContext = static_cast<FMassStateTreeExecutionContext&>(Context);
 	const FMassSmartObjectHandler MassSmartObjectHandler(MassContext.GetEntitySubsystem(), MassContext.GetEntitySubsystemExecutionContext(), SmartObjectSubsystem);
 
@@ -87,7 +99,7 @@ void FMassStateTreeSmartObjectEvaluator::Evaluate(FStateTreeExecutionContext& Co
 	{
 		// Use lanes if possible for faster queries using zone graph annotations
 		const FMassEntityHandle RequestingEntity = MassContext.GetEntity();
-		const FMassZoneGraphLaneLocationFragment* LaneLocation = Context.GetExternalItem(LocationHandle).GetPtr<FMassZoneGraphLaneLocationFragment>();
+		const FMassZoneGraphLaneLocationFragment* LaneLocation = Context.GetExternalItemPtr(LocationHandle);
 		bUsingZoneGraphAnnotations = LaneLocation != nullptr;
 		if (bUsingZoneGraphAnnotations)
 		{
@@ -99,7 +111,7 @@ void FMassStateTreeSmartObjectEvaluator::Evaluate(FStateTreeExecutionContext& Co
 		}
 		else
 		{
-			const FDataFragment_Transform& TransformFragment = Context.GetExternalItem(EntityTransformHandle).Get<FDataFragment_Transform>();
+			const FDataFragment_Transform& TransformFragment = Context.GetExternalItem(EntityTransformHandle);
 			SearchRequestID = MassSmartObjectHandler.FindCandidatesAsync(RequestingEntity, TransformFragment.GetTransform().GetLocation());
 		}
 	}
@@ -127,7 +139,7 @@ void FMassStateTreeSmartObjectEvaluator::Evaluate(FStateTreeExecutionContext& Co
 			{
 				const float DelayInSeconds = bCandidatesFound ? TickInterval : RetryCooldown;
 				NextUpdate = World->GetTimeSeconds() + DelayInSeconds;
-				UMassSignalSubsystem& MassSignalSubsystem = Context.GetExternalItem(MassSignalSubsystemHandle).GetMutable<UMassSignalSubsystem>();
+				UMassSignalSubsystem& MassSignalSubsystem = Context.GetExternalItem(MassSignalSubsystemHandle);
 				MassSignalSubsystem.DelaySignalEntity(UE::Mass::Signals::SmartObjectRequestCandidates, RequestingEntity, DelayInSeconds);
 			}
 		}
diff --git a/Engine/Plugins/AI/MassAI/Source/MassAIBehavior/Private/Evaluators/MassZoneGraphAnnotationEvaluator.cpp b/Engine/Plugins/AI/MassAI/Source/MassAIBehavior/Private/Evaluators/MassZoneGraphAnnotationEvaluator.cpp
--- a/Engine/Plugins/AI/MassAI/Source/MassAIBehavior/Private/Evaluators/MassZoneGraphAnnotationEvaluator.cpp
+++ b/Engine/Plugins/AI/MassAI/Source/MassAIBehavior/Private/Evaluators/MassZoneGraphAnnotationEvaluator.cpp
@@ -9,8 +9,15 @@ FMassZoneGraphAnnotationEvaluator::FMassZoneGraphAnnotationEvaluator()
 {
 }
 
+bool FMassZoneGraphAnnotationEvaluator::Link(FStateTreeLinker& Linker)
+{
+	Linker.LinkExternalItem(BehaviorTagsHandle);
+
+	return true;
+}
+
 void FMassZoneGraphAnnotationEvaluator::Evaluate(FStateTreeExecutionContext& Context, const EStateTreeEvaluationType EvalType, const float DeltaTime)
 {
-	const FMassZoneGraphAnnotationTagsFragment& BehaviorTagsFragment = Context.GetExternalItem(BehaviorTagsHandle).Get<FMassZoneGraphAnnotationTagsFragment>();
+	const FMassZoneGraphAnnotationTagsFragment& BehaviorTagsFragment = Context.GetExternalItem(BehaviorTagsHandle);
 	BehaviorTags = BehaviorTagsFragment.Tags;
 }
diff --git a/Engine/Plugins/AI/MassAI/Source/MassAIBehavior/Private/MassStateTreeProcessors.cpp b/Engine/Plugins/AI/MassAI/Source/MassAIBehavior/Private/MassStateTreeProcessors.cpp
--- a/Engine/Plugins/AI/MassAI/Source/MassAIBehavior/Private/MassStateTreeProcessors.cpp
+++ b/Engine/Plugins/AI/MassAI/Source/MassAIBehavior/Private/MassStateTreeProcessors.cpp
@@ -34,7 +34,7 @@ bool SetExternalFragments(FMassStateTreeExecutionContext& Context, const UMassEn
 			}
 			else
 			{
-				if (!ItemDesc.bOptional)
+				if (ItemDesc.Requirement == EStateTreeItemRequirement::Required)
 				{
 					// Note: Not breaking here, so that we can validate all missing ones in one go with FMassStateTreeExecutionContext::AreExternalItemsValid().
 					bFoundAllFragments = false;
@@ -66,7 +66,7 @@ bool SetExternalSubsystems(FMassStateTreeExecutionContext& Context)
 			}
 			else
 			{
-				if (!ItemDesc.bOptional)
+				if (ItemDesc.Requirement == EStateTreeItemRequirement::Required)
 				{
 					// Note: Not breaking here, so that we can validate all missing ones in one go with FMassStateTreeExecutionContext::AreExternalItemsValid().
 					bFoundAllSubsystems = false;
diff --git a/Engine/Plugins/AI/MassAI/Source/MassAIBehavior/Private/MassStateTreeTrait.cpp b/Engine/Plugins/AI/MassAI/Source/MassAIBehavior/Private/MassStateTreeTrait.cpp
--- a/Engine/Plugins/AI/MassAI/Source/MassAIBehavior/Private/MassStateTreeTrait.cpp
+++ b/Engine/Plugins/AI/MassAI/Source/MassAIBehavior/Private/MassStateTreeTrait.cpp
@@ -57,7 +57,7 @@ void UMassStateTreeTrait::ValidateTemplate(FMassEntityTemplateBuildContext& Buil
 	// Make sure all the required subsystems can be found.
 	for (const FStateTreeExternalItemDesc& ItemDesc : StateTree->GetExternalItems())
 	{
-		if (ensure(ItemDesc.Struct) && !ItemDesc.bOptional)
+		if (ensure(ItemDesc.Struct) && ItemDesc.Requirement == EStateTreeItemRequirement::Required)
 		{
 			if (ItemDesc.Struct->IsChildOf(UWorldSubsystem::StaticClass()))
 			{
diff --git a/Engine/Plugins/AI/MassAI/Source/MassAIBehavior/Private/Tasks/MassLookAtTask.cpp b/Engine/Plugins/AI/MassAI/Source/MassAIBehavior/Private/Tasks/MassLookAtTask.cpp
--- a/Engine/Plugins/AI/MassAI/Source/MassAIBehavior/Private/Tasks/MassLookAtTask.cpp
+++ b/Engine/Plugins/AI/MassAI/Source/MassAIBehavior/Private/Tasks/MassLookAtTask.cpp
@@ -6,12 +6,20 @@
 #include "MassSignalSubsystem.h"
 #include "MassStateTreeExecutionContext.h"
 
+bool FMassLookAtTask::Link(FStateTreeLinker& Linker)
+{
+	Linker.LinkExternalItem(MassSignalSubsystemHandle);
+	Linker.LinkExternalItem(LookAtHandle);
+
+	return true;
+}
+
 EStateTreeRunStatus FMassLookAtTask::EnterState(FStateTreeExecutionContext& Context, const EStateTreeStateChangeType ChangeType, const FStateTreeTransitionResult& Transition)
 {
 	Time = 0.f;
 	
 	const FMassStateTreeExecutionContext& MassContext = static_cast<FMassStateTreeExecutionContext&>(Context);
-	FMassLookAtFragment& LookAtFragment = MassContext.GetExternalItem(LookAtHandle).GetMutable<FMassLookAtFragment>();
+	FMassLookAtFragment& LookAtFragment = MassContext.GetExternalItem(LookAtHandle);
 
 	LookAtFragment.Reset();
 	LookAtFragment.LookAtMode = LookAtMode;
@@ -39,7 +47,7 @@ EStateTreeRunStatus FMassLookAtTask::EnterState(FStateTreeExecutionContext& Cont
 	// Otherwise we schedule a signal to end the task.
 	if (Duration > 0.0f)
 	{
-		UMassSignalSubsystem& MassSignalSubsystem = MassContext.GetExternalItem(MassSignalSubsystemHandle).GetMutable<UMassSignalSubsystem>();
+		UMassSignalSubsystem& MassSignalSubsystem = MassContext.GetExternalItem(MassSignalSubsystemHandle);
 		MassSignalSubsystem.DelaySignalEntity(UE::Mass::Signals::LookAtFinished, MassContext.GetEntity(), Duration);
 	}
 
@@ -49,7 +57,7 @@ EStateTreeRunStatus FMassLookAtTask::EnterState(FStateTreeExecutionContext& Cont
 void FMassLookAtTask::ExitState(FStateTreeExecutionContext& Context, const EStateTreeStateChangeType ChangeType, const FStateTreeTransitionResult& Transition)
 {
 	const FMassStateTreeExecutionContext& MassContext = static_cast<FMassStateTreeExecutionContext&>(Context);
-	FMassLookAtFragment& LookAtFragment = MassContext.GetExternalItem(LookAtHandle).GetMutable<FMassLookAtFragment>();
+	FMassLookAtFragment& LookAtFragment = MassContext.GetExternalItem(LookAtHandle);
 	
 	LookAtFragment.Reset();
 }
diff --git a/Engine/Plugins/AI/MassAI/Source/MassAIBehavior/Private/Tasks/MassStateTreeSmartObjectTask.cpp b/Engine/Plugins/AI/MassAI/Source/MassAIBehavior/Private/Tasks/MassStateTreeSmartObjectTask.cpp
--- a/Engine/Plugins/AI/MassAI/Source/MassAIBehavior/Private/Tasks/MassStateTreeSmartObjectTask.cpp
+++ b/Engine/Plugins/AI/MassAI/Source/MassAIBehavior/Private/Tasks/MassStateTreeSmartObjectTask.cpp
@@ -14,6 +14,15 @@
 //----------------------------------------------------------------------//
 // FMassClaimSmartObjectTask
 //----------------------------------------------------------------------//
+
+bool FMassClaimSmartObjectTask::Link(FStateTreeLinker& Linker)
+{
+	Linker.LinkExternalItem(SmartObjectUserHandle);
+	Linker.LinkExternalItem(SmartObjectSubsystemHandle);
+
+	return true;
+}
+
 EStateTreeRunStatus FMassClaimSmartObjectTask::EnterState(FStateTreeExecutionContext& Context, const EStateTreeStateChangeType ChangeType, const FStateTreeTransitionResult& Transition)
 {
 	if (ChangeType != EStateTreeStateChangeType::Changed)
@@ -36,8 +45,8 @@ EStateTreeRunStatus FMassClaimSmartObjectTask::EnterState(FStateTreeExecutionCon
 	EStateTreeRunStatus Status = EStateTreeRunStatus::Failed;
 
 	// Retrieve fragments and subsystems
-	USmartObjectSubsystem& SmartObjectSubsystem = Context.GetExternalItem(SmartObjectSubsystemHandle).GetMutable<USmartObjectSubsystem>();
-	FDataFragment_SmartObjectUser& SOUser = Context.GetExternalItem(SmartObjectUserHandle).GetMutable<FDataFragment_SmartObjectUser>();
+	USmartObjectSubsystem& SmartObjectSubsystem = Context.GetExternalItem(SmartObjectSubsystemHandle);
+	FDataFragment_SmartObjectUser& SOUser = Context.GetExternalItem(SmartObjectUserHandle);
 
 	// Setup MassSmartObject handler and claim
 	const FMassStateTreeExecutionContext& MassContext = static_cast<FMassStateTreeExecutionContext&>(Context);
@@ -74,7 +83,7 @@ void FMassClaimSmartObjectTask::ExitState(FStateTreeExecutionContext& Context, c
 
 	ClaimResult = EMassSmartObjectClaimResult::Unset;
 
-	FDataFragment_SmartObjectUser& SOUser = Context.GetExternalItem(SmartObjectUserHandle).GetMutable<FDataFragment_SmartObjectUser>();
+	FDataFragment_SmartObjectUser& SOUser = Context.GetExternalItem(SmartObjectUserHandle);
 
 	if (SOUser.ClaimHandle.IsValid())
 	{
@@ -82,7 +91,7 @@ void FMassClaimSmartObjectTask::ExitState(FStateTreeExecutionContext& Context, c
 		{
 			MASSBEHAVIOR_LOG(VeryVerbose, TEXT("Exiting state with a valid claim handle but interaction was not started: release claim on the smart object."));
 			const FMassStateTreeExecutionContext& MassContext = static_cast<FMassStateTreeExecutionContext&>(Context);
-			USmartObjectSubsystem& SmartObjectSubsystem = Context.GetExternalItem(SmartObjectSubsystemHandle).GetMutable<USmartObjectSubsystem>();
+			USmartObjectSubsystem& SmartObjectSubsystem = Context.GetExternalItem(SmartObjectSubsystemHandle);
 			const FMassSmartObjectHandler MassSmartObjectHandler(MassContext.GetEntitySubsystem(), MassContext.GetEntitySubsystemExecutionContext(), SmartObjectSubsystem);
 			MassSmartObjectHandler.ReleaseSmartObject(MassContext.GetEntity(), SOUser);
 		}
@@ -105,6 +114,17 @@ EStateTreeRunStatus FMassClaimSmartObjectTask::Tick(FStateTreeExecutionContext&
 //----------------------------------------------------------------------//
 // FMassUseSmartObjectTask
 //----------------------------------------------------------------------//
+
+bool FMassUseSmartObjectTask::Link(FStateTreeLinker& Linker)
+{
+	Linker.LinkExternalItem(SmartObjectSubsystemHandle);
+	Linker.LinkExternalItem(EntityTransformHandle);
+	Linker.LinkExternalItem(SmartObjectUserHandle);
+	Linker.LinkExternalItem(MoveTargetHandle);
+
+	return true;
+}
+
 EStateTreeRunStatus FMassUseSmartObjectTask::EnterState(FStateTreeExecutionContext& Context, const EStateTreeStateChangeType ChangeType, const FStateTreeTransitionResult& Transition)
 {
 	if (ChangeType != EStateTreeStateChangeType::Changed)
@@ -113,10 +133,10 @@ EStateTreeRunStatus FMassUseSmartObjectTask::EnterState(FStateTreeExecutionConte
 	}
 
 	// Retrieve fragments and subsystems
-	USmartObjectSubsystem& SmartObjectSubsystem = Context.GetExternalItem(SmartObjectSubsystemHandle).GetMutable<USmartObjectSubsystem>();
-	FDataFragment_SmartObjectUser& SOUser = Context.GetExternalItem(SmartObjectUserHandle).GetMutable<FDataFragment_SmartObjectUser>();
-	const FDataFragment_Transform& TransformFragment = Context.GetExternalItem(EntityTransformHandle).Get<FDataFragment_Transform>();
-	FMassMoveTargetFragment& MoveTarget = Context.GetExternalItem(MoveTargetHandle).GetMutable<FMassMoveTargetFragment>();
+	USmartObjectSubsystem& SmartObjectSubsystem = Context.GetExternalItem(SmartObjectSubsystemHandle);
+	FDataFragment_SmartObjectUser& SOUser = Context.GetExternalItem(SmartObjectUserHandle);
+	const FDataFragment_Transform& TransformFragment = Context.GetExternalItem(EntityTransformHandle);
+	FMassMoveTargetFragment& MoveTarget = Context.GetExternalItem(MoveTargetHandle);
 
 	// Setup MassSmartObject handler and start interaction
 	const FMassStateTreeExecutionContext& MassContext = static_cast<FMassStateTreeExecutionContext&>(Context);
@@ -141,14 +161,14 @@ void FMassUseSmartObjectTask::ExitState(FStateTreeExecutionContext& Context, con
 {
 	if (ChangeType == EStateTreeStateChangeType::Changed)
 	{
-		FDataFragment_SmartObjectUser& SOUser = Context.GetExternalItem(SmartObjectUserHandle).GetMutable<FDataFragment_SmartObjectUser>();
+		FDataFragment_SmartObjectUser& SOUser = Context.GetExternalItem(SmartObjectUserHandle);
 
 		if (SOUser.ClaimHandle.IsValid())
 		{
 			MASSBEHAVIOR_LOG(VeryVerbose, TEXT("Exiting state with a valid ClaimHandle: stop using the smart object."));
 
 			const FMassStateTreeExecutionContext& MassContext = static_cast<FMassStateTreeExecutionContext&>(Context);
-			USmartObjectSubsystem& SmartObjectSubsystem = Context.GetExternalItem(SmartObjectSubsystemHandle).GetMutable<USmartObjectSubsystem>();
+			USmartObjectSubsystem& SmartObjectSubsystem = Context.GetExternalItem(SmartObjectSubsystemHandle);
 			const FMassSmartObjectHandler MassSmartObjectHandler(MassContext.GetEntitySubsystem(), MassContext.GetEntitySubsystemExecutionContext(), SmartObjectSubsystem);
 			MassSmartObjectHandler.ReleaseSmartObject(MassContext.GetEntity(), SOUser, EMassSmartObjectInteractionStatus::Aborted);
 		}
@@ -161,14 +181,14 @@ void FMassUseSmartObjectTask::ExitState(FStateTreeExecutionContext& Context, con
 
 void FMassUseSmartObjectTask::StateCompleted(FStateTreeExecutionContext& Context, const EStateTreeRunStatus CompletionStatus, const FStateTreeHandle CompletedState)
 {
-	FDataFragment_SmartObjectUser& SOUser = Context.GetExternalItem(SmartObjectUserHandle).GetMutable<FDataFragment_SmartObjectUser>();
+	FDataFragment_SmartObjectUser& SOUser = Context.GetExternalItem(SmartObjectUserHandle);
 
 	if (SOUser.ClaimHandle.IsValid())
 	{
 		MASSBEHAVIOR_LOG(VeryVerbose, TEXT("Completing state with a valid ClaimHandle: stop using the smart object."));
 
 		const FMassStateTreeExecutionContext& MassContext = static_cast<FMassStateTreeExecutionContext&>(Context);
-		USmartObjectSubsystem& SmartObjectSubsystem = Context.GetExternalItem(SmartObjectSubsystemHandle).GetMutable<USmartObjectSubsystem>();
+		USmartObjectSubsystem& SmartObjectSubsystem = Context.GetExternalItem(SmartObjectSubsystemHandle);
 		const FMassSmartObjectHandler MassSmartObjectHandler(MassContext.GetEntitySubsystem(), MassContext.GetEntitySubsystemExecutionContext(), SmartObjectSubsystem);
 		const EMassSmartObjectInteractionStatus NewStatus = CompletionStatus == EStateTreeRunStatus::Succeeded ? EMassSmartObjectInteractionStatus::Completed : EMassSmartObjectInteractionStatus::Aborted;
 		MassSmartObjectHandler.ReleaseSmartObject(MassContext.GetEntity(), SOUser, NewStatus);
@@ -179,7 +199,7 @@ EStateTreeRunStatus FMassUseSmartObjectTask::Tick(FStateTreeExecutionContext& Co
 {
 	EStateTreeRunStatus Status = EStateTreeRunStatus::Failed;
 
-	FDataFragment_SmartObjectUser& SOUser = Context.GetExternalItem(SmartObjectUserHandle).GetMutable<FDataFragment_SmartObjectUser>();
+	FDataFragment_SmartObjectUser& SOUser = Context.GetExternalItem(SmartObjectUserHandle);
 	switch (SOUser.InteractionStatus)
 	{
 	case EMassSmartObjectInteractionStatus::InProgress:
diff --git a/Engine/Plugins/AI/MassAI/Source/MassAIBehavior/Private/Tasks/MassStateTreeTestTask.cpp b/Engine/Plugins/AI/MassAI/Source/MassAIBehavior/Private/Tasks/MassStateTreeTestTask.cpp
--- a/Engine/Plugins/AI/MassAI/Source/MassAIBehavior/Private/Tasks/MassStateTreeTestTask.cpp
+++ b/Engine/Plugins/AI/MassAI/Source/MassAIBehavior/Private/Tasks/MassStateTreeTestTask.cpp
@@ -17,6 +17,15 @@ FMassStateTreeTestTask::~FMassStateTreeTestTask()
 {
 }
 
+bool FMassStateTreeTestTask::Link(FStateTreeLinker& Linker)
+{
+	Linker.LinkExternalItem(MassStateTreeSubSystemHandle);
+	Linker.LinkExternalItem(SmartObjectUserHandle);
+	Linker.LinkExternalItem(TransformHandle);
+
+	return true;
+}
+
 EStateTreeRunStatus FMassStateTreeTestTask::EnterState(FStateTreeExecutionContext& Context, const EStateTreeStateChangeType ChangeType, const FStateTreeTransitionResult& Transition)
 {
 	Time = 0.0f;
@@ -25,9 +34,9 @@ EStateTreeRunStatus FMassStateTreeTestTask::EnterState(FStateTreeExecutionContex
 
 EStateTreeRunStatus FMassStateTreeTestTask::Tick(FStateTreeExecutionContext& Context, const float DeltaTime)
 {
-	UMassStateTreeSubsystem& MassStateTreeSubSystem = Context.GetExternalItem(MassStateTreeSubSystemHandle).GetMutable<UMassStateTreeSubsystem>();
-	const FDataFragment_SmartObjectUser* SmartObjectUser = Context.GetExternalItem(SmartObjectUserHandle).GetPtr<FDataFragment_SmartObjectUser>();
-	FDataFragment_Transform& Transform = Context.GetExternalItem(TransformHandle).GetMutable<FDataFragment_Transform>();
+	UMassStateTreeSubsystem& MassStateTreeSubSystem = Context.GetExternalItem(MassStateTreeSubSystemHandle);
+	const FDataFragment_SmartObjectUser* SmartObjectUser = Context.GetExternalItemPtr(SmartObjectUserHandle);
+	FDataFragment_Transform& Transform = Context.GetExternalItem(TransformHandle);
 
 	GEngine->AddOnScreenDebugMessage(INDEX_NONE, 1.0f, FColor::Orange, *FString::Printf(TEXT("[%s] Time=%f X=%f\n"), *Name.ToString(), Time, Transform.GetTransform().GetLocation().X));
 
diff --git a/Engine/Plugins/AI/MassAI/Source/MassAIBehavior/Private/Tasks/MassZoneGraphFindEscapeTarget.cpp b/Engine/Plugins/AI/MassAI/Source/MassAIBehavior/Private/Tasks/MassZoneGraphFindEscapeTarget.cpp
--- a/Engine/Plugins/AI/MassAI/Source/MassAIBehavior/Private/Tasks/MassZoneGraphFindEscapeTarget.cpp
+++ b/Engine/Plugins/AI/MassAI/Source/MassAIBehavior/Private/Tasks/MassZoneGraphFindEscapeTarget.cpp
@@ -15,6 +15,15 @@ FMassZoneGraphFindEscapeTarget::FMassZoneGraphFindEscapeTarget()
 {
 }
 
+bool FMassZoneGraphFindEscapeTarget::Link(FStateTreeLinker& Linker)
+{
+	Linker.LinkExternalItem(LocationHandle);
+	Linker.LinkExternalItem(ZoneGraphSubsystemHandle);
+	Linker.LinkExternalItem(ZoneGraphAnnotationSubsystemHandle);
+
+	return true;
+}
+
 EStateTreeRunStatus FMassZoneGraphFindEscapeTarget::EnterState(FStateTreeExecutionContext& Context, const EStateTreeStateChangeType ChangeType, const FStateTreeTransitionResult& Transition)
 {
 	FMassStateTreeExecutionContext& MassContext = static_cast<FMassStateTreeExecutionContext&>(Context);
@@ -26,9 +35,9 @@ EStateTreeRunStatus FMassZoneGraphFindEscapeTarget::EnterState(FStateTreeExecuti
 	bDisplayDebug = UE::Mass::Debug::IsDebuggingEntity(MassContext.GetEntity());
 #endif // WITH_MASSGAMEPLAY_DEBUG
 
-	const FMassZoneGraphLaneLocationFragment& LaneLocation = Context.GetExternalItem(LocationHandle).Get<FMassZoneGraphLaneLocationFragment>();
-	UZoneGraphSubsystem& ZoneGraphSubsystem = Context.GetExternalItem(ZoneGraphSubsystemHandle).GetMutable<UZoneGraphSubsystem>();
-	UZoneGraphAnnotationSubsystem& ZoneGraphAnnotationSubsystem = Context.GetExternalItem(ZoneGraphAnnotationSubsystemHandle).GetMutable<UZoneGraphAnnotationSubsystem>();
+	const FMassZoneGraphLaneLocationFragment& LaneLocation = Context.GetExternalItem(LocationHandle);
+	UZoneGraphSubsystem& ZoneGraphSubsystem = Context.GetExternalItem(ZoneGraphSubsystemHandle);
+	UZoneGraphAnnotationSubsystem& ZoneGraphAnnotationSubsystem = Context.GetExternalItem(ZoneGraphAnnotationSubsystemHandle);
 
 	if (!LaneLocation.LaneHandle.IsValid())
 	{
diff --git a/Engine/Plugins/AI/MassAI/Source/MassAIBehavior/Private/Tasks/MassZoneGraphFindSmartObjectTarget.cpp b/Engine/Plugins/AI/MassAI/Source/MassAIBehavior/Private/Tasks/MassZoneGraphFindSmartObjectTarget.cpp
--- a/Engine/Plugins/AI/MassAI/Source/MassAIBehavior/Private/Tasks/MassZoneGraphFindSmartObjectTarget.cpp
+++ b/Engine/Plugins/AI/MassAI/Source/MassAIBehavior/Private/Tasks/MassZoneGraphFindSmartObjectTarget.cpp
@@ -9,11 +9,20 @@
 #include "MassZoneGraphMovementFragments.h"
 #include "SmartObjectZoneAnnotations.h"
 
+bool FMassZoneGraphFindSmartObjectTarget::Link(FStateTreeLinker& Linker)
+{
+	Linker.LinkExternalItem(SmartObjectUserHandle);
+	Linker.LinkExternalItem(LocationHandle);
+	Linker.LinkExternalItem(AnnotationSubsystemHandle);
+
+	return true;
+}
+
 EStateTreeRunStatus FMassZoneGraphFindSmartObjectTarget::EnterState(FStateTreeExecutionContext& Context, const EStateTreeStateChangeType ChangeType, const FStateTreeTransitionResult& Transition)
 {
-	const UZoneGraphAnnotationSubsystem& AnnotationSubsystem = Context.GetExternalItem(AnnotationSubsystemHandle).Get<UZoneGraphAnnotationSubsystem>();
-	const FDataFragment_SmartObjectUser& SOUser = Context.GetExternalItem(SmartObjectUserHandle).Get<FDataFragment_SmartObjectUser>();
-	const FMassZoneGraphLaneLocationFragment& LaneLocation = Context.GetExternalItem(LocationHandle).Get<FMassZoneGraphLaneLocationFragment>();
+	const UZoneGraphAnnotationSubsystem& AnnotationSubsystem = Context.GetExternalItem(AnnotationSubsystemHandle);
+	const FDataFragment_SmartObjectUser& SOUser = Context.GetExternalItem(SmartObjectUserHandle);
+	const FMassZoneGraphLaneLocationFragment& LaneLocation = Context.GetExternalItem(LocationHandle);
 	const FZoneGraphLaneHandle LaneHandle(LaneLocation.LaneHandle);
 
 	TargetLocationRef = nullptr;
diff --git a/Engine/Plugins/AI/MassAI/Source/MassAIBehavior/Private/Tasks/MassZoneGraphPathFollowTask.cpp b/Engine/Plugins/AI/MassAI/Source/MassAIBehavior/Private/Tasks/MassZoneGraphPathFollowTask.cpp
--- a/Engine/Plugins/AI/MassAI/Source/MassAIBehavior/Private/Tasks/MassZoneGraphPathFollowTask.cpp
+++ b/Engine/Plugins/AI/MassAI/Source/MassAIBehavior/Private/Tasks/MassZoneGraphPathFollowTask.cpp
@@ -11,16 +11,30 @@
 #include "MassZoneGraphMovementUtils.h"
 #include "ZoneGraphSubsystem.h"
 
+bool FMassZoneGraphPathFollowTask::Link(FStateTreeLinker& Linker)
+{
+	Linker.LinkExternalItem(LocationHandle);
+	Linker.LinkExternalItem(MoveTargetHandle);
+	Linker.LinkExternalItem(MovementConfigHandle);
+	Linker.LinkExternalItem(PathRequestHandle);
+	Linker.LinkExternalItem(ShortPathHandle);
+	Linker.LinkExternalItem(CachedLaneHandle);
+	Linker.LinkExternalItem(AgentRadiusHandle);
+	Linker.LinkExternalItem(ZoneGraphSubsystemHandle);
+
+	return true;
+}
+
 bool FMassZoneGraphPathFollowTask::RequestPath(FMassStateTreeExecutionContext& Context, const FMassZoneGraphTargetLocation& RequestedTargetLocation) const
 {
-	const UZoneGraphSubsystem& ZoneGraphSubsystem = Context.GetExternalItem(ZoneGraphSubsystemHandle).Get<UZoneGraphSubsystem>();
-	const FMassZoneGraphLaneLocationFragment& LaneLocation = Context.GetExternalItem(LocationHandle).Get<FMassZoneGraphLaneLocationFragment>();
-	const FMassMovementConfigFragment& MovementConfig = Context.GetExternalItem(MovementConfigHandle).Get<FMassMovementConfigFragment>();
-	const FDataFragment_AgentRadius& AgentRadius = Context.GetExternalItem(AgentRadiusHandle).Get<FDataFragment_AgentRadius>();
-	FMassZoneGraphShortPathFragment& ShortPath = Context.GetExternalItem(ShortPathHandle).GetMutable<FMassZoneGraphShortPathFragment>();
-	FMassZoneGraphCachedLaneFragment& CachedLane = Context.GetExternalItem(CachedLaneHandle).GetMutable<FMassZoneGraphCachedLaneFragment>();
-	FMassMoveTargetFragment& MoveTarget = Context.GetExternalItem(MoveTargetHandle).GetMutable<FMassMoveTargetFragment>();
-	FMassZoneGraphPathRequestFragment& RequestFragment = Context.GetExternalItem(PathRequestHandle).GetMutable<FMassZoneGraphPathRequestFragment>();
+	const UZoneGraphSubsystem& ZoneGraphSubsystem = Context.GetExternalItem(ZoneGraphSubsystemHandle);
+	const FMassZoneGraphLaneLocationFragment& LaneLocation = Context.GetExternalItem(LocationHandle);
+	const FMassMovementConfigFragment& MovementConfig = Context.GetExternalItem(MovementConfigHandle);
+	const FDataFragment_AgentRadius& AgentRadius = Context.GetExternalItem(AgentRadiusHandle);
+	FMassZoneGraphShortPathFragment& ShortPath = Context.GetExternalItem(ShortPathHandle);
+	FMassZoneGraphCachedLaneFragment& CachedLane = Context.GetExternalItem(CachedLaneHandle);
+	FMassMoveTargetFragment& MoveTarget = Context.GetExternalItem(MoveTargetHandle);
+	FMassZoneGraphPathRequestFragment& RequestFragment = Context.GetExternalItem(PathRequestHandle);
 
 	const UMassMovementSettings* Settings = GetDefault<UMassMovementSettings>();
 	check(Settings);
@@ -101,7 +115,7 @@ EStateTreeRunStatus FMassZoneGraphPathFollowTask::EnterState(FStateTreeExecution
 		return EStateTreeRunStatus::Failed;
 	}
 
-	const FMassZoneGraphLaneLocationFragment& LaneLocation = Context.GetExternalItem(LocationHandle).Get<FMassZoneGraphLaneLocationFragment>();
+	const FMassZoneGraphLaneLocationFragment& LaneLocation = Context.GetExternalItem(LocationHandle);
 
 	const FMassZoneGraphTargetLocation* TargetLocationPtr = TargetLocation.GetPtr<FMassZoneGraphTargetLocation>();
 	if (!TargetLocationPtr)
@@ -137,8 +151,8 @@ EStateTreeRunStatus FMassZoneGraphPathFollowTask::Tick(FStateTreeExecutionContex
 		MASSBEHAVIOR_LOG(Verbose, TEXT("tick"));
 	}
 
-	const FMassZoneGraphShortPathFragment& ShortPath = Context.GetExternalItem(ShortPathHandle).Get<FMassZoneGraphShortPathFragment>();
-	const FMassZoneGraphLaneLocationFragment& LaneLocation = Context.GetExternalItem(LocationHandle).Get<FMassZoneGraphLaneLocationFragment>();
+	const FMassZoneGraphShortPathFragment& ShortPath = Context.GetExternalItem(ShortPathHandle);
+	const FMassZoneGraphLaneLocationFragment& LaneLocation = Context.GetExternalItem(LocationHandle);
 
 	const FMassZoneGraphTargetLocation* TargetLocationPtr = TargetLocation.GetPtr<FMassZoneGraphTargetLocation>();
 	if (!TargetLocationPtr)
diff --git a/Engine/Plugins/AI/MassAI/Source/MassAIBehavior/Private/Tasks/MassZoneGraphStandTask.cpp b/Engine/Plugins/AI/MassAI/Source/MassAIBehavior/Private/Tasks/MassZoneGraphStandTask.cpp
--- a/Engine/Plugins/AI/MassAI/Source/MassAIBehavior/Private/Tasks/MassZoneGraphStandTask.cpp
+++ b/Engine/Plugins/AI/MassAI/Source/MassAIBehavior/Private/Tasks/MassZoneGraphStandTask.cpp
@@ -12,6 +12,20 @@
 #include "MassSimulationLOD.h"
 #include "MassZoneGraphMovementUtils.h"
 
+bool FMassZoneGraphStandTask::Link(FStateTreeLinker& Linker)
+{
+	Linker.LinkExternalItem(LocationHandle);
+	Linker.LinkExternalItem(MoveTargetHandle);
+	Linker.LinkExternalItem(ShortPathHandle);
+	Linker.LinkExternalItem(CachedLaneHandle);
+	Linker.LinkExternalItem(ZoneGraphSubsystemHandle);
+	Linker.LinkExternalItem(MassSignalSubsystemHandle);
+	Linker.LinkExternalItem(SimulationLODHandle);
+	Linker.LinkExternalItem(MovementConfigHandle);
+
+	return true;
+}
+
 EStateTreeRunStatus FMassZoneGraphStandTask::EnterState(FStateTreeExecutionContext& Context, const EStateTreeStateChangeType ChangeType, const FStateTreeTransitionResult& Transition)
 {
 	// Do not reset of the state if current state is still active after transition, unless transitioned specifically to this state.
@@ -22,9 +36,9 @@ EStateTreeRunStatus FMassZoneGraphStandTask::EnterState(FStateTreeExecutionConte
 
 	const FMassStateTreeExecutionContext& MassContext = static_cast<FMassStateTreeExecutionContext&>(Context);
 
-	const FMassZoneGraphLaneLocationFragment& LaneLocation = Context.GetExternalItem(LocationHandle).Get<FMassZoneGraphLaneLocationFragment>();
-	const UZoneGraphSubsystem& ZoneGraphSubsystem = Context.GetExternalItem(ZoneGraphSubsystemHandle).Get<UZoneGraphSubsystem>();
-	const FMassMovementConfigFragment& MovementConfig = Context.GetExternalItem(MovementConfigHandle).Get<FMassMovementConfigFragment>();
+	const FMassZoneGraphLaneLocationFragment& LaneLocation = Context.GetExternalItem(LocationHandle);
+	const UZoneGraphSubsystem& ZoneGraphSubsystem = Context.GetExternalItem(ZoneGraphSubsystemHandle);
+	const FMassMovementConfigFragment& MovementConfig = Context.GetExternalItem(MovementConfigHandle);
 
 	const UMassMovementSettings* Settings = GetDefault<UMassMovementSettings>();
 	check(Settings);
@@ -42,9 +56,9 @@ EStateTreeRunStatus FMassZoneGraphStandTask::EnterState(FStateTreeExecutionConte
 		return EStateTreeRunStatus::Failed;
 	}
 
-	FMassZoneGraphShortPathFragment& ShortPath = Context.GetExternalItem(ShortPathHandle).GetMutable<FMassZoneGraphShortPathFragment>();
-	FMassZoneGraphCachedLaneFragment& CachedLane = Context.GetExternalItem(CachedLaneHandle).GetMutable<FMassZoneGraphCachedLaneFragment>();
-	FMassMoveTargetFragment& MoveTarget = Context.GetExternalItem(MoveTargetHandle).GetMutable<FMassMoveTargetFragment>();
+	FMassZoneGraphShortPathFragment& ShortPath = Context.GetExternalItem(ShortPathHandle);
+	FMassZoneGraphCachedLaneFragment& CachedLane = Context.GetExternalItem(CachedLaneHandle);
+	FMassMoveTargetFragment& MoveTarget = Context.GetExternalItem(MoveTargetHandle);
 
 	// TODO: This could be smarter too, like having a stand location/direction, or even make a small path to stop, if we're currently running.
 
@@ -64,7 +78,7 @@ EStateTreeRunStatus FMassZoneGraphStandTask::EnterState(FStateTreeExecutionConte
 	// Otherwise we schedule a signal to end the task.
 	if (Duration > 0.0f)
 	{
-		UMassSignalSubsystem& MassSignalSubsystem = Context.GetExternalItem(MassSignalSubsystemHandle).GetMutable<UMassSignalSubsystem>();
+		UMassSignalSubsystem& MassSignalSubsystem = Context.GetExternalItem(MassSignalSubsystemHandle);
 		MassSignalSubsystem.DelaySignalEntity(UE::Mass::Signals::StandTaskFinished, MassContext.GetEntity(), Duration);
 	}
 
diff --git a/Engine/Plugins/AI/MassAI/Source/MassAIBehavior/Public/Evaluators/MassComponentHitEvaluator.h b/Engine/Plugins/AI/MassAI/Source/MassAIBehavior/Public/Evaluators/MassComponentHitEvaluator.h
--- a/Engine/Plugins/AI/MassAI/Source/MassAIBehavior/Public/Evaluators/MassComponentHitEvaluator.h
+++ b/Engine/Plugins/AI/MassAI/Source/MassAIBehavior/Public/Evaluators/MassComponentHitEvaluator.h
@@ -6,6 +6,8 @@
 #include "MassStateTreeTypes.h"
 #include "MassComponentHitEvaluator.generated.h"
 
+class UMassComponentHitSubsystem;
+
 /**
  * Evaluator to extract last hit from the MassComponentHitSubsystem and expose it for tasks and transitions
  */
@@ -15,10 +17,10 @@ struct MASSAIBEHAVIOR_API FMassComponentHitEvaluator : public FMassStateTreeEval
 	GENERATED_BODY()
 
 protected:
+	virtual bool Link(FStateTreeLinker& Linker) override;
 	virtual void Evaluate(FStateTreeExecutionContext& Context, const EStateTreeEvaluationType EvalType, const float DeltaTime) override;
 
-	UPROPERTY(meta=(BaseClass="MassComponentHitSubsystem"))
-	FStateTreeExternalItemHandle ComponentHitSubsystemHandle;
+	TStateTreeItemHandle<UMassComponentHitSubsystem> ComponentHitSubsystemHandle;
 
 	UPROPERTY(VisibleAnywhere, Category = Output)
 	bool bGotHit = false;
diff --git a/Engine/Plugins/AI/MassAI/Source/MassAIBehavior/Public/Evaluators/MassStateTreeSmartObjectEvaluator.h b/Engine/Plugins/AI/MassAI/Source/MassAIBehavior/Public/Evaluators/MassStateTreeSmartObjectEvaluator.h
--- a/Engine/Plugins/AI/MassAI/Source/MassAIBehavior/Public/Evaluators/MassStateTreeSmartObjectEvaluator.h
+++ b/Engine/Plugins/AI/MassAI/Source/MassAIBehavior/Public/Evaluators/MassStateTreeSmartObjectEvaluator.h
@@ -8,6 +8,8 @@
 #include "MassStateTreeSmartObjectEvaluator.generated.h"
 
 struct FStateTreeExecutionContext;
+struct FDataFragment_Transform;
+struct FMassZoneGraphLaneLocationFragment;
 
 /**
  * Evaluator that will keep track if there is one or more smart object(s) that can be used.
@@ -18,25 +20,17 @@ struct MASSAIBEHAVIOR_API FMassStateTreeSmartObjectEvaluator : public FMassState
 	GENERATED_BODY()
 
 protected:
+	virtual bool Link(FStateTreeLinker& Linker) override;
 	virtual void EnterState(FStateTreeExecutionContext& Context, const EStateTreeStateChangeType ChangeType, const FStateTreeTransitionResult& Transition) override;
 	virtual void ExitState(FStateTreeExecutionContext& Context, const EStateTreeStateChangeType ChangeType, const FStateTreeTransitionResult& Transition) override;
 	virtual void Evaluate(FStateTreeExecutionContext& Context, const EStateTreeEvaluationType EvalType, const float DeltaTime) override;
 	void Reset();
 
-	UPROPERTY(meta=(BaseClass="SmartObjectSubsystem"))
-	FStateTreeExternalItemHandle SmartObjectSubsystemHandle;
-
-	UPROPERTY(meta=(BaseClass="MassSignalSubsystem"))
-	FStateTreeExternalItemHandle MassSignalSubsystemHandle;
-
-	UPROPERTY(meta=(BaseStruct="DataFragment_Transform"))
-	FStateTreeExternalItemHandle EntityTransformHandle;
-
-	UPROPERTY(meta=(BaseStruct="DataFragment_SmartObjectUser"))
-	FStateTreeExternalItemHandle SmartObjectUserHandle;
-
-	UPROPERTY(meta=(BaseStruct="MassZoneGraphLaneLocationFragment", Optional))
-	FStateTreeExternalItemHandle LocationHandle;
+	TStateTreeItemHandle<USmartObjectSubsystem> SmartObjectSubsystemHandle;
+	TStateTreeItemHandle<UMassSignalSubsystem> MassSignalSubsystemHandle;
+	TStateTreeItemHandle<FDataFragment_Transform> EntityTransformHandle;
+	TStateTreeItemHandle<FDataFragment_SmartObjectUser> SmartObjectUserHandle;
+	TStateTreeItemHandle<FMassZoneGraphLaneLocationFragment, EStateTreeItemRequirement::Optional> LocationHandle;
 
 	/** The identifier of the search request send by the evaluator to find candidates */
 	UPROPERTY(VisibleAnywhere, Category = SmartObject)
diff --git a/Engine/Plugins/AI/MassAI/Source/MassAIBehavior/Public/Evaluators/MassZoneGraphAnnotationEvaluator.h b/Engine/Plugins/AI/MassAI/Source/MassAIBehavior/Public/Evaluators/MassZoneGraphAnnotationEvaluator.h
--- a/Engine/Plugins/AI/MassAI/Source/MassAIBehavior/Public/Evaluators/MassZoneGraphAnnotationEvaluator.h
+++ b/Engine/Plugins/AI/MassAI/Source/MassAIBehavior/Public/Evaluators/MassZoneGraphAnnotationEvaluator.h
@@ -7,7 +7,7 @@
 #include "MassZoneGraphAnnotationEvaluator.generated.h"
 
 struct FStateTreeExecutionContext;
-
+struct FMassZoneGraphAnnotationTagsFragment;
 /**
  * Evaluator to expose ZoneGraph Annotation Tags for decision making.
  */
@@ -19,10 +19,10 @@ struct MASSAIBEHAVIOR_API FMassZoneGraphAnnotationEvaluator : public FMassStateT
 	FMassZoneGraphAnnotationEvaluator();
 
 protected:
+	virtual bool Link(FStateTreeLinker& Linker) override;
 	virtual void Evaluate(FStateTreeExecutionContext& Context, const EStateTreeEvaluationType EvalType, const float DeltaTime) override;
 
-	UPROPERTY(meta=(BaseStruct="MassZoneGraphAnnotationTagsFragment"))
-	FStateTreeExternalItemHandle BehaviorTagsHandle;
+	TStateTreeItemHandle<FMassZoneGraphAnnotationTagsFragment> BehaviorTagsHandle;
 
 	UPROPERTY(EditAnywhere, Category = Parameters)
 	FZoneGraphTagMask BehaviorTags = FZoneGraphTagMask::None;
diff --git a/Engine/Plugins/AI/MassAI/Source/MassAIBehavior/Public/Tasks/MassLookAtTask.h b/Engine/Plugins/AI/MassAI/Source/MassAIBehavior/Public/Tasks/MassLookAtTask.h
--- a/Engine/Plugins/AI/MassAI/Source/MassAIBehavior/Public/Tasks/MassLookAtTask.h
+++ b/Engine/Plugins/AI/MassAI/Source/MassAIBehavior/Public/Tasks/MassLookAtTask.h
@@ -7,6 +7,8 @@
 #include "MassLookAtFragments.h"
 #include "MassLookAtTask.generated.h"
 
+class UMassSignalSubsystem;
+
 /**
  * Task to assign a LookAt target for mass processing
  */
@@ -16,15 +18,13 @@ struct MASSAIBEHAVIOR_API FMassLookAtTask : public FMassStateTreeTaskBase
 	GENERATED_BODY()
 
 protected:
+	virtual bool Link(FStateTreeLinker& Linker) override;
 	virtual EStateTreeRunStatus EnterState(FStateTreeExecutionContext& Context, const EStateTreeStateChangeType ChangeType, const FStateTreeTransitionResult& Transition) override;
 	virtual void ExitState(FStateTreeExecutionContext& Context, const EStateTreeStateChangeType ChangeType, const FStateTreeTransitionResult& Transition) override;
 	virtual EStateTreeRunStatus Tick(FStateTreeExecutionContext& Context, const float DeltaTime) override;
 
-	UPROPERTY(meta = (BaseClass="MassSignalSubsystem"))
-	FStateTreeExternalItemHandle MassSignalSubsystemHandle;
-
-	UPROPERTY(meta = (BaseStruct="MassLookAtFragment"))
-	FStateTreeExternalItemHandle LookAtHandle;
+	TStateTreeItemHandle<UMassSignalSubsystem> MassSignalSubsystemHandle;
+	TStateTreeItemHandle<FMassLookAtFragment> LookAtHandle;
 
 	/** Delay before the task ends. Default (0 or any negative) will run indefinitely so it requires a transition in the state tree to stop it. */
 	UPROPERTY(EditAnywhere, Category = Parameters, meta = (Bindable))
diff --git a/Engine/Plugins/AI/MassAI/Source/MassAIBehavior/Public/Tasks/MassStateTreeSmartObjectTask.h b/Engine/Plugins/AI/MassAI/Source/MassAIBehavior/Public/Tasks/MassStateTreeSmartObjectTask.h
--- a/Engine/Plugins/AI/MassAI/Source/MassAIBehavior/Public/Tasks/MassStateTreeSmartObjectTask.h
+++ b/Engine/Plugins/AI/MassAI/Source/MassAIBehavior/Public/Tasks/MassStateTreeSmartObjectTask.h
@@ -9,6 +9,10 @@
 #include "MassStateTreeSmartObjectTask.generated.h"
 
 struct FStateTreeExecutionContext;
+struct FDataFragment_SmartObjectUser;
+class USmartObjectSubsystem;
+struct FDataFragment_Transform;
+struct FMassMoveTargetFragment;
 
 /**
  * Tasks to claim a smart object from search results and release it when done.
@@ -19,15 +23,13 @@ struct MASSAIBEHAVIOR_API FMassClaimSmartObjectTask : public FMassStateTreeTaskB
 	GENERATED_BODY()
 
 protected:
+	virtual bool Link(FStateTreeLinker& Linker) override;
 	virtual EStateTreeRunStatus EnterState(FStateTreeExecutionContext& Context, const EStateTreeStateChangeType ChangeType, const FStateTreeTransitionResult& Transition) override;
 	virtual void ExitState(FStateTreeExecutionContext& Context, const EStateTreeStateChangeType ChangeType, const FStateTreeTransitionResult& Transition) override;
 	virtual EStateTreeRunStatus Tick(FStateTreeExecutionContext& Context, const float DeltaTime) override;
 
-	UPROPERTY(meta=(BaseStruct="DataFragment_SmartObjectUser"))
-	FStateTreeExternalItemHandle SmartObjectUserHandle;
-
-	UPROPERTY(meta=(BaseClass="SmartObjectSubsystem"))
-	FStateTreeExternalItemHandle SmartObjectSubsystemHandle;
+	TStateTreeItemHandle<FDataFragment_SmartObjectUser> SmartObjectUserHandle;
+	TStateTreeItemHandle<USmartObjectSubsystem> SmartObjectSubsystemHandle;
 
 	/** Result of the candidates search request (Input) */
 	UPROPERTY(VisibleAnywhere, Category = SmartObject, meta = (Bindable), Transient)
@@ -47,22 +49,16 @@ struct MASSAIBEHAVIOR_API FMassUseSmartObjectTask : public FMassStateTreeTaskBas
 	GENERATED_BODY()
 
 protected:
+	virtual bool Link(FStateTreeLinker& Linker) override;
 	virtual EStateTreeRunStatus EnterState(FStateTreeExecutionContext& Context, const EStateTreeStateChangeType ChangeType, const FStateTreeTransitionResult& Transition) override;
 	virtual void ExitState(FStateTreeExecutionContext& Context, const EStateTreeStateChangeType ChangeType, const FStateTreeTransitionResult& Transition) override;
 	virtual void StateCompleted(FStateTreeExecutionContext& Context, const EStateTreeRunStatus CompletionStatus, const FStateTreeHandle CompletedState) override;
 	virtual EStateTreeRunStatus Tick(FStateTreeExecutionContext& Context, const float DeltaTime) override;
 
-	UPROPERTY(meta=(BaseClass="SmartObjectSubsystem"))
-	FStateTreeExternalItemHandle SmartObjectSubsystemHandle;
-
-	UPROPERTY(meta=(BaseStruct="DataFragment_Transform"))
-	FStateTreeExternalItemHandle EntityTransformHandle;
-
-	UPROPERTY(meta=(BaseStruct="DataFragment_SmartObjectUser"))
-	FStateTreeExternalItemHandle SmartObjectUserHandle;
-
-	UPROPERTY(meta=(BaseStruct="MassMoveTargetFragment"))
-	FStateTreeExternalItemHandle MoveTargetHandle;
+	TStateTreeItemHandle<USmartObjectSubsystem> SmartObjectSubsystemHandle;
+	TStateTreeItemHandle<FDataFragment_Transform> EntityTransformHandle;
+	TStateTreeItemHandle<FDataFragment_SmartObjectUser> SmartObjectUserHandle;
+	TStateTreeItemHandle<FMassMoveTargetFragment> MoveTargetHandle;
 
 	/** Delay in seconds before trying to find & use another smart object */
 	UPROPERTY(EditAnywhere, Category = SmartObject)
diff --git a/Engine/Plugins/AI/MassAI/Source/MassAIBehavior/Public/Tasks/MassStateTreeTestTask.h b/Engine/Plugins/AI/MassAI/Source/MassAIBehavior/Public/Tasks/MassStateTreeTestTask.h
--- a/Engine/Plugins/AI/MassAI/Source/MassAIBehavior/Public/Tasks/MassStateTreeTestTask.h
+++ b/Engine/Plugins/AI/MassAI/Source/MassAIBehavior/Public/Tasks/MassStateTreeTestTask.h
@@ -6,6 +6,10 @@
 #include "MassStateTreeTypes.h"
 #include "MassStateTreeTestTask.generated.h"
 
+class UMassStateTreeSubsystem;
+struct FDataFragment_SmartObjectUser;
+struct FDataFragment_Transform;
+
 /**
  * Test Task, will be removed later.
  */
@@ -16,22 +20,20 @@ struct MASSAIBEHAVIOR_API FMassStateTreeTestTask : public FMassStateTreeTaskBase
 
 	FMassStateTreeTestTask();
 	virtual ~FMassStateTreeTestTask();
+
+protected:
+	virtual bool Link(FStateTreeLinker& Linker) override;
 	virtual EStateTreeRunStatus EnterState(FStateTreeExecutionContext& Context, const EStateTreeStateChangeType ChangeType, const FStateTreeTransitionResult& Transition) override;
 	virtual EStateTreeRunStatus Tick(FStateTreeExecutionContext& Context, const float DeltaTime) override;
 
-protected:
+	TStateTreeItemHandle<UMassStateTreeSubsystem> MassStateTreeSubSystemHandle;
+	TStateTreeItemHandle<FDataFragment_SmartObjectUser> SmartObjectUserHandle;
+	TStateTreeItemHandle<FDataFragment_Transform> TransformHandle;
+
 	UPROPERTY()
 	float Time = 0.0f;
 
 	UPROPERTY(EditAnywhere, Category = Test, meta = (Bindable))
 	float Duration = 5.0f;
 
-	UPROPERTY(meta=(BaseClass="MassStateTreeSubsystem"))
-	FStateTreeExternalItemHandle MassStateTreeSubSystemHandle;
-
-	UPROPERTY(meta=(BaseStruct="DataFragment_SmartObjectUser", Optional))
-	FStateTreeExternalItemHandle SmartObjectUserHandle;
-
-	UPROPERTY(meta=(BaseStruct="DataFragment_Transform"))
-	FStateTreeExternalItemHandle TransformHandle;
 };
diff --git a/Engine/Plugins/AI/MassAI/Source/MassAIBehavior/Public/Tasks/MassZoneGraphFindEscapeTarget.h b/Engine/Plugins/AI/MassAI/Source/MassAIBehavior/Public/Tasks/MassZoneGraphFindEscapeTarget.h
--- a/Engine/Plugins/AI/MassAI/Source/MassAIBehavior/Public/Tasks/MassZoneGraphFindEscapeTarget.h
+++ b/Engine/Plugins/AI/MassAI/Source/MassAIBehavior/Public/Tasks/MassZoneGraphFindEscapeTarget.h
@@ -7,6 +7,9 @@
 #include "MassZoneGraphFindEscapeTarget.generated.h"
 
 struct FStateTreeExecutionContext;
+struct FMassZoneGraphLaneLocationFragment;
+class UZoneGraphSubsystem;
+class UZoneGraphAnnotationSubsystem;
 
 /**
  * Updates TargetLocation to a escape target based on the agents current location on ZoneGraph, and disturbance annotation.
@@ -19,21 +22,14 @@ struct MASSAIBEHAVIOR_API FMassZoneGraphFindEscapeTarget : public FMassStateTree
 	FMassZoneGraphFindEscapeTarget();
 
 protected:
+	virtual bool Link(FStateTreeLinker& Linker) override;
 	virtual EStateTreeRunStatus EnterState(FStateTreeExecutionContext& Context, const EStateTreeStateChangeType ChangeType, const FStateTreeTransitionResult& Transition) override;
 	virtual void ExitState(FStateTreeExecutionContext& Context, const EStateTreeStateChangeType ChangeType, const FStateTreeTransitionResult& Transition) override;
 	virtual EStateTreeRunStatus Tick(FStateTreeExecutionContext& Context, const float DeltaTime) override;
 
-	UPROPERTY(meta=(BaseStruct="MassZoneGraphLaneLocationFragment"))
-	FStateTreeExternalItemHandle LocationHandle;
-
-	UPROPERTY(meta=(BaseStruct="MassMoveTargetFragment"))
-	FStateTreeExternalItemHandle MoveTargetHandle;
-
-	UPROPERTY(meta=(BaseClass="ZoneGraphSubsystem"))
-	FStateTreeExternalItemHandle ZoneGraphSubsystemHandle;
-
-	UPROPERTY(meta=(BaseClass="ZoneGraphAnnotationSubsystem"))
-	FStateTreeExternalItemHandle ZoneGraphAnnotationSubsystemHandle;
+	TStateTreeItemHandle<FMassZoneGraphLaneLocationFragment> LocationHandle;
+	TStateTreeItemHandle<UZoneGraphSubsystem> ZoneGraphSubsystemHandle;
+	TStateTreeItemHandle<UZoneGraphAnnotationSubsystem> ZoneGraphAnnotationSubsystemHandle;
 
 	UPROPERTY(EditAnywhere, Category = Parameters)
 	FZoneGraphTag DisturbanceAnnotationTag = FZoneGraphTag::None;
diff --git a/Engine/Plugins/AI/MassAI/Source/MassAIBehavior/Public/Tasks/MassZoneGraphFindSmartObjectTarget.h b/Engine/Plugins/AI/MassAI/Source/MassAIBehavior/Public/Tasks/MassZoneGraphFindSmartObjectTarget.h
--- a/Engine/Plugins/AI/MassAI/Source/MassAIBehavior/Public/Tasks/MassZoneGraphFindSmartObjectTarget.h
+++ b/Engine/Plugins/AI/MassAI/Source/MassAIBehavior/Public/Tasks/MassZoneGraphFindSmartObjectTarget.h
@@ -3,9 +3,12 @@
 #pragma once
 #include "MassStateTreeTypes.h"
 #include "MassZoneGraphPathFollowTask.h"
-
 #include "MassZoneGraphFindSmartObjectTarget.generated.h"
 
+struct FDataFragment_SmartObjectUser;
+struct FMassZoneGraphLaneLocationFragment;
+class UZoneGraphAnnotationSubsystem;
+
 /**
 * Computes move target to a smart object based on current location on ZoneGraph.
 */
@@ -15,18 +18,14 @@ struct MASSAIBEHAVIOR_API FMassZoneGraphFindSmartObjectTarget : public FMassStat
 	GENERATED_BODY()
 
 protected:
+	virtual bool Link(FStateTreeLinker& Linker) override;
 	virtual EStateTreeRunStatus EnterState(FStateTreeExecutionContext& Context, const EStateTreeStateChangeType ChangeType, const FStateTreeTransitionResult& Transition) override;
 	virtual void ExitState(FStateTreeExecutionContext& Context, const EStateTreeStateChangeType ChangeType, const FStateTreeTransitionResult& Transition) override;
 	virtual EStateTreeRunStatus Tick(FStateTreeExecutionContext& Context, const float DeltaTime) override;
 
-	UPROPERTY(meta=(BaseStruct="DataFragment_SmartObjectUser"))
-	FStateTreeExternalItemHandle SmartObjectUserHandle;
-
-	UPROPERTY(meta=(BaseStruct="MassZoneGraphLaneLocationFragment"))
-	FStateTreeExternalItemHandle LocationHandle;
-
-	UPROPERTY(meta=(BaseClass="ZoneGraphAnnotationSubsystem"))
-	FStateTreeExternalItemHandle AnnotationSubsystemHandle;
+	TStateTreeItemHandle<FDataFragment_SmartObjectUser> SmartObjectUserHandle;
+	TStateTreeItemHandle<FMassZoneGraphLaneLocationFragment> LocationHandle;
+	TStateTreeItemHandle<UZoneGraphAnnotationSubsystem> AnnotationSubsystemHandle;
 
 	FMassZoneGraphTargetLocation TargetLocation;
 
diff --git a/Engine/Plugins/AI/MassAI/Source/MassAIBehavior/Public/Tasks/MassZoneGraphPathFollowTask.h b/Engine/Plugins/AI/MassAI/Source/MassAIBehavior/Public/Tasks/MassZoneGraphPathFollowTask.h
--- a/Engine/Plugins/AI/MassAI/Source/MassAIBehavior/Public/Tasks/MassZoneGraphPathFollowTask.h
+++ b/Engine/Plugins/AI/MassAI/Source/MassAIBehavior/Public/Tasks/MassZoneGraphPathFollowTask.h
@@ -8,6 +8,14 @@
 #include "MassZoneGraphPathFollowTask.generated.h"
 
 struct FMassStateTreeExecutionContext;
+struct FMassZoneGraphLaneLocationFragment;
+struct FMassMoveTargetFragment;
+struct FMassMovementConfigFragment;
+struct FMassZoneGraphPathRequestFragment;
+struct FMassZoneGraphShortPathFragment;
+struct FMassZoneGraphCachedLaneFragment;
+struct FDataFragment_AgentRadius;
+class UZoneGraphSubsystem;
 
 USTRUCT()
 struct MASSAIBEHAVIOR_API FMassZoneGraphTargetLocation : public FStateTreeResult
@@ -65,34 +73,20 @@ struct MASSAIBEHAVIOR_API FMassZoneGraphPathFollowTask : public FMassStateTreeTa
 	GENERATED_BODY()
 
 protected:
+	virtual bool Link(FStateTreeLinker& Linker) override;
 	virtual EStateTreeRunStatus EnterState(FStateTreeExecutionContext& Context, const EStateTreeStateChangeType ChangeType, const FStateTreeTransitionResult& Transition) override;
 	virtual EStateTreeRunStatus Tick(FStateTreeExecutionContext& Context, const float DeltaTime) override;
 
 	bool RequestPath(FMassStateTreeExecutionContext& Context, const FMassZoneGraphTargetLocation& TargetLocation) const;
 
-	UPROPERTY(meta=(BaseStruct="MassZoneGraphLaneLocationFragment"))
-	FStateTreeExternalItemHandle LocationHandle;
-
-	UPROPERTY(meta=(BaseStruct="MassMoveTargetFragment"))
-	FStateTreeExternalItemHandle MoveTargetHandle;
-
-	UPROPERTY(meta=(BaseStruct="MassMovementConfigFragment"))
-	FStateTreeExternalItemHandle MovementConfigHandle;
-
-	UPROPERTY(meta=(BaseStruct="MassZoneGraphPathRequestFragment"))
-	FStateTreeExternalItemHandle PathRequestHandle;
-
-	UPROPERTY(meta=(BaseStruct="MassZoneGraphShortPathFragment"))
-	FStateTreeExternalItemHandle ShortPathHandle;
-
-	UPROPERTY(meta=(BaseStruct="MassZoneGraphCachedLaneFragment"))
-	FStateTreeExternalItemHandle CachedLaneHandle;
-
-	UPROPERTY(meta=(BaseStruct="DataFragment_AgentRadius"))
-	FStateTreeExternalItemHandle AgentRadiusHandle;
-	
-	UPROPERTY(meta=(BaseClass="ZoneGraphSubsystem"))
-	FStateTreeExternalItemHandle ZoneGraphSubsystemHandle;
+	TStateTreeItemHandle<FMassZoneGraphLaneLocationFragment> LocationHandle;
+	TStateTreeItemHandle<FMassMoveTargetFragment> MoveTargetHandle;
+	TStateTreeItemHandle<FMassMovementConfigFragment> MovementConfigHandle;
+	TStateTreeItemHandle<FMassZoneGraphPathRequestFragment> PathRequestHandle;
+	TStateTreeItemHandle<FMassZoneGraphShortPathFragment> ShortPathHandle;
+	TStateTreeItemHandle<FMassZoneGraphCachedLaneFragment> CachedLaneHandle;
+	TStateTreeItemHandle<FDataFragment_AgentRadius> AgentRadiusHandle;
+	TStateTreeItemHandle<UZoneGraphSubsystem> ZoneGraphSubsystemHandle;
 
 	UPROPERTY(EditAnywhere, Category = Parameters, meta=(Struct="MassZoneGraphTargetLocation", Bindable))
 	FStateTreeResultRef TargetLocation;
diff --git a/Engine/Plugins/AI/MassAI/Source/MassAIBehavior/Public/Tasks/MassZoneGraphStandTask.h b/Engine/Plugins/AI/MassAI/Source/MassAIBehavior/Public/Tasks/MassZoneGraphStandTask.h
--- a/Engine/Plugins/AI/MassAI/Source/MassAIBehavior/Public/Tasks/MassZoneGraphStandTask.h
+++ b/Engine/Plugins/AI/MassAI/Source/MassAIBehavior/Public/Tasks/MassZoneGraphStandTask.h
@@ -7,6 +7,14 @@
 #include "MassZoneGraphStandTask.generated.h"
 
 struct FStateTreeExecutionContext;
+struct FMassZoneGraphLaneLocationFragment;
+struct FMassMoveTargetFragment;
+struct FMassZoneGraphShortPathFragment;
+struct FMassZoneGraphCachedLaneFragment;
+class UZoneGraphSubsystem;
+class UMassSignalSubsystem;
+struct FMassSimulationLODFragment;
+struct FMassMovementConfigFragment;
 
 /**
  * Stop, and stand on current ZoneGraph location
@@ -17,29 +25,18 @@ struct MASSAIBEHAVIOR_API FMassZoneGraphStandTask : public FMassStateTreeTaskBas
 	GENERATED_BODY()
 
 protected:
+	virtual bool Link(FStateTreeLinker& Linker) override;
 	virtual EStateTreeRunStatus EnterState(FStateTreeExecutionContext& Context, const EStateTreeStateChangeType ChangeType, const FStateTreeTransitionResult& Transition) override;
 	virtual EStateTreeRunStatus Tick(FStateTreeExecutionContext& Context, const float DeltaTime) override;
 
-	UPROPERTY(meta=(BaseStruct="MassZoneGraphLaneLocationFragment"))
-	FStateTreeExternalItemHandle LocationHandle;
-
-	UPROPERTY(meta=(BaseStruct="MassMoveTargetFragment"))
-	FStateTreeExternalItemHandle MoveTargetHandle;
-
-	UPROPERTY(meta=(BaseStruct="MassZoneGraphShortPathFragment"))
-	FStateTreeExternalItemHandle ShortPathHandle;
-
-	UPROPERTY(meta=(BaseStruct="MassZoneGraphCachedLaneFragment"))
-	FStateTreeExternalItemHandle CachedLaneHandle;
-
-	UPROPERTY(meta=(BaseClass="ZoneGraphSubsystem"))
-	FStateTreeExternalItemHandle ZoneGraphSubsystemHandle;
-
-	UPROPERTY(meta=(BaseClass="MassSignalSubsystem"))
-	FStateTreeExternalItemHandle MassSignalSubsystemHandle;
-
-	UPROPERTY(meta=(BaseStruct="MassMovementConfigFragment"))
-	FStateTreeExternalItemHandle MovementConfigHandle;
+	TStateTreeItemHandle<FMassZoneGraphLaneLocationFragment> LocationHandle;
+	TStateTreeItemHandle<FMassMoveTargetFragment> MoveTargetHandle;
+	TStateTreeItemHandle<FMassZoneGraphShortPathFragment> ShortPathHandle;
+	TStateTreeItemHandle<FMassZoneGraphCachedLaneFragment> CachedLaneHandle;
+	TStateTreeItemHandle<UZoneGraphSubsystem> ZoneGraphSubsystemHandle;
+	TStateTreeItemHandle<UMassSignalSubsystem> MassSignalSubsystemHandle;
+	TStateTreeItemHandle<FMassSimulationLODFragment> SimulationLODHandle;
+	TStateTreeItemHandle<FMassMovementConfigFragment> MovementConfigHandle;
 
 	UPROPERTY()
 	float Time = 0.0f;
diff --git a/Engine/Plugins/AI/MassCrowd/Source/MassCrowd/Private/Tasks/MassCrowdClaimWaitSlotTask.cpp b/Engine/Plugins/AI/MassCrowd/Source/MassCrowd/Private/Tasks/MassCrowdClaimWaitSlotTask.cpp
--- a/Engine/Plugins/AI/MassCrowd/Source/MassCrowd/Private/Tasks/MassCrowdClaimWaitSlotTask.cpp
+++ b/Engine/Plugins/AI/MassCrowd/Source/MassCrowd/Private/Tasks/MassCrowdClaimWaitSlotTask.cpp
@@ -7,6 +7,15 @@
 #include "MassAIMovementFragments.h"
 #include "MassStateTreeExecutionContext.h"
 
+bool FMassCrowdClaimWaitSlotTask::Link(FStateTreeLinker& Linker)
+{
+	Linker.LinkExternalItem(LocationHandle);
+	Linker.LinkExternalItem(MoveTargetHandle);
+	Linker.LinkExternalItem(CrowdSubsystemHandle);
+
+	return true;
+}
+
 EStateTreeRunStatus FMassCrowdClaimWaitSlotTask::EnterState(FStateTreeExecutionContext& Context, const EStateTreeStateChangeType ChangeType, const FStateTreeTransitionResult& Transition)
 {
 	if (ChangeType != EStateTreeStateChangeType::Changed)
@@ -19,9 +28,9 @@ EStateTreeRunStatus FMassCrowdClaimWaitSlotTask::EnterState(FStateTreeExecutionC
 	const FMassStateTreeExecutionContext& MassContext = static_cast<FMassStateTreeExecutionContext&>(Context);
 	const FMassEntityHandle Entity = MassContext.GetEntity();
 	
-	const FMassZoneGraphLaneLocationFragment& LaneLocation = Context.GetExternalItem(LocationHandle).Get<FMassZoneGraphLaneLocationFragment>();
-	const FMassMoveTargetFragment& MoveTarget = Context.GetExternalItem(MoveTargetHandle).Get<FMassMoveTargetFragment>();
-	UMassCrowdSubsystem& CrowdSubsystem = Context.GetExternalItem(CrowdSubsystemHandle).GetMutable<UMassCrowdSubsystem>();
+	const FMassZoneGraphLaneLocationFragment& LaneLocation = Context.GetExternalItem(LocationHandle);
+	const FMassMoveTargetFragment& MoveTarget = Context.GetExternalItem(MoveTargetHandle);
+	UMassCrowdSubsystem& CrowdSubsystem = Context.GetExternalItem(CrowdSubsystemHandle);
 
 	FVector SlotPosition = FVector::ZeroVector;
 	FVector SlotDirection = FVector::ForwardVector;
@@ -63,7 +72,7 @@ void FMassCrowdClaimWaitSlotTask::ExitState(FStateTreeExecutionContext& Context,
 	const FMassStateTreeExecutionContext& MassContext = static_cast<FMassStateTreeExecutionContext&>(Context);
 	const FMassEntityHandle Entity = MassContext.GetEntity();
 	
-	UMassCrowdSubsystem& CrowdSubsystem = Context.GetExternalItem(CrowdSubsystemHandle).GetMutable<UMassCrowdSubsystem>();
+	UMassCrowdSubsystem& CrowdSubsystem = Context.GetExternalItem(CrowdSubsystemHandle);
 
 	if (WaitingSlotIndex != INDEX_NONE)
 	{
diff --git a/Engine/Plugins/AI/MassCrowd/Source/MassCrowd/Private/Tasks/MassZoneGraphFindWanderTarget.cpp b/Engine/Plugins/AI/MassCrowd/Source/MassCrowd/Private/Tasks/MassZoneGraphFindWanderTarget.cpp
--- a/Engine/Plugins/AI/MassCrowd/Source/MassCrowd/Private/Tasks/MassZoneGraphFindWanderTarget.cpp
+++ b/Engine/Plugins/AI/MassCrowd/Source/MassCrowd/Private/Tasks/MassZoneGraphFindWanderTarget.cpp
@@ -17,14 +17,24 @@ FMassZoneGraphFindWanderTarget::FMassZoneGraphFindWanderTarget()
 {
 }
 
+bool FMassZoneGraphFindWanderTarget::Link(FStateTreeLinker& Linker)
+{
+	Linker.LinkExternalItem(LocationHandle);
+	Linker.LinkExternalItem(ZoneGraphSubsystemHandle);
+	Linker.LinkExternalItem(ZoneGraphAnnotationSubsystemHandle);
+	Linker.LinkExternalItem(MassCrowdSubsystemHandle);
+
+	return true;
+}
+
 EStateTreeRunStatus FMassZoneGraphFindWanderTarget::EnterState(FStateTreeExecutionContext& Context, const EStateTreeStateChangeType ChangeType, const FStateTreeTransitionResult& Transition)
 {
 	const FMassStateTreeExecutionContext& MassContext = static_cast<FMassStateTreeExecutionContext&>(Context);
 
-	const FMassZoneGraphLaneLocationFragment& LaneLocation = Context.GetExternalItem(LocationHandle).Get<FMassZoneGraphLaneLocationFragment>();
-	const UZoneGraphSubsystem& ZoneGraphSubsystem = Context.GetExternalItem(ZoneGraphSubsystemHandle).Get<UZoneGraphSubsystem>();
-	UZoneGraphAnnotationSubsystem& ZoneGraphAnnotationSubsystem = Context.GetExternalItem(ZoneGraphAnnotationSubsystemHandle).GetMutable<UZoneGraphAnnotationSubsystem>();
-	const UMassCrowdSubsystem& MassCrowdSubsystem = Context.GetExternalItem(MassCrowdSubsystemHandle).Get<UMassCrowdSubsystem>();
+	const FMassZoneGraphLaneLocationFragment& LaneLocation = Context.GetExternalItem(LocationHandle);
+	const UZoneGraphSubsystem& ZoneGraphSubsystem = Context.GetExternalItem(ZoneGraphSubsystemHandle);
+	UZoneGraphAnnotationSubsystem& ZoneGraphAnnotationSubsystem = Context.GetExternalItem(ZoneGraphAnnotationSubsystemHandle);
+	const UMassCrowdSubsystem& MassCrowdSubsystem = Context.GetExternalItem(MassCrowdSubsystemHandle);
 
 	bool bDisplayDebug = false;
 #if WITH_MASSGAMEPLAY_DEBUG
diff --git a/Engine/Plugins/AI/MassCrowd/Source/MassCrowd/Public/Tasks/MassCrowdClaimWaitSlotTask.h b/Engine/Plugins/AI/MassCrowd/Source/MassCrowd/Public/Tasks/MassCrowdClaimWaitSlotTask.h
--- a/Engine/Plugins/AI/MassCrowd/Source/MassCrowd/Public/Tasks/MassCrowdClaimWaitSlotTask.h
+++ b/Engine/Plugins/AI/MassCrowd/Source/MassCrowd/Public/Tasks/MassCrowdClaimWaitSlotTask.h
@@ -8,6 +8,7 @@
 #include "MassCrowdClaimWaitSlotTask.generated.h"
 
 struct FStateTreeExecutionContext;
+class UMassCrowdSubsystem;
 
 /**
 * Claim wait slot and expose slot position for path follow.
@@ -18,18 +19,14 @@ struct MASSCROWD_API FMassCrowdClaimWaitSlotTask : public FMassStateTreeTaskBase
 	GENERATED_BODY()
 
 protected:
+	virtual bool Link(FStateTreeLinker& Linker) override;
 	virtual EStateTreeRunStatus EnterState(FStateTreeExecutionContext& Context, const EStateTreeStateChangeType ChangeType, const FStateTreeTransitionResult& Transition) override;
 	virtual void ExitState(FStateTreeExecutionContext& Context, const EStateTreeStateChangeType ChangeType, const FStateTreeTransitionResult& Transition) override;
 	virtual EStateTreeRunStatus Tick(FStateTreeExecutionContext& Context, const float DeltaTime) override;
 
-	UPROPERTY(meta=(BaseStruct="MassZoneGraphLaneLocationFragment"))
-	FStateTreeExternalItemHandle LocationHandle;
-
-	UPROPERTY(meta=(BaseStruct="MassMoveTargetFragment"))
-	FStateTreeExternalItemHandle MoveTargetHandle;
-
-	UPROPERTY(meta=(BaseClass="MassCrowdSubsystem"))
-	FStateTreeExternalItemHandle CrowdSubsystemHandle;
+	TStateTreeItemHandle<FMassZoneGraphLaneLocationFragment> LocationHandle;
+	TStateTreeItemHandle<FMassMoveTargetFragment> MoveTargetHandle;
+	TStateTreeItemHandle<UMassCrowdSubsystem> CrowdSubsystemHandle;
 
 	UPROPERTY(EditAnywhere, Category = Parameters, meta=(Struct="MassZoneGraphTargetLocation"))
 	FStateTreeResultRef WaitSlotLocation;
diff --git a/Engine/Plugins/AI/MassCrowd/Source/MassCrowd/Public/Tasks/MassZoneGraphFindWanderTarget.h b/Engine/Plugins/AI/MassCrowd/Source/MassCrowd/Public/Tasks/MassZoneGraphFindWanderTarget.h
--- a/Engine/Plugins/AI/MassCrowd/Source/MassCrowd/Public/Tasks/MassZoneGraphFindWanderTarget.h
+++ b/Engine/Plugins/AI/MassCrowd/Source/MassCrowd/Public/Tasks/MassZoneGraphFindWanderTarget.h
@@ -7,6 +7,10 @@
 #include "MassZoneGraphFindWanderTarget.generated.h"
 
 struct FStateTreeExecutionContext;
+struct FMassZoneGraphLaneLocationFragment;
+class UZoneGraphSubsystem;
+class UZoneGraphAnnotationSubsystem;
+class UMassCrowdSubsystem;
 
 /**
  * Updates TargetLocation to a wander target based on the agents current location on ZoneGraph.
@@ -19,21 +23,15 @@ struct MASSCROWD_API FMassZoneGraphFindWanderTarget : public FMassStateTreeTaskB
 	FMassZoneGraphFindWanderTarget();
 
 protected:
+	virtual bool Link(FStateTreeLinker& Linker) override;
 	virtual EStateTreeRunStatus EnterState(FStateTreeExecutionContext& Context, const EStateTreeStateChangeType ChangeType, const FStateTreeTransitionResult& Transition) override;
 	virtual void ExitState(FStateTreeExecutionContext& Context, const EStateTreeStateChangeType ChangeType, const FStateTreeTransitionResult& Transition) override;
 	virtual EStateTreeRunStatus Tick(FStateTreeExecutionContext& Context, const float DeltaTime) override;
 
-	UPROPERTY(meta=(BaseStruct="MassZoneGraphLaneLocationFragment"))
-	FStateTreeExternalItemHandle LocationHandle;
-
-	UPROPERTY(meta=(BaseClass="ZoneGraphSubsystem"))
-	FStateTreeExternalItemHandle ZoneGraphSubsystemHandle;
-
-	UPROPERTY(meta=(BaseClass="ZoneGraphAnnotationSubsystem"))
-	FStateTreeExternalItemHandle ZoneGraphAnnotationSubsystemHandle;
-
-	UPROPERTY(meta=(BaseClass="MassCrowdSubsystem"))
-	FStateTreeExternalItemHandle MassCrowdSubsystemHandle;
+	TStateTreeItemHandle<FMassZoneGraphLaneLocationFragment> LocationHandle;
+	TStateTreeItemHandle<UZoneGraphSubsystem> ZoneGraphSubsystemHandle;
+	TStateTreeItemHandle<UZoneGraphAnnotationSubsystem> ZoneGraphAnnotationSubsystemHandle;
+	TStateTreeItemHandle<UMassCrowdSubsystem> MassCrowdSubsystemHandle;
 
 	UPROPERTY(EditAnywhere, Category = Parameters)
 	FZoneGraphTagFilter AllowedBehaviorTags;
diff --git a/Engine/Plugins/Runtime/StateTree/Source/StateTreeEditorModule/Private/Customizations/StateTreeTransitionDetails.cpp b/Engine/Plugins/Runtime/StateTree/Source/StateTreeEditorModule/Private/Customizations/StateTreeTransitionDetails.cpp
--- a/Engine/Plugins/Runtime/StateTree/Source/StateTreeEditorModule/Private/Customizations/StateTreeTransitionDetails.cpp
+++ b/Engine/Plugins/Runtime/StateTree/Source/StateTreeEditorModule/Private/Customizations/StateTreeTransitionDetails.cpp
@@ -6,8 +6,9 @@
 #include "DetailLayoutBuilder.h"
 #include "IDetailPropertyRow.h"
 #include "IDetailChildrenBuilder.h"
-#include "StateTree.h"
 #include "StateTreeState.h"
+#include "Widgets/Text/STextBlock.h"
+#include "Widgets/Layout/SBox.h"
 
 #define LOCTEXT_NAMESPACE "StateTreeEditor"
 
diff --git a/Engine/Plugins/Runtime/StateTree/Source/StateTreeEditorModule/Private/StateTreeBaker.cpp b/Engine/Plugins/Runtime/StateTree/Source/StateTreeEditorModule/Private/StateTreeBaker.cpp
--- a/Engine/Plugins/Runtime/StateTree/Source/StateTreeEditorModule/Private/StateTreeBaker.cpp
+++ b/Engine/Plugins/Runtime/StateTree/Source/StateTreeEditorModule/Private/StateTreeBaker.cpp
@@ -85,11 +85,13 @@ bool FStateTreeBaker::Bake(UStateTree& InStateTree)
 
 	BindingsCompiler.Finalize();
 
-	StateTree->InitRuntimeStorage();
-
 	// TODO: This is just for testing. It should be called just before use.
 	StateTree->ResolvePropertyPaths();
 
+	StateTree->Link();
+	
+	StateTree->InitRuntimeStorage();
+
 	return true;
 }
 
@@ -285,111 +287,6 @@ bool FStateTreeBaker::GetAndValidateBindings(const FStateTreeBindableStructDesc&
 	return true;
 }
 
-bool FStateTreeBaker::CreateExternalItemHandles(FStructView Item)
-{
-	check(StateTree);
-	
-	const UScriptStruct* ScriptStruct = Item.GetScriptStruct();
-	check(ScriptStruct);
-
-	static const FName BaseStructMetaName(TEXT("BaseStruct"));
-	static const FName BaseClassMetaName(TEXT("BaseClass"));
-	static const FName OptionalMetaName(TEXT("Optional"));
-
-	//
-	// Iterate over the main level properties of the struct and look for properties of type FStateTreeExternalItemHandle.
-	// The meta data of the property defines what type will bound to it. Examples:
-	//
-	// Handle to WorldSubsystem:
-	//		UPROPERTY(meta=(BaseClass="MassStateTreeSubsystem"))
-	//		FStateTreeExternalItemHandle MassStateTreeSubSystemHandle;
-	//
-	// Handle to a struct (fragment), optional, meaning it can be null:
-	//		UPROPERTY(meta=(BaseStruct="DataFragment_SmartObjectUser", Optional))
-	//		FStateTreeExternalItemHandle SmartObjectUserHandle;
-
-    for (TPropertyValueIterator<const FStructProperty> It(ScriptStruct, Item.GetMutableMemory(), EPropertyValueIteratorFlags::NoRecursion); It; ++It)
-    {
-    	const FStructProperty* StructProperty = It->Key;
-    	check(StructProperty);
-    	if (StructProperty->Struct == FStateTreeExternalItemHandle::StaticStruct())
-    	{
-    		// Find Class or ScriptStruct.
-    		const UStruct* Struct = nullptr;
-    		if (StructProperty->HasMetaData(BaseStructMetaName))
-    		{
-	    		const FString& BaseStructName = StructProperty->GetMetaData(BaseStructMetaName);
-				Struct = FindObject<UScriptStruct>(ANY_PACKAGE, *BaseStructName);
-    			if (Struct == nullptr)
-    			{
-    				UE_LOG(LogStateTree, Error, TEXT("%s: Struct '%s' does not exists."), ANSI_TO_TCHAR(__FUNCTION__), *BaseStructName);
-    				return false;
-    			}
-    		}
-    		else if (StructProperty->HasMetaData(BaseClassMetaName))
-    		{
-    			const FString& BaseClassName = StructProperty->GetMetaData(BaseClassMetaName);
-				Struct = FindObject<UClass>(ANY_PACKAGE, *BaseClassName);
-    			if (Struct == nullptr)
-    			{
-    				UE_LOG(LogStateTree, Error, TEXT("%s: Class '%s' does not exists."), ANSI_TO_TCHAR(__FUNCTION__), *BaseClassName);
-    				return false;
-    			}
-    		}
-    		else
-    		{
-    			UE_LOG(LogStateTree, Error, TEXT("%s: FStateTreeExternalItemHandle must have 'BaseStruct' or 'BaseClass' set."), ANSI_TO_TCHAR(__FUNCTION__));
-    			return false;
-    		}
-
-    		// Parse optional
-    		const bool bOptional = StructProperty->HasMetaData(OptionalMetaName);
-
-    		// The struct/class must be accepted by the schema.
-    		if (StateTree->Schema && !StateTree->Schema->IsExternalItemAllowed(*Struct))
-    		{
-    			UE_LOG(LogStateTree, Error, TEXT("%s: Schema %s does not allow item type %s."), ANSI_TO_TCHAR(__FUNCTION__), *GetNameSafe(StateTree->Schema), *GetNameSafe(Struct));
-    			return false;
-    		}
-
-    		// Check if similar struct already exists, if not add new.
-    		FStateTreeExternalItemDesc* ExternalItem = StateTree->ExternalItems.FindByPredicate([Struct](const FStateTreeExternalItemDesc& Item) { return Item.Struct == Struct; });
-    		if (ExternalItem == nullptr)
-    		{
-    			ExternalItem = &StateTree->ExternalItems.Emplace_GetRef(Struct, bOptional);
-    		
-    			// The external struct pointer is stored in the same array as property binding sources.
-    			// This is done also as anticipation to allow to bind to external items.
-    			FStateTreeBindableStructDesc StructDesc;
-    			StructDesc.Struct = Struct;
-    			StructDesc.Name = FName(Struct->GetName() + TEXT("_External"));
-    			StructDesc.ID = FGuid(); // Empty GUID, this item cannot be bound to.
-
-    			const int32 StructIndex = BindingsCompiler.AddSourceStruct(StructDesc);
-
-    			check(StructIndex <= int32(MAX_uint16));
-    			ExternalItem->Handle.SetIndex(uint16(StructIndex));
-			}
-    		else
-    		{
-    			// If same type is requested as required, clear optional flag. 
-    			if (!bOptional)
-    			{
-    				ExternalItem->bOptional = false;
-    			}
-    		}
-
-    		check(ExternalItem);
-
-    		// Update the handle on Eval/struct
-    		FStateTreeExternalItemHandle& Handle = *static_cast<FStateTreeExternalItemHandle*>(const_cast<void*>(It.Value()));
-    		Handle = ExternalItem->Handle;
-    	}
-    }
-
-	return true;
-}
-
 bool FStateTreeBaker::CreateStateRecursive(UStateTreeState& State, const FStateTreeHandle Parent)
 {
 	const int32 StateIdx = StateTree->States.AddDefaulted();
@@ -440,12 +337,6 @@ bool FStateTreeBaker::CreateStateRecursive(UStateTreeState& State, const FStateT
 			Eval.BindingsBatch = BatchIndex == INDEX_NONE ? FStateTreeHandle::Invalid : FStateTreeHandle(uint16(BatchIndex));
 			check(SourceStructIndex <= int32(MAX_uint16));
 			Eval.SourceStructIndex = uint16(SourceStructIndex);
-
-			if (!CreateExternalItemHandles(EvalPtr))
-			{
-				UE_LOG(LogStateTree, Error, TEXT("%s: '%s': %s failed to create external item handles."), ANSI_TO_TCHAR(__FUNCTION__), *GetNameSafe(StateTree), *Eval.Name.ToString());
-				return false;
-			}
 		}
 	}
 	const int32 EvaluatorsNum = StateTree->RuntimeStorageItems.Num() - int32(BakedState.EvaluatorsBegin);
@@ -492,12 +383,6 @@ bool FStateTreeBaker::CreateStateRecursive(UStateTreeState& State, const FStateT
 			Task.BindingsBatch = BatchIndex == INDEX_NONE ? FStateTreeHandle::Invalid : FStateTreeHandle(uint16(BatchIndex));
 			check(SourceStructIndex <= int32(MAX_uint16));
 			Task.SourceStructIndex = uint16(SourceStructIndex);
-
-			if (!CreateExternalItemHandles(TaskPtr))
-			{
-				UE_LOG(LogStateTree, Error, TEXT("%s: '%s': %s failed to create external item handles."), ANSI_TO_TCHAR(__FUNCTION__), *GetNameSafe(StateTree), *Task.Name.ToString());
-				return false;
-			}
 		}
 	}
 	const int32 TasksNum = StateTree->RuntimeStorageItems.Num() - int32(BakedState.TasksBegin);
diff --git a/Engine/Plugins/Runtime/StateTree/Source/StateTreeEditorModule/Public/StateTreeBaker.h b/Engine/Plugins/Runtime/StateTree/Source/StateTreeEditorModule/Public/StateTreeBaker.h
--- a/Engine/Plugins/Runtime/StateTree/Source/StateTreeEditorModule/Public/StateTreeBaker.h
+++ b/Engine/Plugins/Runtime/StateTree/Source/StateTreeEditorModule/Public/StateTreeBaker.h
@@ -38,7 +38,6 @@ private:
 	bool CreateCondition(const FStateTreeConditionItem& CondItem);
 	bool GetAndValidateBindings(const FStateTreeBindableStructDesc& TargetStruct, TArray<FStateTreeEditorPropertyBinding>& OutBindings) const;
 	bool IsPropertyAnyEnum(const FStateTreeBindableStructDesc& Struct, FStateTreeEditorPropertyPath Path) const;
-	bool CreateExternalItemHandles(FStructView Item);
 
 	UStateTree* StateTree = nullptr;
 	UStateTreeEditorData* TreeData = nullptr;
diff --git a/Engine/Plugins/Runtime/StateTree/Source/StateTreeModule/Private/StateTree.cpp b/Engine/Plugins/Runtime/StateTree/Source/StateTreeModule/Private/StateTree.cpp
--- a/Engine/Plugins/Runtime/StateTree/Source/StateTreeModule/Private/StateTree.cpp
+++ b/Engine/Plugins/Runtime/StateTree/Source/StateTreeModule/Private/StateTree.cpp
@@ -2,6 +2,7 @@
 
 #include "StateTree.h"
 #include "StateTreeEvaluatorBase.h"
+#include "StateTreeTaskBase.h"
 #include "CoreMinimal.h"
 #include "StateTreeDelegates.h"
 
@@ -57,8 +58,11 @@ void UStateTree::GetAssetRegistryTags(TArray<FAssetRegistryTag>& OutTags) const
 void UStateTree::PostLoad()
 {
 	Super::PostLoad();
+	
 	PropertyBindings.ResolvePaths();
 
+	Link();
+	
 #if WITH_EDITOR
 	InitRuntimeStorage();
 #else
@@ -88,6 +92,25 @@ void UStateTree::ResolvePropertyPaths()
 	PropertyBindings.ResolvePaths();
 }
 
+void UStateTree::Link()
+{
+	FStateTreeLinker Linker;
+
+	for (FInstancedStruct& RuntimeItem : RuntimeStorageItems)
+	{
+		if (FStateTreeEvaluatorBase* Eval = RuntimeItem.GetMutablePtr<FStateTreeEvaluatorBase>())
+		{
+			Eval->Link(Linker);
+		}
+		else if (FStateTreeTaskBase* Task = RuntimeItem.GetMutablePtr<FStateTreeTaskBase>())
+		{
+			Task->Link(Linker);
+		}
+	}
+
+	ExternalItems = Linker.GetItemDescs();
+}
+
 void UStateTree::InitRuntimeStorage()
 {
 	RuntimeStorageStruct = nullptr;
diff --git a/Engine/Plugins/Runtime/StateTree/Source/StateTreeModule/Private/StateTreeExecutionContext.cpp b/Engine/Plugins/Runtime/StateTree/Source/StateTreeModule/Private/StateTreeExecutionContext.cpp
--- a/Engine/Plugins/Runtime/StateTree/Source/StateTreeModule/Private/StateTreeExecutionContext.cpp
+++ b/Engine/Plugins/Runtime/StateTree/Source/StateTreeModule/Private/StateTreeExecutionContext.cpp
@@ -929,7 +929,7 @@ void FStateTreeExecutionContext::DebugPrintInternalLayout(FStateTreeItemView Ext
 	DebugString += FString::Printf(TEXT("\nExternal items(%d)\n  [ %-40s | %-8s | %5s ]\n"), StateTree->ExternalItems.Num(), TEXT("Name"), TEXT("Optional"), TEXT("Index"));
 	for (const FStateTreeExternalItemDesc& Desc : StateTree->ExternalItems)
 	{
-		DebugString += FString::Printf(TEXT("  | %-40s | %8s | %5d |\n"), Desc.Struct ? *Desc.Struct->GetName() : TEXT("null"), *LexToString(Desc.bOptional), Desc.Handle.GetIndex());
+		DebugString += FString::Printf(TEXT("  | %-40s | %8s | %5d |\n"), Desc.Struct ? *Desc.Struct->GetName() : TEXT("null"), *UEnum::GetValueAsString(Desc.Requirement), Desc.Handle.ItemIndex);
 	}
 
 	// Bindings
@@ -942,8 +942,8 @@ void FStateTreeExecutionContext::DebugPrintInternalLayout(FStateTreeItemView Ext
 	{
 		DebugString += FString::Printf(TEXT("  | %3d | %15s | %-40s | %-40s | %8d |\n"),
 									Transition.ConditionsBegin, *Transition.State.Describe(),
-									*StaticEnum<EStateTreeTransitionType>()->GetValueAsString(Transition.Type),
-									*StaticEnum<EStateTreeTransitionEvent>()->GetValueAsString(Transition.Event),
+									*UEnum::GetValueAsString(Transition.Type),
+									*UEnum::GetValueAsString(Transition.Event),
 									Transition.ConditionsNum);
 	}
 
diff --git a/Engine/Plugins/Runtime/StateTree/Source/StateTreeModule/Private/StateTreeTypes.cpp b/Engine/Plugins/Runtime/StateTree/Source/StateTreeModule/Private/StateTreeTypes.cpp
--- a/Engine/Plugins/Runtime/StateTree/Source/StateTreeModule/Private/StateTreeTypes.cpp
+++ b/Engine/Plugins/Runtime/StateTree/Source/StateTreeModule/Private/StateTreeTypes.cpp
@@ -7,3 +7,5 @@ DEFINE_LOG_CATEGORY(LogStateTree);
 const FStateTreeHandle FStateTreeHandle::Invalid = FStateTreeHandle();
 const FStateTreeHandle FStateTreeHandle::Succeeded = FStateTreeHandle(FStateTreeHandle::SucceededIndex);
 const FStateTreeHandle FStateTreeHandle::Failed = FStateTreeHandle(FStateTreeHandle::FailedIndex);
+
+const FStateTreeItemHandle FStateTreeItemHandle::Invalid = FStateTreeItemHandle();
diff --git a/Engine/Plugins/Runtime/StateTree/Source/StateTreeModule/Public/StateTree.h b/Engine/Plugins/Runtime/StateTree/Source/StateTreeModule/Public/StateTree.h
--- a/Engine/Plugins/Runtime/StateTree/Source/StateTreeModule/Public/StateTree.h
+++ b/Engine/Plugins/Runtime/StateTree/Source/StateTreeModule/Public/StateTree.h
@@ -63,6 +63,9 @@ protected:
 	/** Initializes the types and default values related to StateTree runtime storage */
 	void InitRuntimeStorage();
 
+	/** Resolved references between data in the StateTree. */
+	void Link();
+	
 	/** @return Number of items (Evaluators & Tasks) in the runtime storage. */
 	int32 GetRuntimeStorageItemCount() const { return RuntimeStorageOffsets.Num(); }
 
@@ -91,8 +94,8 @@ private:
 	UPROPERTY(meta = (BaseStruct = "StateTreeConditionBase", ExcludeBaseStruct))
 	TArray<FInstancedStruct> Conditions;
 
-	/** List of external items required by the state tree */
-	UPROPERTY()
+	/** List of external items required by the state tree, creating during linking. */
+	UPROPERTY(Transient)
 	TArray<FStateTreeExternalItemDesc> ExternalItems;
 
 	UPROPERTY()
diff --git a/Engine/Plugins/Runtime/StateTree/Source/StateTreeModule/Public/StateTreeEvaluatorBase.h b/Engine/Plugins/Runtime/StateTree/Source/StateTreeModule/Public/StateTreeEvaluatorBase.h
--- a/Engine/Plugins/Runtime/StateTree/Source/StateTreeModule/Public/StateTreeEvaluatorBase.h
+++ b/Engine/Plugins/Runtime/StateTree/Source/StateTreeModule/Public/StateTreeEvaluatorBase.h
@@ -23,6 +23,14 @@ struct STATETREEMODULE_API FStateTreeEvaluatorBase // TODO: change to FStateTree
 
 	virtual ~FStateTreeEvaluatorBase() {}
 
+	/**
+	 * Called when the StateTree asset is linked. Allows to resolve references to other StateTree data.
+	 * @see TStateTreeItemHandle.
+	 * @param Linker Reference to the linker
+	 * @return true if linking succeeded. 
+	 */
+	virtual bool Link(FStateTreeLinker& Linker) { return true; }
+	
 	/**
 	 * Called when a new state is entered and evaluator is part of active states. The change type parameter describes if the evaluator's state
 	 * was previously part of the list of active states (Sustained), or if it just became active (Changed).
diff --git a/Engine/Plugins/Runtime/StateTree/Source/StateTreeModule/Public/StateTreeExecutionContext.h b/Engine/Plugins/Runtime/StateTree/Source/StateTreeModule/Public/StateTreeExecutionContext.h
--- a/Engine/Plugins/Runtime/StateTree/Source/StateTreeModule/Public/StateTreeExecutionContext.h
+++ b/Engine/Plugins/Runtime/StateTree/Source/StateTreeModule/Public/StateTreeExecutionContext.h
@@ -90,40 +90,74 @@ public:
 		bool bResult = true;
 		for (const FStateTreeExternalItemDesc& ItemDesc : StateTree->ExternalItems)
 		{
-			const FStateTreeItemView& ItemView = ItemViews[ItemDesc.Handle.GetIndex()];
-			if (!ItemDesc.bOptional && (ItemView.GetMemory() == nullptr || ItemView.GetStruct() != ItemDesc.Struct))
+			const FStateTreeItemView& ItemView = ItemViews[ItemDesc.Handle.ItemIndex];
+			
+			if (ItemDesc.Requirement == EStateTreeItemRequirement::Required)
 			{
-				bResult = false;
-				break;
+				// Required items must have valid pointer and expected type.  
+				if (ItemView.GetMemory() == nullptr || ItemView.GetStruct() != ItemDesc.Struct)
+				{
+					bResult = false;
+					break;
+				}
+			}
+			else
+			{
+				// Optional items must have same type if they are set.
+				if (ItemView.IsValid() && ItemView.GetStruct() != ItemDesc.Struct)
+				{
+					bResult = false;
+					break;
+				}
+			 
 			}
 		}
 		return bResult;
 	}
 
-	/** @return Handle to external item of type InStruct, or invalid handle if not found */ 
-	FStateTreeExternalItemHandle GetExternalItemHandleByStruct(const UStruct* InStruct) const
+	/** @return Handle to external item of type InStruct, or invalid handle if struct not found. */ 
+	FStateTreeItemHandle GetExternalItemHandleByStruct(const UStruct* InStruct) const
 	{
 		check(StateTree);
 		const FStateTreeExternalItemDesc* Item = StateTree->ExternalItems.FindByPredicate([InStruct](const FStateTreeExternalItemDesc& Item) { return Item.Struct == InStruct; });
-		return Item != nullptr ? Item->Handle : FStateTreeExternalItemHandle();
+		return Item != nullptr ? Item->Handle : FStateTreeItemHandle::Invalid;
 	}
 
-	/** Sets external item view based on handle */ 
-	void SetExternalItem(const FStateTreeExternalItemHandle ItemHandle, FStateTreeItemView Item)
+	/** Sets external item view value for specific item. */ 
+	void SetExternalItem(const FStateTreeItemHandle ItemHandle, FStateTreeItemView Item)
 	{
 		check(StateTree);
 		check(ItemHandle.IsValid());
-		ItemViews[ItemHandle.GetIndex()] = Item;
+		ItemViews[ItemHandle.ItemIndex] = Item;
 	}
 
-	/** @return item view to external item based on handle */ 
-	FStateTreeItemView GetExternalItem(const FStateTreeExternalItemHandle ItemHandle) const
+	/**
+	 * Returns reference to external item based on provided item handle. The return type is deducted from the handle's template type.
+     * @param ItemHandle Valid TStateTreeItemHandle<> item handle. 
+	 * @return reference to external item based on handle or null if item is not set.
+	 */ 
+	template <typename T>
+	typename T::ItemType& GetExternalItem(const T ItemHandle) const
 	{
 		check(StateTree);
 		check(ItemHandle.IsValid());
-		return ItemViews[ItemHandle.GetIndex()];
+		checkSlow(StateTree->ExternalItems[ItemHandle.ItemIndex].Requirement != EStateTreeItemRequirement::Optional); // Optionals should query pointer instead.
+		return ItemViews[ItemHandle.ItemIndex].GetMutable<T::ItemType>();
 	}
 
+	/**
+	 * Returns pointer to external item based on provided item handle. The return type is deducted from the handle's template type.
+     * @param ItemHandle Valid TStateTreeItemHandle<> item handle.
+	 * @return pointer to external item based on handle or null if item is not set.
+	 */ 
+	template <typename T>
+	typename T::ItemType* GetExternalItemPtr(const T ItemHandle) const
+	{
+		check(StateTree);
+		check(ItemHandle.IsValid());
+		return ItemViews[ItemHandle.ItemIndex].GetMutablePtr<T::ItemType>();
+	}
+	
 	EStateTreeRunStatus GetLastTickStatus(FStateTreeItemView ExternalStorage = FStateTreeItemView()) const;
 	EStateTreeRunStatus GetEnterStateStatus() const { return EnterStateStatus; }
 
diff --git a/Engine/Plugins/Runtime/StateTree/Source/StateTreeModule/Public/StateTreeTaskBase.h b/Engine/Plugins/Runtime/StateTree/Source/StateTreeModule/Public/StateTreeTaskBase.h
--- a/Engine/Plugins/Runtime/StateTree/Source/StateTreeModule/Public/StateTreeTaskBase.h
+++ b/Engine/Plugins/Runtime/StateTree/Source/StateTreeModule/Public/StateTreeTaskBase.h
@@ -22,6 +22,14 @@ struct STATETREEMODULE_API FStateTreeTaskBase
 
 	virtual ~FStateTreeTaskBase() {}
 
+	/**
+	 * Called when the StateTree asset is linked. Allows to resolve references to other StateTree data.
+	 * @see TStateTreeItemHandle.
+	 * @param Linker Reference to the linker
+	 * @return true if linking succeeded. 
+	 */
+	virtual bool Link(FStateTreeLinker& Linker) { return true; }
+
 	/**
 	 * Called when a new state is entered and task is part of active states. The change type parameter describes if the task's state
 	 * was previously part of the list of active states (Sustained), or if it just became active (Changed).
diff --git a/Engine/Plugins/Runtime/StateTree/Source/StateTreeModule/Public/StateTreeTypes.h b/Engine/Plugins/Runtime/StateTree/Source/StateTreeModule/Public/StateTreeTypes.h
--- a/Engine/Plugins/Runtime/StateTree/Source/StateTreeModule/Public/StateTreeTypes.h
+++ b/Engine/Plugins/Runtime/StateTree/Source/StateTreeModule/Public/StateTreeTypes.h
@@ -363,100 +363,136 @@ struct FStateTreeRuntimeStorageItemOffset
 	int32 Offset = 0;
 };
 
-
+UENUM()
+enum class EStateTreeItemRequirement : uint8
+{
+	Required,	// StateTree cannot be executed if the item is not present.
+	Optional,	// Item is optional for StateTree execution.
+};
 
 /**
- * Handle to an external struct or object. 
+ * Handle to access an external struct or object.
+ * Note: Use the templated version below. 
  */
 USTRUCT()
-struct STATETREEMODULE_API FStateTreeExternalItemHandle
+struct STATETREEMODULE_API FStateTreeItemHandle
 {
 	GENERATED_BODY()
 
-	FStateTreeExternalItemHandle() = default;
-	explicit FStateTreeExternalItemHandle(const uint16 InIndex) : ItemIndex(InIndex) {}
-	
-	static const uint16 InvalidIndex = uint16(-1);		// Index value indicating invalid item.
+	static const FStateTreeItemHandle Invalid;
+	static constexpr uint8 IndexNone = MAX_uint8;
 
-	bool IsValid() const { return ItemIndex != InvalidIndex; }
+	static bool IsValidIndex(const int32 Index) { return Index >= 0 && Index < (int32)IndexNone; }
 
-	void SetIndex(const uint16 Index) { ItemIndex = Index; }
-	uint16 GetIndex() const { return ItemIndex; }
+	bool IsValid() const { return ItemIndex != IndexNone; }
 	
-protected:
-	UPROPERTY();
-	uint16 ItemIndex = InvalidIndex;
+	uint8 ItemIndex = IndexNone;
 };
 
+/**
+ * Handle to access an external struct or object.
+ * This reference handle can be used in StateTree tasks and evaluators to have quick access to external items.
+ * The type provided to the template is used by the linker and context to pass along the type.
+ *
+ * USTRUCT()
+ * struct FExampleTask : public FStateTreeTaskBase
+ * {
+ *    ...
+ *
+ *    bool Link(FStateTreeLinker& Linker)
+ *    {
+ *      Linker.LinkExternalItem(ExampleSubsystemHandle);
+ *      return true;
+ *    }
+ * 
+ *    EStateTreeRunStatus EnterState(FStateTreeExecutionContext& Context, const EStateTreeStateChangeType ChangeType, const FStateTreeTransitionResult& Transition)
+ *    {
+ *      const UExampleSubsystem& ExampleSubsystem = Context.GetExternalItem(ExampleSubsystemHandle);
+ *      ...
+ *    }
+ *
+ *    TStateTreeItemHandle<UExampleSubsystem> ExampleSubsystemHandle;
+ * }
+ */
+template<typename T, EStateTreeItemRequirement Req = EStateTreeItemRequirement::Required>
+struct TStateTreeItemHandle : FStateTreeItemHandle
+{
+	typedef T ItemType;
+	static constexpr EStateTreeItemRequirement ItemRequirement = Req;
+};
 
+/**
+ * Describes an external item. The item can point to a struct or object.
+ * The code that handles StateTree ticking is responsible for passing in the actually data, see FStateTreeExecutionContext.
+ */
 USTRUCT()
 struct STATETREEMODULE_API FStateTreeExternalItemDesc
 {
 	GENERATED_BODY()
 
 	FStateTreeExternalItemDesc() = default;
-	FStateTreeExternalItemDesc(const UStruct* InStruct, const bool bInOptional) : Struct(InStruct), bOptional(bInOptional) {}
+	FStateTreeExternalItemDesc(const UStruct* InStruct, const EStateTreeItemRequirement InRequirement) : Struct(InStruct), Requirement(InRequirement) {}
+
+	bool operator==(const FStateTreeExternalItemDesc& Other) const
+	{
+		return Struct == Other.Struct && Requirement == Other.Requirement;
+	}
 	
 	/** Class or struct of the external item. */
 	UPROPERTY();
 	const UStruct* Struct = nullptr;
 
-	/** Handle to the StateTreeExecutionContext item views array */
+	/** Handle/Index to the StateTreeExecutionContext item views array */
 	UPROPERTY();
-	FStateTreeExternalItemHandle Handle;
+	FStateTreeItemHandle Handle;
 
-	/** If true, the extern item is optional (can be null). */
+	/** Describes if the item is required or not. */
 	UPROPERTY();
-	bool bOptional = false;
+	EStateTreeItemRequirement Requirement = EStateTreeItemRequirement::Required;
 };
 
-
-namespace FStateTreeVariableHelpers
+/**
+ * The StateTree linker is used to resolved references to various StateTree data at load time.
+ * @see TStateTreeItemHandle<> for example usage.
+ */
+struct FStateTreeLinker
 {
-	// Returns size of a variable type in bytes.
-	inline uint16 GetVariableMemoryUsage(EStateTreeVariableType Type)
+	/**
+	 * Links reference to an external UObject.
+	 * @param Handle Reference to TStateTreeItemHandle<> with UOBJECT type to link to.
+	 */
+	template <typename T>
+	typename TEnableIf<TIsDerivedFrom<typename T::ItemType, UObject>::IsDerived, void>::Type LinkExternalItem(T& Handle)
 	{
-		static const uint64 Alignment = 4;
-		switch (Type)
-		{
-		case EStateTreeVariableType::Float:
-			return uint16(Align(sizeof(float), Alignment));
-		case EStateTreeVariableType::Int:
-			return uint16(Align(sizeof(int32), Alignment));
-		case EStateTreeVariableType::Bool:
-			return uint16(Align(sizeof(bool), Alignment));
-		case EStateTreeVariableType::Vector:
-			return uint16(Align(sizeof(FVector), Alignment));
-		case EStateTreeVariableType::Object:
-			return uint16(Align(sizeof(FWeakObjectPtr), Alignment));
-		default:
-			return 0;
-		}
+		LinkExternalItem(Handle, T::ItemType::StaticClass(), T::ItemRequirement);
 	}
 
-	// Get simple value based on pointer in byte array.
-	template<typename T>
-	static const T& GetValueFromMemory(const uint8* MemoryBlock)
+	/**
+	 * Links reference to an external UObject.
+	 * @param Handle Reference to TStateTreeItemHandle<> with USTRUCT type to link to.
+	 */
+	template <typename T>
+	typename TEnableIf<!TIsDerivedFrom<typename T::ItemType, UObject>::IsDerived, void>::Type LinkExternalItem(T& Handle)
 	{
-		return *((T*)MemoryBlock);
+		LinkExternalItem(Handle, T::ItemType::StaticStruct(), T::ItemRequirement);
 	}
 
-	// Helper function for writing typed data to memory block, returns true if value has changed.
-	template<typename T>
-	static bool SetValueInMemory(uint8* MemoryBlock, const T& Value)
-	{
-		const bool bChanged = *((T*)MemoryBlock) != Value;
-		*((T*)MemoryBlock) = Value;
-
-		return bChanged;
-	}
+	/** @return linked external item descriptors. */
+	TConstArrayView<FStateTreeExternalItemDesc> GetItemDescs() const { return ItemDescs; }
 
-	// Helper function for writing weak object data to memory block, returns true if value has changed.
-	static bool SetWeakObjectInMemory(uint8* MemoryBlock, const FWeakObjectPtr& Value)
+protected:
+	void LinkExternalItem(FStateTreeItemHandle& Handle, const UStruct* Struct, const EStateTreeItemRequirement Requirement)
 	{
-		FWeakObjectPtr* PrevValue = (FWeakObjectPtr*)MemoryBlock;
-		const bool bChanged = !PrevValue->HasSameIndexAndSerialNumber(Value);
-		*((FWeakObjectPtr*)MemoryBlock) = Value;
-		return bChanged;
+		const FStateTreeExternalItemDesc Desc(Struct, Requirement);
+		int32 Index = ItemDescs.Find(Desc);
+		if (Index == INDEX_NONE)
+		{
+			Index = ItemDescs.Add(Desc);
+			ItemDescs[Index].Handle.ItemIndex = (uint8)Index;
+		}
+		check(FStateTreeItemHandle::IsValidIndex(Index));
+		Handle.ItemIndex = (uint8)Index;
 	}
-};
\ No newline at end of file
+
+	TArray<FStateTreeExternalItemDesc> ItemDescs;
+};
