Added new console command to freeze any mass movement, helps debugging pops

diff --git a/Engine/Plugins/AI/MassAI/Source/MassAIMovement/Private/NavigationProcessor.cpp b/Engine/Plugins/AI/MassAI/Source/MassAIMovement/Private/NavigationProcessor.cpp
--- a/Engine/Plugins/AI/MassAI/Source/MassAIMovement/Private/NavigationProcessor.cpp
+++ b/Engine/Plugins/AI/MassAI/Source/MassAIMovement/Private/NavigationProcessor.cpp
@@ -21,6 +21,8 @@
 
 namespace UE::MassMovement
 {
+	int32 bFreezeMovement = 0;
+	FAutoConsoleVariableRef CVarFreezeMovement(TEXT("ai.debug.mass.FreezeMovement"), bFreezeMovement, TEXT("Freeze any movement by the UMassApplyVelocityMoveTargetProcessor"));
 
 	static float Damp(const float X, const float Goal, const float HalfLife, const float DeltaTime)
 	{
@@ -156,6 +158,15 @@ void UMassApplyVelocityMoveTargetProcessor::Execute(UMassEntitySubsystem& Entity
 
 				FVector DeltaLoc = Velocity.Value * TimeDelta;
 
+#if WITH_MASSGAMEPLAY_DEBUG
+				if (UE::MassMovement::bFreezeMovement)
+				{
+					DeltaLoc.X = 0;
+					DeltaLoc.Y = 0;
+					Velocity.Value = FVector::ZeroVector;
+				}
+#endif // WITH_MASSGAMEPLAY_DEBUG
+
 				// Apply DeltaLoc on X,Y and set the Z of the current lane location.
 				FTransform& CurrentTransform = LocationList[EntityIndex].GetMutableTransform();
 
@@ -205,6 +216,12 @@ void UMassApplyVelocityMoveTargetProcessor::Execute(UMassEntitySubsystem& Entity
 
 		LowResEntityQuery_Conditional.ForEachEntityChunk(EntitySubsystem, Context, [this](FMassExecutionContext& Context)
 			{
+#if WITH_MASSGAMEPLAY_DEBUG
+				if (UE::MassMovement::bFreezeMovement)
+				{
+					return;
+				}
+#endif // WITH_MASSGAMEPLAY_DEBUG
 				const int32 NumEntities = Context.GetNumEntities();
 
 				const TArrayView<FDataFragment_Transform> LocationList = Context.GetMutableFragmentView<FDataFragment_Transform>();
