Cut UVelocityFragmentRandomizer

diff --git a/Engine/Plugins/AI/MassAI/Source/MassAIMovement/Private/MovementProcessor.cpp b/Engine/Plugins/AI/MassAI/Source/MassAIMovement/Private/MovementProcessor.cpp
--- a/Engine/Plugins/AI/MassAI/Source/MassAIMovement/Private/MovementProcessor.cpp
+++ b/Engine/Plugins/AI/MassAI/Source/MassAIMovement/Private/MovementProcessor.cpp
@@ -197,29 +197,3 @@ void UMassAvoidanceObstacleRemoverFragmentDestructor::Execute(UMassEntitySubsyst
 			}
 		});
 }
-
-
-//----------------------------------------------------------------------//
-//  UVelocityFragmentInitializer
-//----------------------------------------------------------------------//
-UVelocityFragmentInitializer::UVelocityFragmentInitializer()
-{
-	FragmentType = FMassVelocityFragment::StaticStruct();
-}
-
-void UVelocityFragmentInitializer::ConfigureQueries()
-{
-	EntityQuery.AddRequirement<FMassVelocityFragment>(EMassFragmentAccess::ReadWrite);
-}
-
-void UVelocityFragmentInitializer::Execute(UMassEntitySubsystem& EntitySubsystem, FMassExecutionContext& Context)
-{
-	EntityQuery.ForEachEntityChunk(EntitySubsystem, Context, [this](FMassExecutionContext& Context)
-		{
-			const TArrayView<FMassVelocityFragment> VelocitiesList = Context.GetMutableFragmentView<FMassVelocityFragment>();
-			for (FMassVelocityFragment& VelocityFragment : VelocitiesList)
-			{
-				VelocityFragment.Value = FVector(FMath::FRandRange(-100.f, 100.f), FMath::FRandRange(-100.f, 100.f), 0);
-			}
-		});
-}
diff --git a/Engine/Plugins/AI/MassAI/Source/MassAIMovement/Public/MovementProcessor.h b/Engine/Plugins/AI/MassAI/Source/MassAIMovement/Public/MovementProcessor.h
--- a/Engine/Plugins/AI/MassAI/Source/MassAIMovement/Public/MovementProcessor.h
+++ b/Engine/Plugins/AI/MassAI/Source/MassAIMovement/Public/MovementProcessor.h
@@ -51,17 +51,3 @@ protected:
 
 	FMassEntityQuery EntityQuery;
 };
-
-UCLASS()
-class MASSAIMOVEMENT_API UVelocityFragmentInitializer : public UMassFragmentInitializer
-{
-	GENERATED_BODY()
-public:
-	UVelocityFragmentInitializer();
-
-protected:
-	virtual void ConfigureQueries() override;
-	virtual void Execute(UMassEntitySubsystem& EntitySubsystem, FMassExecutionContext& Context) override;
-
-	FMassEntityQuery EntityQuery;
-};
