Rename the StateTreeRequestInitialization or StateTreeActivate

#rnx

diff --git a/Engine/Plugins/AI/MassAI/Source/MassAIBehavior/Private/MassStateTreeProcessors.cpp b/Engine/Plugins/AI/MassAI/Source/MassAIBehavior/Private/MassStateTreeProcessors.cpp
--- a/Engine/Plugins/AI/MassAI/Source/MassAIBehavior/Private/MassStateTreeProcessors.cpp
+++ b/Engine/Plugins/AI/MassAI/Source/MassAIBehavior/Private/MassStateTreeProcessors.cpp
@@ -210,56 +210,56 @@ void UMassStateTreeFragmentDestructor::Execute(UMassEntitySubsystem& EntitySubsy
 }
 
 //----------------------------------------------------------------------//
-// UMassStateTreeSignalRequestInitialization
+// UMassStateTreeActivationProcessor
 //----------------------------------------------------------------------//
-UMassStateTreeSignalRequestInitialization::UMassStateTreeSignalRequestInitialization()
+UMassStateTreeActivationProcessor::UMassStateTreeActivationProcessor()
 {
 	ExecutionOrder.ExecuteAfter.Add(UE::Mass::ProcessorGroupNames::LOD);
 	ExecutionOrder.ExecuteBefore.Add(UE::Mass::ProcessorGroupNames::Behavior);
-	MaxInitializationRequestsPerLOD[EMassLOD::High] = 100;
-	MaxInitializationRequestsPerLOD[EMassLOD::Medium] = 100;
-	MaxInitializationRequestsPerLOD[EMassLOD::Low] = 100;
-	MaxInitializationRequestsPerLOD[EMassLOD::Off] = 100;
+	MaxActivationsPerLOD[EMassLOD::High] = 100;
+	MaxActivationsPerLOD[EMassLOD::Medium] = 100;
+	MaxActivationsPerLOD[EMassLOD::Low] = 100;
+	MaxActivationsPerLOD[EMassLOD::Off] = 100;
 }
-void UMassStateTreeSignalRequestInitialization::Initialize(UObject& Owner)
+void UMassStateTreeActivationProcessor::Initialize(UObject& Owner)
 {
 	SignalSubsystem = UWorld::GetSubsystem<UMassSignalSubsystem>(Owner.GetWorld());
 }
-void UMassStateTreeSignalRequestInitialization::ConfigureQueries()
+void UMassStateTreeActivationProcessor::ConfigureQueries()
 {
 	EntityQuery.AddRequirement<FMassStateTreeFragment>(EMassFragmentAccess::ReadOnly);
-	EntityQuery.AddTagRequirement<FMassStateTreeInitializationRequestDone>(EMassFragmentPresence::None);
+	EntityQuery.AddTagRequirement<FMassStateTreeActivated>(EMassFragmentPresence::None);
 	EntityQuery.AddChunkRequirement<FMassSimulationVariableTickChunkFragment>(EMassFragmentAccess::ReadOnly);
 }
-void UMassStateTreeSignalRequestInitialization::Execute(UMassEntitySubsystem& EntitySubsystem, FMassExecutionContext& Context)
+void UMassStateTreeActivationProcessor::Execute(UMassEntitySubsystem& EntitySubsystem, FMassExecutionContext& Context)
 {
 	// StateTree processor relies on signals to be ticked but we need an 'initial tick' to set the tree in the proper state.
 	// The initializer provides that by sending a signal to all new entities that use StateTree.
 	TArray<FMassEntityHandle> EntitiesToSignal;
-	int32 InitializationRequestCounts[EMassLOD::Max] {0,0,0,0};
-	EntityQuery.ForEachEntityChunk(EntitySubsystem, Context, [&EntitiesToSignal, &InitializationRequestCounts, MaxInitializationRequestsPerLOD = MaxInitializationRequestsPerLOD](FMassExecutionContext& Context)
+	int32 ActivationCounts[EMassLOD::Max] {0,0,0,0};
+	EntityQuery.ForEachEntityChunk(EntitySubsystem, Context, [&EntitiesToSignal, &ActivationCounts, MaxActivationsPerLOD = MaxActivationsPerLOD](FMassExecutionContext& Context)
 	{
 		const int32 NumEntities = Context.GetNumEntities();
 		// Check if we already reached the maximum for this frame
 		const EMassLOD::Type ChunkLOD = FMassSimulationVariableTickChunkFragment::GetChunkLOD(Context);
-		if (InitializationRequestCounts[ChunkLOD] > MaxInitializationRequestsPerLOD[ChunkLOD])
+		if (ActivationCounts[ChunkLOD] > MaxActivationsPerLOD[ChunkLOD])
 		{
 			return;
 		}
-		InitializationRequestCounts[ChunkLOD] += NumEntities;
+		ActivationCounts[ChunkLOD] += NumEntities;
 		// Append all entities of the current chunk to the consolidated list to send signal once
 		EntitiesToSignal.Append(Context.GetEntities().GetData(), Context.GetEntities().Num());
 		// Adding a tag on each entities to remember we have sent the state tree initialization signal
 		for (int32 i = 0; i < NumEntities; ++i)
 		{
-			Context.Defer().AddTag<FMassStateTreeInitializationRequestDone>(Context.GetEntity(i));
+			Context.Defer().AddTag<FMassStateTreeActivated>(Context.GetEntity(i));
 		}
 	});
 	// Signal all entities inside the consolidated list
 	if (EntitiesToSignal.Num())
 	{
-		checkf(SignalSubsystem != nullptr, TEXT("MassSignalSubsystem should exist when executing sending state tree initialization requests."));
-		SignalSubsystem->SignalEntities(UE::Mass::Signals::StateTreeInitializationRequested, EntitiesToSignal);
+		checkf(SignalSubsystem != nullptr, TEXT("Expecting a valid MassSignalSubsystem when activating state trees."));
+		SignalSubsystem->SignalEntities(UE::Mass::Signals::StateTreeActivate, EntitiesToSignal);
 	}
 }
 
@@ -285,7 +285,7 @@ void UMassStateTreeProcessor::Initialize(UObject& Owner)
 	Super::Initialize(Owner);
 	MassStateTreeSubsystem = UWorld::GetSubsystem<UMassStateTreeSubsystem>(Owner.GetWorld());
 
-	SubscribeToSignal(UE::Mass::Signals::StateTreeInitializationRequested);
+	SubscribeToSignal(UE::Mass::Signals::StateTreeActivate);
 	SubscribeToSignal(UE::Mass::Signals::LookAtFinished);
 	SubscribeToSignal(UE::Mass::Signals::NewStateTreeTaskRequired);
 	SubscribeToSignal(UE::Mass::Signals::StandTaskFinished);
diff --git a/Engine/Plugins/AI/MassAI/Source/MassAIBehavior/Public/MassStateTreeProcessors.h b/Engine/Plugins/AI/MassAI/Source/MassAIBehavior/Public/MassStateTreeProcessors.h
--- a/Engine/Plugins/AI/MassAI/Source/MassAIBehavior/Public/MassStateTreeProcessors.h
+++ b/Engine/Plugins/AI/MassAI/Source/MassAIBehavior/Public/MassStateTreeProcessors.h
@@ -44,22 +44,21 @@ protected:
 };
 
 /**
- * Special tag to know if the state tree initialization request was signaled
+ * Special tag to know if the state tree has been activated
  */
 USTRUCT()
-struct MASSAIBEHAVIOR_API FMassStateTreeInitializationRequestDone : public FMassTag
+struct MASSAIBEHAVIOR_API FMassStateTreeActivated : public FMassTag
 {
 	GENERATED_BODY()
 };
 /**
- * Processor to signal initialization requests
- */
+ * Processor to send the activation signal to the state tree which will execute the first tick */
 UCLASS()
-class MASSAIBEHAVIOR_API UMassStateTreeSignalRequestInitialization : public UMassProcessor
+class MASSAIBEHAVIOR_API UMassStateTreeActivationProcessor : public UMassProcessor
 {
 	GENERATED_BODY()
 public:
-	UMassStateTreeSignalRequestInitialization();
+	UMassStateTreeActivationProcessor();
 protected:
 	virtual void Initialize(UObject& Owner) override;
 	virtual void ConfigureQueries() override;
@@ -70,9 +69,9 @@ protected:
 	
 	FMassEntityQuery EntityQuery;
 
-	/** The maximum number of entities to send State Tree initialization request per LOD for each update, -1 means no limit */
+	/** The maximum number of entities to signal activation per LOD for each update, -1 means no limit */
 	UPROPERTY(EditAnywhere, Category = "Mass|LOD", config)
-	int32 MaxInitializationRequestsPerLOD[EMassLOD::Max];
+	int32 MaxActivationsPerLOD[EMassLOD::Max];
 };
 
 /** 
diff --git a/Engine/Plugins/AI/MassAI/Source/MassAIBehavior/Public/MassStateTreeTypes.h b/Engine/Plugins/AI/MassAI/Source/MassAIBehavior/Public/MassStateTreeTypes.h
--- a/Engine/Plugins/AI/MassAI/Source/MassAIBehavior/Public/MassStateTreeTypes.h
+++ b/Engine/Plugins/AI/MassAI/Source/MassAIBehavior/Public/MassStateTreeTypes.h
@@ -12,7 +12,7 @@
  */
 namespace UE::Mass::Signals
 {
-	const FName StateTreeInitializationRequested = FName(TEXT("StateTreeInitializationRequested"));
+	const FName StateTreeActivate = FName(TEXT("StateTreeActivate"));
 	const FName LookAtFinished = FName(TEXT("LookAtFinished"));
 	const FName NewStateTreeTaskRequired = FName(TEXT("NewStateTreeTaskRequired"));
 	const FName StandTaskFinished = FName(TEXT("StandTaskFinished"));
