removed unnecessary FMassSimulationLODFragment handle from MassZoneGraphStandTask

diff --git a/Engine/Plugins/AI/MassAI/Source/MassAIBehavior/Private/Tasks/MassZoneGraphStandTask.cpp b/Engine/Plugins/AI/MassAI/Source/MassAIBehavior/Private/Tasks/MassZoneGraphStandTask.cpp
--- a/Engine/Plugins/AI/MassAI/Source/MassAIBehavior/Private/Tasks/MassZoneGraphStandTask.cpp
+++ b/Engine/Plugins/AI/MassAI/Source/MassAIBehavior/Private/Tasks/MassZoneGraphStandTask.cpp
@@ -20,7 +20,6 @@ bool FMassZoneGraphStandTask::Link(FStateTreeLinker& Linker)
 	Linker.LinkExternalItem(CachedLaneHandle);
 	Linker.LinkExternalItem(ZoneGraphSubsystemHandle);
 	Linker.LinkExternalItem(MassSignalSubsystemHandle);
-	Linker.LinkExternalItem(SimulationLODHandle);
 	Linker.LinkExternalItem(MovementConfigHandle);
 
 	return true;
diff --git a/Engine/Plugins/AI/MassAI/Source/MassAIBehavior/Public/Tasks/MassZoneGraphStandTask.h b/Engine/Plugins/AI/MassAI/Source/MassAIBehavior/Public/Tasks/MassZoneGraphStandTask.h
--- a/Engine/Plugins/AI/MassAI/Source/MassAIBehavior/Public/Tasks/MassZoneGraphStandTask.h
+++ b/Engine/Plugins/AI/MassAI/Source/MassAIBehavior/Public/Tasks/MassZoneGraphStandTask.h
@@ -13,7 +13,6 @@ struct FMassZoneGraphShortPathFragment;
 struct FMassZoneGraphCachedLaneFragment;
 class UZoneGraphSubsystem;
 class UMassSignalSubsystem;
-struct FMassSimulationLODFragment;
 struct FMassMovementConfigFragment;
 
 /**
@@ -35,7 +34,6 @@ protected:
 	TStateTreeItemHandle<FMassZoneGraphCachedLaneFragment> CachedLaneHandle;
 	TStateTreeItemHandle<UZoneGraphSubsystem> ZoneGraphSubsystemHandle;
 	TStateTreeItemHandle<UMassSignalSubsystem> MassSignalSubsystemHandle;
-	TStateTreeItemHandle<FMassSimulationLODFragment> SimulationLODHandle;
 	TStateTreeItemHandle<FMassMovementConfigFragment> MovementConfigHandle;
 
 	UPROPERTY()
