Removed MassMovement -> MassAIMovement redirect as unnecessary.

diff --git a/Engine/Plugins/AI/MassAI/Config/BaseMassAI.ini b/Engine/Plugins/AI/MassAI/Config/BaseMassAI.ini
--- a/Engine/Plugins/AI/MassAI/Config/BaseMassAI.ini
+++ b/Engine/Plugins/AI/MassAI/Config/BaseMassAI.ini
@@ -1,6 +1,5 @@
 [CoreRedirects]
 +PackageRedirects=(OldName="/Script/MassVisualization",NewName="/Script/MassRepresentation")
-+PackageRedirects=(OldName="/Script/MassMovement",NewName="/Script/MassAIMovement")
 +PackageRedirects=(OldName="/Script/MassBehavior",NewName="/Script/MassAIBehavior")
 
 +ClassRedirects=(OldName="/Script/MassMovement.MassAvoidanceObstacleTrait",NewName="/Script/MassAIMovement.MassAvoidanceObstacleTrait")
