Cleaned up MassAvoidanceTrait to not contain unnecessary stuff (ZoneGraphMovement already has its own trait) and extracted the LOD part into a separate trait.

diff --git a/Engine/Plugins/AI/MassAI/Source/MassAIMovement/Private/MassAvoidanceTrait.cpp b/Engine/Plugins/AI/MassAI/Source/MassAIMovement/Private/MassAvoidanceTrait.cpp
--- a/Engine/Plugins/AI/MassAI/Source/MassAIMovement/Private/MassAvoidanceTrait.cpp
+++ b/Engine/Plugins/AI/MassAI/Source/MassAIMovement/Private/MassAvoidanceTrait.cpp
@@ -6,43 +6,15 @@
 #include "MassAvoidanceProcessors.h"
 #include "MassCommonFragments.h"
 #include "MassAIMovementFragments.h"
-#include "MassSimulationLOD.h"
-#include "MassZoneGraphMovementFragments.h"
+
 
 void UMassAvoidanceTrait::BuildTemplate(FMassEntityTemplateBuildContext& BuildContext, UWorld& World) const
 {
 	// UMassAvoidanceProcessor
 	BuildContext.AddFragmentWithDefaultInitializer<FDataFragment_AgentRadius>(); // From UMassProcessor_AgentMovement
 	BuildContext.AddFragmentWithDefaultInitializer<FMassNavigationEdgesFragment>();
-	BuildContext.AddFragmentWithDefaultInitializer<FMassMoveTargetFragment>();
 	BuildContext.AddFragmentWithDefaultInitializer<FDataFragment_Transform>(); // From UMassProcessor_Movement
 	BuildContext.AddFragmentWithDefaultInitializer<FMassVelocityFragment>(); // From UMassProcessor_Movement
 	BuildContext.AddFragmentWithDefaultInitializer<FMassSteeringFragment>();
-	
-	// LOD: UMassSimulationLODProcessors
-	BuildContext.AddFragmentWithDefaultInitializer<FDataFragment_Transform>();
-	BuildContext.AddFragmentWithDefaultInitializer<FMassLODInfoFragment>();
-	BuildContext.AddFragmentWithDefaultInitializer<FMassSimulationLODFragment>();
-	BuildContext.AddTag<FMassOffLODTag>();
-	BuildContext.AddChunkFragment<FMassSimulationVariableTickChunkFragment>();
-
-	if (bUseZoneGraphMovement)
-	{
-		// UMassLaneCacheBoundaryProcessor
-		BuildContext.AddFragmentWithDefaultInitializer<FMassZoneGraphCachedLaneFragment>();
-		BuildContext.AddFragmentWithDefaultInitializer<FMassMoveTargetFragment>();
-		BuildContext.AddFragmentWithDefaultInitializer<FMassZoneGraphLaneLocationFragment>();
-		BuildContext.AddFragmentWithDefaultInitializer<FMassNavigationEdgesFragment>();
-		BuildContext.AddFragmentWithDefaultInitializer<FMassLaneCacheBoundaryFragment>();
-	}
-	else
-	{
-		// Put non-ZG movement specific fragments here
-		// UMassLaneBoundaryProcessor
-		BuildContext.AddFragmentWithDefaultInitializer<FDataFragment_Transform>();
-		BuildContext.AddFragmentWithDefaultInitializer<FMassNavigationEdgesFragment>();
-		BuildContext.AddFragmentWithDefaultInitializer<FMassLastUpdatePositionFragment>();
-		BuildContext.AddFragmentWithDefaultInitializer<FMassZoneGraphLaneLocationFragment>();
-		BuildContext.AddFragmentWithDefaultInitializer<FMassAvoidanceBoundaryLastLaneHandleFragment>();
-	}
+	BuildContext.AddFragmentWithDefaultInitializer<FMassMoveTargetFragment>();
 }
diff --git a/Engine/Plugins/AI/MassAI/Source/MassAIMovement/Public/MassAvoidanceTrait.h b/Engine/Plugins/AI/MassAI/Source/MassAIMovement/Public/MassAvoidanceTrait.h
--- a/Engine/Plugins/AI/MassAI/Source/MassAIMovement/Public/MassAvoidanceTrait.h
+++ b/Engine/Plugins/AI/MassAI/Source/MassAIMovement/Public/MassAvoidanceTrait.h
@@ -12,7 +12,4 @@ class MASSAIMOVEMENT_API UMassAvoidanceTrait : public UMassEntityTraitBase
 
 protected:
 	virtual void BuildTemplate(FMassEntityTemplateBuildContext& BuildContext, UWorld& World) const override;
-
-	UPROPERTY(Category="Movement", EditAnywhere)
-	bool bUseZoneGraphMovement = false;
 };
diff --git a/Engine/Plugins/Runtime/MassGameplay/Source/MassLOD/MassLOD.Build.cs b/Engine/Plugins/Runtime/MassGameplay/Source/MassLOD/MassLOD.Build.cs
--- a/Engine/Plugins/Runtime/MassGameplay/Source/MassLOD/MassLOD.Build.cs
+++ b/Engine/Plugins/Runtime/MassGameplay/Source/MassLOD/MassLOD.Build.cs
@@ -31,6 +31,7 @@ namespace UnrealBuildTool.Rules
 					"StructUtils",
 					"MassCommon",
 					"MassSimulation",
+					"MassSpawner"
 				}
 			);
 
diff --git a/Engine/Plugins/Runtime/MassGameplay/Source/MassLOD/Private/MassLODTrait.cpp b/Engine/Plugins/Runtime/MassGameplay/Source/MassLOD/Private/MassLODTrait.cpp
new file mode 100644
--- /dev/null
+++ b/Engine/Plugins/Runtime/MassGameplay/Source/MassLOD/Private/MassLODTrait.cpp
@@ -0,0 +1,18 @@
+// Copyright Epic Games, Inc. All Rights Reserved.
+
+#include "MassLODTrait.h"
+#include "MassEntityTemplateRegistry.h"
+#include "Engine/World.h"
+#include "MassCommonFragments.h"
+#include "MassSimulationLOD.h"
+
+
+void UMassLODTrait::BuildTemplate(FMassEntityTemplateBuildContext& BuildContext, UWorld& World) const
+{
+	BuildContext.AddFragmentWithDefaultInitializer<FDataFragment_Transform>();
+	BuildContext.AddFragmentWithDefaultInitializer<FMassLODInfoFragment>();
+	BuildContext.AddFragmentWithDefaultInitializer<FMassSimulationLODFragment>();
+	BuildContext.AddFragmentWithDefaultInitializer<FDataFragment_MassSimulationLODInfo>();
+	BuildContext.AddTag<FMassOffLODTag>();
+	BuildContext.AddChunkFragment<FMassSimulationVariableTickChunkFragment>();
+}
diff --git a/Engine/Plugins/Runtime/MassGameplay/Source/MassLOD/Public/MassLODTrait.h b/Engine/Plugins/Runtime/MassGameplay/Source/MassLOD/Public/MassLODTrait.h
new file mode 100644
--- /dev/null
+++ b/Engine/Plugins/Runtime/MassGameplay/Source/MassLOD/Public/MassLODTrait.h
@@ -0,0 +1,15 @@
+// Copyright Epic Games, Inc. All Rights Reserved.
+
+#pragma once 
+
+#include "MassEntityTraitBase.h"
+#include "MassLODTrait.generated.h"
+
+UCLASS(BlueprintType, EditInlineNew, CollapseCategories, meta = (DisplayName = "LOD"))
+class MASSLOD_API UMassLODTrait : public UMassEntityTraitBase
+{
+	GENERATED_BODY()
+
+protected:
+	virtual void BuildTemplate(FMassEntityTemplateBuildContext& BuildContext, UWorld& World) const override;
+};
