Made UMassApplyVelocityMoveTargetProcessor's and UMassDynamicObstacleProcessor's FMassSimulationVariableTickChunkFragment use optional

diff --git a/Engine/Plugins/AI/MassAI/Source/MassAIMovement/Private/NavigationProcessor.cpp b/Engine/Plugins/AI/MassAI/Source/MassAIMovement/Private/NavigationProcessor.cpp
--- a/Engine/Plugins/AI/MassAI/Source/MassAIMovement/Private/NavigationProcessor.cpp
+++ b/Engine/Plugins/AI/MassAI/Source/MassAIMovement/Private/NavigationProcessor.cpp
@@ -96,7 +96,7 @@ void UMassApplyVelocityMoveTargetProcessor::ConfigureQueries()
 	LowResEntityQuery_Conditional.AddRequirement<FDataFragment_Transform>(EMassFragmentAccess::ReadWrite);
 	LowResEntityQuery_Conditional.AddRequirement<FMassMoveTargetFragment>(EMassFragmentAccess::ReadOnly);
 	LowResEntityQuery_Conditional.AddTagRequirement<FMassOffLODTag>(EMassFragmentPresence::All);
-	LowResEntityQuery_Conditional.AddChunkRequirement<FMassSimulationVariableTickChunkFragment>(EMassFragmentAccess::ReadOnly);
+	LowResEntityQuery_Conditional.AddChunkRequirement<FMassSimulationVariableTickChunkFragment>(EMassFragmentAccess::ReadOnly, EMassFragmentPresence::Optional);
 	LowResEntityQuery_Conditional.SetChunkFilter(&FMassSimulationVariableTickChunkFragment::ShouldTickChunkThisFrame);
 }
 
@@ -254,7 +254,7 @@ void UMassDynamicObstacleProcessor::ConfigureQueries()
 	EntityQuery_Conditional.AddRequirement<FDataFragment_Transform>(EMassFragmentAccess::ReadOnly);
 	EntityQuery_Conditional.AddRequirement<FDataFragment_AgentRadius>(EMassFragmentAccess::ReadOnly);
 	EntityQuery_Conditional.AddRequirement<FMassDynamicObstacleFragment>(EMassFragmentAccess::ReadWrite);
-	EntityQuery_Conditional.AddChunkRequirement<FMassSimulationVariableTickChunkFragment>(EMassFragmentAccess::ReadOnly);
+	EntityQuery_Conditional.AddChunkRequirement<FMassSimulationVariableTickChunkFragment>(EMassFragmentAccess::ReadOnly, EMassFragmentPresence::Optional);
 	EntityQuery_Conditional.SetChunkFilter(&FMassSimulationVariableTickChunkFragment::ShouldTickChunkThisFrame);
 }
 
