Misc type conversion warning fixes

diff --git a/Engine/Plugins/AI/MassAI/Source/MassAIBehavior/Private/Tasks/MassZoneGraphFindSmartObjectTarget.cpp b/Engine/Plugins/AI/MassAI/Source/MassAIBehavior/Private/Tasks/MassZoneGraphFindSmartObjectTarget.cpp
--- a/Engine/Plugins/AI/MassAI/Source/MassAIBehavior/Private/Tasks/MassZoneGraphFindSmartObjectTarget.cpp
+++ b/Engine/Plugins/AI/MassAI/Source/MassAIBehavior/Private/Tasks/MassZoneGraphFindSmartObjectTarget.cpp
@@ -55,8 +55,7 @@ EStateTreeRunStatus FMassZoneGraphFindSmartObjectTarget::EnterState(FStateTreeEx
 	TargetLocation.EndOfPathIntent = EMassMovementAction::Stand;
 	TargetLocation.EndOfPathPosition = SOUser.GetTargetLocation();
 	// Let's start moving toward the interaction a bit before the entry point.
-	constexpr float AnticipationDistance = 100.f;
-	TargetLocation.AnticipationDistance = AnticipationDistance;
+	TargetLocation.AnticipationDistance.Set(100.f);
 
 	if (SOAnnotations != nullptr)
 	{
diff --git a/Engine/Plugins/AI/MassAI/Source/MassAIBehavior/Public/Tasks/MassZoneGraphPathFollowTask.h b/Engine/Plugins/AI/MassAI/Source/MassAIBehavior/Public/Tasks/MassZoneGraphPathFollowTask.h
--- a/Engine/Plugins/AI/MassAI/Source/MassAIBehavior/Public/Tasks/MassZoneGraphPathFollowTask.h
+++ b/Engine/Plugins/AI/MassAI/Source/MassAIBehavior/Public/Tasks/MassZoneGraphPathFollowTask.h
@@ -3,6 +3,7 @@
 #pragma once
 
 #include "MassAIMovementTypes.h"
+#include "MassCommonTypes.h"
 #include "MassStateTreeTypes.h"
 #include "ZoneGraphTypes.h"
 #include "MassZoneGraphPathFollowTask.generated.h"
@@ -32,7 +33,7 @@ struct MASSAIBEHAVIOR_API FMassZoneGraphTargetLocation : public FStateTreeResult
 		bMoveReverse = false;
 		TargetDistance = 0.0f;
 		EndOfPathPosition.Reset();
-		AnticipationDistance = 50.0f;
+		AnticipationDistance.Set(50.0f);
 		EndOfPathIntent = EMassMovementAction::Move;
 	}
 
@@ -52,7 +53,7 @@ struct MASSAIBEHAVIOR_API FMassZoneGraphTargetLocation : public FStateTreeResult
 	TOptional<FVector> EndOfPathDirection;
 
 	/** If start or end of path is off-lane, the distance along the lane is pushed forward/back along the lane to make smoother transition. */
-	float AnticipationDistance = 50.0f;
+	FMassInt16Real AnticipationDistance = FMassInt16Real(50.0f);
 
 	/** True, if we're moving reverse along the lane. */
 	bool bMoveReverse = false;
diff --git a/Engine/Plugins/AI/MassAI/Source/MassAIDebug/Private/MassDebugStateTreeProcessor.cpp b/Engine/Plugins/AI/MassAI/Source/MassAIDebug/Private/MassDebugStateTreeProcessor.cpp
--- a/Engine/Plugins/AI/MassAI/Source/MassAIDebug/Private/MassDebugStateTreeProcessor.cpp
+++ b/Engine/Plugins/AI/MassAI/Source/MassAIDebug/Private/MassDebugStateTreeProcessor.cpp
@@ -95,7 +95,7 @@ void UMassDebugStateTreeProcessor::Execute(UMassEntitySubsystem& EntitySubsystem
 					const FStructView Storage = EntitySubsystem.GetFragmentDataStruct(Entity, StateTree->GetRuntimeStorageStruct());
 
 					// State
-					UE_VLOG_SEGMENT_THICK(this, LogStateTree, Log, Position, Position + FVector(0,0,50), EntityColor, 2.0f, TEXT("%s %s"),
+					UE_VLOG_SEGMENT_THICK(this, LogStateTree, Log, Position, Position + FVector(0,0,50), EntityColor, /*Thickness*/ 2, TEXT("%s %s"),
 						*Entity.DebugGetDescription(), *StateTreeContext.GetActiveStateName(Storage));
 				}
 			}
diff --git a/Engine/Plugins/AI/MassAI/Source/MassAIMovement/Private/MassZoneGraphMovementFragments.cpp b/Engine/Plugins/AI/MassAI/Source/MassAIMovement/Private/MassZoneGraphMovementFragments.cpp
--- a/Engine/Plugins/AI/MassAI/Source/MassAIMovement/Private/MassZoneGraphMovementFragments.cpp
+++ b/Engine/Plugins/AI/MassAI/Source/MassAIMovement/Private/MassZoneGraphMovementFragments.cpp
@@ -212,7 +212,7 @@ bool FMassZoneGraphShortPathFragment::RequestPath(const FMassZoneGraphCachedLane
 	if (bStartOffLane)
 	{
 		// The start point was off-lane, move the start location along the lane a bit further to have smoother connection.
-		const float StartForwardOffset = FMath::Clamp(Request.AnticipationDistance + StartLaneForwardOffset, 0.0f, Request.AnticipationDistance);
+		const float StartForwardOffset = FMath::Clamp(Request.AnticipationDistance.Get() + StartLaneForwardOffset, 0.0f, Request.AnticipationDistance.Get());
 		StartDistanceAlongPath += StartForwardOffset * TangentSign; // Not clamping this distance intentionally so that the halfway point and clamping later works correctly.
 	}
 
@@ -240,7 +240,7 @@ bool FMassZoneGraphShortPathFragment::RequestPath(const FMassZoneGraphCachedLane
 		EndLaneOffset = FMath::Clamp(EndLaneOffset, -DeflatedLaneRight, DeflatedLaneLeft);
 
 		// Move the end location along the lane a bit back to have smoother connection.
-		const float EndForwardOffset = FMath::Clamp(Request.AnticipationDistance - EndLaneForwardOffset, 0.0f, Request.AnticipationDistance);
+		const float EndForwardOffset = FMath::Clamp(Request.AnticipationDistance.Get() - EndLaneForwardOffset, 0.0f, Request.AnticipationDistance.Get());
 		EndDistanceAlongPath -= EndForwardOffset * TangentSign; // Not clamping this distance intentionally so that the halfway point and clamping later works correctly.
 	}
 
diff --git a/Engine/Plugins/AI/MassAI/Source/MassAIMovement/Public/MassZoneGraphMovementTypes.h b/Engine/Plugins/AI/MassAI/Source/MassAIMovement/Public/MassZoneGraphMovementTypes.h
--- a/Engine/Plugins/AI/MassAI/Source/MassAIMovement/Public/MassZoneGraphMovementTypes.h
+++ b/Engine/Plugins/AI/MassAI/Source/MassAIMovement/Public/MassZoneGraphMovementTypes.h
@@ -50,7 +50,7 @@ struct MASSAIMOVEMENT_API FZoneGraphShortPathRequest
 
 	/** If start or end of path is off-lane, the distance along the lane is pushed forward/back along the lane to make smoother transition. */
 	UPROPERTY(Transient)
-	uint16 AnticipationDistance = 50;
+	FMassInt16Real AnticipationDistance = FMassInt16Real(50.f);
 
 	UPROPERTY(Transient)
 	FMassInt16Real EndOfPathOffset = FMassInt16Real(0.0f);
diff --git a/Engine/Plugins/AI/MassAI/Source/MassAIMovementEditor/Private/MassMovementTestingActor.cpp b/Engine/Plugins/AI/MassAI/Source/MassAIMovementEditor/Private/MassMovementTestingActor.cpp
--- a/Engine/Plugins/AI/MassAI/Source/MassAIMovementEditor/Private/MassMovementTestingActor.cpp
+++ b/Engine/Plugins/AI/MassAI/Source/MassAIMovementEditor/Private/MassMovementTestingActor.cpp
@@ -150,7 +150,7 @@ void UMassMovementTestingComponent::UpdateTests()
 			{
 				PathRequest.EndOfPathPosition = GoalWorldPosition;
 			}
-			PathRequest.AnticipationDistance = AnticipationDistance;
+			PathRequest.AnticipationDistance.Set(AnticipationDistance);
 
 			PathRequest.StartPosition = WorldPosition;
 			float StartDistanceAlongPath = LaneLocation.DistanceAlongLane;
diff --git a/Engine/Plugins/AI/MassCrowd/Source/MassCrowd/Private/Tasks/MassCrowdClaimWaitSlotTask.cpp b/Engine/Plugins/AI/MassCrowd/Source/MassCrowd/Private/Tasks/MassCrowdClaimWaitSlotTask.cpp
--- a/Engine/Plugins/AI/MassCrowd/Source/MassCrowd/Private/Tasks/MassCrowdClaimWaitSlotTask.cpp
+++ b/Engine/Plugins/AI/MassCrowd/Source/MassCrowd/Private/Tasks/MassCrowdClaimWaitSlotTask.cpp
@@ -52,8 +52,7 @@ EStateTreeRunStatus FMassCrowdClaimWaitSlotTask::EnterState(FStateTreeExecutionC
 	TargetLocation.EndOfPathDirection = SlotDirection;
 	TargetLocation.TargetDistance = LaneLocation.LaneLength; // Go to end of lane
 	// Let's start moving toward the interaction a bit before the entry point.
-	constexpr float AnticipationDistance = 100.f;
-	TargetLocation.AnticipationDistance = AnticipationDistance;
+	TargetLocation.AnticipationDistance.Set(100.f);
 
 	WaitSlotLocation = &TargetLocation;
 	
diff --git a/Engine/Plugins/Runtime/MassEntity/Source/MassEntity/Private/MassArchetypeData.cpp b/Engine/Plugins/Runtime/MassEntity/Source/MassEntity/Private/MassArchetypeData.cpp
--- a/Engine/Plugins/Runtime/MassEntity/Source/MassEntity/Private/MassArchetypeData.cpp
+++ b/Engine/Plugins/Runtime/MassEntity/Source/MassEntity/Private/MassArchetypeData.cpp
@@ -634,7 +634,7 @@ void FMassArchetypeData::CompactEntities(const double TimeAllowed)
 		ChunkToFill->AddMultipleInstances(NumberOfEntitiesToMove);
 		ChunkToEmpty->RemoveMultipleInstances(NumberOfEntitiesToMove);
 
-		const int32 ChunkToFillIdx = ChunkToFill - &Chunks[0];
+		const int32 ChunkToFillIdx = UE_PTRDIFF_TO_INT32(ChunkToFill - &Chunks[0]);
 		check(ChunkToFillIdx >=0 && ChunkToFillIdx < Chunks.Num());
 		const int32 AbsoluteIndex = ChunkToFillIdx * NumEntitiesPerChunk + ToIndex;
 
diff --git a/Engine/Source/Runtime/AIModule/Public/GraphAStar.h b/Engine/Source/Runtime/AIModule/Public/GraphAStar.h
--- a/Engine/Source/Runtime/AIModule/Public/GraphAStar.h
+++ b/Engine/Source/Runtime/AIModule/Public/GraphAStar.h
@@ -187,7 +187,7 @@ struct FGraphAStar
 		FORCEINLINE FSearchNode& Add(const FSearchNode& SearchNode)
 		{
 			FSearchNode& NewNode = Super::Emplace_GetRef(SearchNode);
-			NewNode.SearchNodeIndex = &NewNode - Super::GetData();
+			NewNode.SearchNodeIndex = UE_PTRDIFF_TO_INT32(&NewNode - Super::GetData());
 			NodeMap.Add(SearchNode.NodeRef, NewNode.SearchNodeIndex);
 			return NewNode;
 		}
@@ -204,7 +204,7 @@ struct FGraphAStar
 
 			// node not found, add it and setup index in node map
 			FSearchNode& NewNode = Super::Emplace_GetRef(NodeRef);
-			NewNode.SearchNodeIndex = &NewNode - Super::GetData();
+			NewNode.SearchNodeIndex = UE_PTRDIFF_TO_INT32(&NewNode - Super::GetData());
 			Index = NewNode.SearchNodeIndex;
 
 			return NewNode;
@@ -255,7 +255,7 @@ struct FGraphAStar
 			{
 				if (NodeIndex == SearchNode.SearchNodeIndex)
 				{
-					AlgoImpl::HeapSiftUp(Super::GetData(), 0, &NodeIndex - Super::GetData(), FIdentityFunctor(), NodeSorter);
+					AlgoImpl::HeapSiftUp(Super::GetData(), 0, UE_PTRDIFF_TO_INT32(&NodeIndex - Super::GetData()), FIdentityFunctor(), NodeSorter);
 					return;
 				}
 			}
