Misc type conversion warning trivial fixes

diff --git a/Engine/Plugins/AI/MassAI/Source/MassAIBehavior/Private/MassLookAtProcessors.cpp b/Engine/Plugins/AI/MassAI/Source/MassAIBehavior/Private/MassLookAtProcessors.cpp
--- a/Engine/Plugins/AI/MassAI/Source/MassAIBehavior/Private/MassLookAtProcessors.cpp
+++ b/Engine/Plugins/AI/MassAI/Source/MassAIBehavior/Private/MassLookAtProcessors.cpp
@@ -343,7 +343,7 @@ void UMassLookAtProcessor::UpdateLookAtTrajectory(const FTransform& Transform, c
 		if (bDisplayDebug)
 		{
 			const FVector ZOffset(0.f,0.f,DebugZOffset);
-			UE_VLOG_SEGMENT_THICK(this, LogMassBehavior, Display, AgentPosition + ZOffset, AgentPosition + ZOffset + 100.f * NewGlobalDirection, FColor::White, /*Thickness*/3.0f, TEXT("LookAt Trajectory"));
+			UE_VLOG_SEGMENT_THICK(this, LogMassBehavior, Display, AgentPosition + ZOffset, AgentPosition + ZOffset + 100.f * NewGlobalDirection, FColor::White, /*Thickness*/3, TEXT("LookAt Trajectory"));
 		}
 #endif
 	}
@@ -364,7 +364,7 @@ void UMassLookAtProcessor::UpdateLookAtTrackedEntity(const UMassEntitySubsystem&
 			if (bDisplayDebug)
 			{
 				const FVector ZOffset(0.f,0.f,DebugZOffset);
-				UE_VLOG_SEGMENT_THICK(this, LogMassBehavior, Display, AgentPosition + ZOffset, AgentPosition + ZOffset + 100.f * NewGlobalDirection, FColor::White, /*Thickness*/3.0f, TEXT("LookAt Track"));
+				UE_VLOG_SEGMENT_THICK(this, LogMassBehavior, Display, AgentPosition + ZOffset, AgentPosition + ZOffset + 100.f * NewGlobalDirection, FColor::White, /*Thickness*/3, TEXT("LookAt Track"));
 			}
 #endif
 		}
@@ -390,7 +390,7 @@ bool UMassLookAtProcessor::UpdateGazeTrackedEntity(const UMassEntitySubsystem& E
 			if (bDisplayDebug)
 			{
 				const FVector ZOffset(0.f,0.f,DebugZOffset);
-				UE_VLOG_SEGMENT_THICK(this, LogMassBehavior, Display, AgentPosition + ZOffset, AgentPosition + ZOffset + 100.f * NewGlobalDirection, FColor(160,160,160), /*Thickness*/3.0f, TEXT("Gaze Track"));
+				UE_VLOG_SEGMENT_THICK(this, LogMassBehavior, Display, AgentPosition + ZOffset, AgentPosition + ZOffset + 100.f * NewGlobalDirection, FColor(160,160,160), /*Thickness*/3, TEXT("Gaze Track"));
 			}
 #endif
 		}
@@ -503,7 +503,7 @@ void UMassLookAtProcessor::BuildTrajectory(const FMassZoneGraphLaneLocationFragm
 			{
 				const float T = (j + 1) * DeltaT;
 				const FVector Point = UE::CubicBezier::Eval(StartPoint, StartControlPoint, EndControlPoint, EndPoint, T);
-				UE_VLOG_SEGMENT_THICK(this, LogMassBehavior, Display, PrevPoint + ZOffset, Point + ZOffset, FColor::White, /*Thickness*/3.0f, TEXT(""));
+				UE_VLOG_SEGMENT_THICK(this, LogMassBehavior, Display, PrevPoint + ZOffset, Point + ZOffset, FColor::White, /*Thickness*/3, TEXT(""));
 				PrevPoint = Point;
 			}
 		}
@@ -513,7 +513,7 @@ void UMassLookAtProcessor::BuildTrajectory(const FMassZoneGraphLaneLocationFragm
 			const FMassLookAtTrajectoryPoint& CurrPoint = LookAtTrajectory.Points[PointIndex];
 			const FVector CurrBase = CurrPoint.Position + ZOffset * 1.1f;
 			// Tangents
-			UE_VLOG_SEGMENT_THICK(this, LogMassBehavior, Display, CurrBase, CurrBase + CurrPoint.Tangent.GetVector() * 100.0f, FColorList::Grey, /*Thickness*/1.0f,
+			UE_VLOG_SEGMENT_THICK(this, LogMassBehavior, Display, CurrBase, CurrBase + CurrPoint.Tangent.GetVector() * 100.0f, FColorList::Grey, /*Thickness*/1,
 				TEXT("D:%.1f"), CurrPoint.DistanceAlongLane.Get());
 		}
 	}
diff --git a/Engine/Plugins/AI/MassAI/Source/MassAIBehavior/Private/Tasks/MassZoneGraphFindEscapeTarget.cpp b/Engine/Plugins/AI/MassAI/Source/MassAIBehavior/Private/Tasks/MassZoneGraphFindEscapeTarget.cpp
--- a/Engine/Plugins/AI/MassAI/Source/MassAIBehavior/Private/Tasks/MassZoneGraphFindEscapeTarget.cpp
+++ b/Engine/Plugins/AI/MassAI/Source/MassAIBehavior/Private/Tasks/MassZoneGraphFindEscapeTarget.cpp
@@ -68,7 +68,7 @@ EStateTreeRunStatus FMassZoneGraphFindEscapeTarget::EnterState(FStateTreeExecuti
 
 	const uint8 SpanIndex = EscapeAction->FindSpanIndex(LaneLocation.DistanceAlongLane);
 	const FZoneGraphEscapeLaneSpan& EscapeSpan = EscapeAction->Spans[SpanIndex];
-	const float MoveDir = EscapeSpan.bReverseLaneDirection ? -1 : 1;	
+	const float MoveDir = EscapeSpan.bReverseLaneDirection ? -1.f : 1.f;
 
 	constexpr float AdjacentMoveDistance = 50.0f;
 	constexpr float MoveDistanceRandomDeviation = 250.0f;
@@ -91,7 +91,7 @@ EStateTreeRunStatus FMassZoneGraphFindEscapeTarget::EnterState(FStateTreeExecuti
 		
 		// Small move, and goto adjacent lane
 		TargetLocation.LaneHandle = LaneLocation.LaneHandle;
-		TargetLocation.TargetDistance = FMath::Clamp(LaneLocation.DistanceAlongLane + AdjacentMoveDistance * MoveDir, 0, LaneLocation.LaneLength);
+		TargetLocation.TargetDistance = FMath::Clamp(LaneLocation.DistanceAlongLane + AdjacentMoveDistance * MoveDir, 0.f, LaneLocation.LaneLength);
 		TargetLocation.NextExitLinkType = EZoneLaneLinkType::Adjacent;
 		TargetLocation.NextLaneHandle = FZoneGraphLaneHandle(EscapeSpan.ExitLaneIndex, ZoneGraphStorage->DataHandle);
 		TargetLocation.bMoveReverse = EscapeSpan.bReverseLaneDirection;
diff --git a/Engine/Plugins/AI/MassAI/Source/MassAIMovement/Private/MassAvoidanceProcessors.cpp b/Engine/Plugins/AI/MassAI/Source/MassAIMovement/Private/MassAvoidanceProcessors.cpp
--- a/Engine/Plugins/AI/MassAI/Source/MassAIMovement/Private/MassAvoidanceProcessors.cpp
+++ b/Engine/Plugins/AI/MassAI/Source/MassAIMovement/Private/MassAvoidanceProcessors.cpp
@@ -238,7 +238,7 @@ namespace UE::MassAvoidance
 			return;
 		}
 
-		UE_VLOG_SEGMENT_THICK(Context.LogOwner, Context.Category, Log, Start, End, Color, Thickness, TEXT(""));
+		UE_VLOG_SEGMENT_THICK(Context.LogOwner, Context.Category, Log, Start, End, Color, (int16)Thickness, TEXT(""));
 
 		if (UseDrawDebugHelper() && Context.World)
 		{
@@ -259,9 +259,9 @@ namespace UE::MassAvoidance
 		const FVector Perp = FVector::CrossProduct(UnitV, FVector::UpVector);
 		const FVector Left = Perp - (Pointyness*UnitV);
 		const FVector Right = -Perp - (Pointyness*UnitV);
-		UE_VLOG_SEGMENT_THICK(Context.LogOwner, Context.Category, Log, Start, End, Color, Thickness, TEXT(""));
-		UE_VLOG_SEGMENT_THICK(Context.LogOwner, Context.Category, Log, End, End + HeadSize * Left, Color, Thickness, TEXT(""));
-		UE_VLOG_SEGMENT_THICK(Context.LogOwner, Context.Category, Log, End, End + HeadSize * Right, Color, Thickness, TEXT(""));
+		UE_VLOG_SEGMENT_THICK(Context.LogOwner, Context.Category, Log, Start, End, Color, (int16)Thickness, TEXT(""));
+		UE_VLOG_SEGMENT_THICK(Context.LogOwner, Context.Category, Log, End, End + HeadSize * Left, Color, (int16)Thickness, TEXT(""));
+		UE_VLOG_SEGMENT_THICK(Context.LogOwner, Context.Category, Log, End, End + HeadSize * Right, Color, (int16)Thickness, TEXT(""));
 
 		if (UseDrawDebugHelper() && Context.World)
 		{
@@ -337,10 +337,10 @@ namespace UE::MassAvoidance
 		const FVector Left = Perp - (Pointyness * UnitV);
 		const FVector Right = -Perp - (Pointyness * UnitV);
 		const float HeadSize = 0.08f * Line.Size();
-		UE_VLOG_SEGMENT_THICK(Context.LogOwner, Context.Category, Log, Start, End, Color, Thickness, TEXT(""));
-		UE_VLOG_SEGMENT_THICK(Context.LogOwner, Context.Category, Log, End, End + HeadSize * Left, Color, Thickness, TEXT(""));
-		UE_VLOG_SEGMENT_THICK(Context.LogOwner, Context.Category, Log, End, End + HeadSize * Right, Color, Thickness, TEXT(""));
-		UE_VLOG_SEGMENT_THICK(Context.LogOwner, Context.Category, Log, End + HeadSize * Left, End + HeadSize * Right, Color, Thickness, TEXT(""));
+		UE_VLOG_SEGMENT_THICK(Context.LogOwner, Context.Category, Log, Start, End, Color, (int16)Thickness, TEXT(""));
+		UE_VLOG_SEGMENT_THICK(Context.LogOwner, Context.Category, Log, End, End + HeadSize * Left, Color, (int16)Thickness, TEXT(""));
+		UE_VLOG_SEGMENT_THICK(Context.LogOwner, Context.Category, Log, End, End + HeadSize * Right, Color, (int16)Thickness, TEXT(""));
+		UE_VLOG_SEGMENT_THICK(Context.LogOwner, Context.Category, Log, End + HeadSize * Left, End + HeadSize * Right, Color, (int16)Thickness, TEXT(""));
 
 		if (UseDrawDebugHelper() && Context.World)
 		{
@@ -1761,7 +1761,7 @@ void UMassLaneBoundaryProcessor::Execute(UMassEntitySubsystem& EntitySubsystem,
 						}
 
 						const FZoneLaneData& Lane = Storage.Lanes[LaneLocation.LaneHandle.Index];
-						const float HalfWidth = 0.5*Lane.Width;
+						const float HalfWidth = 0.5f*Lane.Width;
 						const int32 Segment = LaneLocation.LaneSegment;
 
 						static const int32 MaxPoints = 4;
@@ -1976,7 +1976,7 @@ void UMassLaneCacheBoundaryProcessor::Execute(UMassEntitySubsystem& EntitySubsys
 			}
 #endif // WITH_MASSGAMEPLAY_DEBUG
 
-			const float HalfWidth = 0.5 * CachedLane.LaneWidth.Get();
+			const float HalfWidth = 0.5f * CachedLane.LaneWidth.Get();
 
 			static const int32 MaxPoints = 4;
 			FVector Points[MaxPoints];
diff --git a/Engine/Plugins/AI/MassAI/Source/MassAIMovement/Private/MassZoneGraphMovementProcessors.cpp b/Engine/Plugins/AI/MassAI/Source/MassAIMovement/Private/MassZoneGraphMovementProcessors.cpp
--- a/Engine/Plugins/AI/MassAI/Source/MassAIMovement/Private/MassZoneGraphMovementProcessors.cpp
+++ b/Engine/Plugins/AI/MassAI/Source/MassAIMovement/Private/MassZoneGraphMovementProcessors.cpp
@@ -424,7 +424,7 @@ void UMassZoneGraphPathFollowProcessor::Execute(UMassEntitySubsystem& EntitySubs
 						const FMassZoneGraphPathPoint& NextPoint = ShortPath.Points[PointIndex + 1];
 
 						// Path
-						UE_VLOG_SEGMENT_THICK(this, LogMassNavigation, Display, CurrPoint.Position + ZOffset, NextPoint.Position + ZOffset, EntityColor, /*Thickness*/3.0f, TEXT(""));
+						UE_VLOG_SEGMENT_THICK(this, LogMassNavigation, Display, CurrPoint.Position + ZOffset, NextPoint.Position + ZOffset, EntityColor, /*Thickness*/3, TEXT(""));
 					}
 					
 					for (uint8 PointIndex = 0; PointIndex < ShortPath.NumPoints; PointIndex++)
@@ -432,14 +432,14 @@ void UMassZoneGraphPathFollowProcessor::Execute(UMassEntitySubsystem& EntitySubs
 						const FMassZoneGraphPathPoint& CurrPoint = ShortPath.Points[PointIndex];
 						const FVector CurrBase = CurrPoint.Position + ZOffset;
 						// Lane tangents
-						UE_VLOG_SEGMENT_THICK(this, LogMassNavigation, Display, CurrBase, CurrBase + CurrPoint.Tangent.GetVector() * 100.0f, LightEntityColor, /*Thickness*/1.0f, TEXT(""));
+						UE_VLOG_SEGMENT_THICK(this, LogMassNavigation, Display, CurrBase, CurrBase + CurrPoint.Tangent.GetVector() * 100.0f, LightEntityColor, /*Thickness*/1, TEXT(""));
 					}
 
 					if (ShortPath.NumPoints > 0 && ShortPath.NextLaneHandle.IsValid())
 					{
 						const FMassZoneGraphPathPoint& LastPoint = ShortPath.Points[ShortPath.NumPoints - 1];
 						const FVector CurrBase = LastPoint.Position + ZOffset;
-						UE_VLOG_SEGMENT_THICK(this, LogMassNavigation, Display, CurrBase, CurrBase + FVector(0,0,100), FColor::Red, /*Thickness*/3.0f, TEXT("Next: %s"), *ShortPath.NextLaneHandle.ToString());
+						UE_VLOG_SEGMENT_THICK(this, LogMassNavigation, Display, CurrBase, CurrBase + FVector(0,0,100), FColor::Red, /*Thickness*/3, TEXT("Next: %s"), *ShortPath.NextLaneHandle.ToString());
 					}
 				}
 #endif // WITH_MASSGAMEPLAY_DEBUG
@@ -722,24 +722,24 @@ void UMassZoneGraphSteeringProcessor::Execute(UMassEntitySubsystem& EntitySubsys
 				const FVector MoveTargetCenter = MoveTarget.Center + ZOffset;
 
 				// MoveTarget slack boundary
-				UE_VLOG_CIRCLE_THICK(this, LogMassNavigation, Log, MoveTargetCenter, FVector::UpVector, CurrentMovementConfig->Steering.LookAheadDistance, EntityColor, 2,
+				UE_VLOG_CIRCLE_THICK(this, LogMassNavigation, Log, MoveTargetCenter, FVector::UpVector, CurrentMovementConfig->Steering.LookAheadDistance, EntityColor, /*Thickness*/2,
 					TEXT("%s MoveTgt %s"), *Entity.DebugGetDescription(), *UEnum::GetDisplayValueAsText(MoveTarget.IntentAtGoal).ToString());
 
 				// MoveTarget orientation
-				UE_VLOG_SEGMENT_THICK(this, LogMassNavigation, Log, MoveTargetCenter, MoveTargetCenter + MoveTarget.Forward * CurrentMovementConfig->Steering.LookAheadDistance, EntityColor, 2.0f, TEXT(""));
+				UE_VLOG_SEGMENT_THICK(this, LogMassNavigation, Log, MoveTargetCenter, MoveTargetCenter + MoveTarget.Forward * CurrentMovementConfig->Steering.LookAheadDistance, EntityColor, /*Thickness*/2, TEXT(""));
 
 				// MoveTarget - current location relation.
 				if (FVector::Dist2D(CurrentLocation, MoveTarget.Center) > CurrentMovementConfig->Steering.LookAheadDistance * 1.5f)
 				{
-					UE_VLOG_SEGMENT_THICK(this, LogMassNavigation, Log, MoveTargetCenter, CurrentLocation + ZOffset, FColor::Red, 1.0f, TEXT("LOST"));
+					UE_VLOG_SEGMENT_THICK(this, LogMassNavigation, Log, MoveTargetCenter, CurrentLocation + ZOffset, FColor::Red, /*Thickness*/1, TEXT("LOST"));
 				}
 				else
 				{
-					UE_VLOG_SEGMENT_THICK(this, LogMassNavigation, Log, MoveTargetCenter, CurrentLocation + ZOffset, DarkEntityColor, 1.0f, TEXT(""));
+					UE_VLOG_SEGMENT_THICK(this, LogMassNavigation, Log, MoveTargetCenter, CurrentLocation + ZOffset, DarkEntityColor, /*Thickness*/1, TEXT(""));
 				}
 
 				// Steering
-				UE_VLOG_SEGMENT_THICK(this, LogMassNavigation, Log, CurrentLocation + ZOffset, CurrentLocation + Steering.DesiredVelocity + ZOffset, LightEntityColor, 2.0f,
+				UE_VLOG_SEGMENT_THICK(this, LogMassNavigation, Log, CurrentLocation + ZOffset, CurrentLocation + Steering.DesiredVelocity + ZOffset, LightEntityColor, /*Thickness*/2,
 					TEXT("%s Steer %.1f"), *Entity.DebugGetDescription(), Steering.DesiredVelocity.Length());
 			}
 #endif // WITH_MASSGAMEPLAY_DEBUG
diff --git a/Engine/Plugins/AI/MassAI/Source/MassAIMovement/Private/NavigationProcessor.cpp b/Engine/Plugins/AI/MassAI/Source/MassAIMovement/Private/NavigationProcessor.cpp
--- a/Engine/Plugins/AI/MassAI/Source/MassAIMovement/Private/NavigationProcessor.cpp
+++ b/Engine/Plugins/AI/MassAI/Source/MassAIMovement/Private/NavigationProcessor.cpp
@@ -187,7 +187,7 @@ void UMassApplyVelocityMoveTargetProcessor::Execute(UMassEntitySubsystem& Entity
 					if (MoveTarget.IntentAtGoal == EMassMovementAction::Stand && MoveTarget.DistanceToGoal < OrientationEndOfPathHeadingAnticipation)
 					{
 						// Fade towards the movement target direction at the end of the path.
-						const float Strength = 1.0f - FMath::Square(FMath::Clamp(MoveTarget.DistanceToGoal / OrientationEndOfPathHeadingAnticipation, 0.0f, 1.0));
+						const float Strength = 1.0f - FMath::Square(FMath::Clamp(MoveTarget.DistanceToGoal / OrientationEndOfPathHeadingAnticipation, 0.0f, 1.0f));
 						Blend = FMath::Lerp(Blend, OrientationBlendWhileStanding, Strength);
 					}
 					else
diff --git a/Engine/Plugins/AI/MassAI/Source/MassAIMovement/Public/MassAIMovementFragments.h b/Engine/Plugins/AI/MassAI/Source/MassAIMovement/Public/MassAIMovementFragments.h
--- a/Engine/Plugins/AI/MassAI/Source/MassAIMovement/Public/MassAIMovementFragments.h
+++ b/Engine/Plugins/AI/MassAI/Source/MassAIMovement/Public/MassAIMovementFragments.h
@@ -227,7 +227,7 @@ struct MASSAIMOVEMENT_API FMassDynamicObstacleFragment : public FMassFragment
 	TArray<FMassLaneObstacleID, TFixedAllocator<MaxLaneObstacles>> LaneObstacleIDs;
 
 	/** Time stamp when that obstacle stopped moving. */
-	FDateTime LastMovedTimeStamp = 0.f;
+	FDateTime LastMovedTimeStamp;
 
 	/** Position of the dynamic obstacle when it last moved. */
 	FVector LastPosition = FVector::ZeroVector;
diff --git a/Engine/Plugins/Runtime/MassGameplay/Source/MassRepresentation/Private/MassVisualizationComponent.cpp b/Engine/Plugins/Runtime/MassGameplay/Source/MassRepresentation/Private/MassVisualizationComponent.cpp
--- a/Engine/Plugins/Runtime/MassGameplay/Source/MassRepresentation/Private/MassVisualizationComponent.cpp
+++ b/Engine/Plugins/Runtime/MassGameplay/Source/MassRepresentation/Private/MassVisualizationComponent.cpp
@@ -81,7 +81,7 @@ void UMassVisualizationComponent::ConstructStaticMeshComponents()
 						ISMC->SetMaterial(ElementIndex, MaterialOverride);
 					}
 				}
-				ISMC->SetCullDistances(0.0f, 1000000.0); // @todo: Need to figure out what to do here, either LOD or cull distances.
+				ISMC->SetCullDistances(0, 1000000); // @todo: Need to figure out what to do here, either LOD or cull distances.
 				ISMC->SetupAttachment(ActorOwner->GetRootComponent());
 				ISMC->SetCanEverAffectNavigation(false);
 				ISMC->SetCollisionProfileName(UCollisionProfile::NoCollision_ProfileName);
diff --git a/Engine/Plugins/Runtime/MassGameplay/Source/MassSpawner/Private/MassEntityZoneGraphSpawnPointsGenerator.cpp b/Engine/Plugins/Runtime/MassGameplay/Source/MassSpawner/Private/MassEntityZoneGraphSpawnPointsGenerator.cpp
--- a/Engine/Plugins/Runtime/MassGameplay/Source/MassSpawner/Private/MassEntityZoneGraphSpawnPointsGenerator.cpp
+++ b/Engine/Plugins/Runtime/MassGameplay/Source/MassSpawner/Private/MassEntityZoneGraphSpawnPointsGenerator.cpp
@@ -72,7 +72,7 @@ void UMassEntityZoneGraphSpawnPointsGenerator::GeneratePointsForZoneGraphData(co
 				FZoneGraphLaneLocation LaneLocation;
 				UE::ZoneGraph::Query::CalculateLocationAlongLane(ZoneGraphStorage, LaneIndex, Distance, LaneLocation);
 				const FVector Perp = LaneLocation.Direction ^ LaneLocation.Up;
-				Locations.Add(LaneLocation.Position + Perp * RandomStream.RandRange(-LaneHalfWidth, LaneHalfWidth));
+				Locations.Add(LaneLocation.Position + Perp * RandomStream.FRandRange(-LaneHalfWidth, LaneHalfWidth));
 
 				// Advance ahead past the space we just consumed, plus a random gap.
 				Distance += RandomStream.FRandRange(MinGap, MaxGap);
diff --git a/Engine/Plugins/Runtime/MassGameplay/Source/MassSpawner/Private/MassSpawner.cpp b/Engine/Plugins/Runtime/MassGameplay/Source/MassSpawner/Private/MassSpawner.cpp
--- a/Engine/Plugins/Runtime/MassGameplay/Source/MassSpawner/Private/MassSpawner.cpp
+++ b/Engine/Plugins/Runtime/MassGameplay/Source/MassSpawner/Private/MassSpawner.cpp
@@ -406,7 +406,7 @@ void AMassSpawner::SpawnAtLocations(const TArray<FVector>& Locations)
 	{
 		if (FVisualLogEntry* LogEntry = FVisualLogger::Get().GetLastEntryForObject(this))
 		{
-			FVisualLogShapeElement Element(TEXT(""), FColor::Orange, 20.f, LogMassSpawner.GetCategoryName());
+			FVisualLogShapeElement Element(TEXT(""), FColor::Orange, /*Thickness*/20, LogMassSpawner.GetCategoryName());
 			Element.Points += Locations;
 			Element.Type = EVisualLoggerShapeElement::SinglePoint;
 			Element.Verbosity = ELogVerbosity::Display;
diff --git a/Engine/Plugins/Runtime/StateTree/Source/StateTreeEditorModule/Private/Customizations/StateTreeBindableItemDetails.cpp b/Engine/Plugins/Runtime/StateTree/Source/StateTreeEditorModule/Private/Customizations/StateTreeBindableItemDetails.cpp
--- a/Engine/Plugins/Runtime/StateTree/Source/StateTreeEditorModule/Private/Customizations/StateTreeBindableItemDetails.cpp
+++ b/Engine/Plugins/Runtime/StateTree/Source/StateTreeEditorModule/Private/Customizations/StateTreeBindableItemDetails.cpp
@@ -188,7 +188,7 @@ void FStateTreeBindableItemDetails::CustomizeHeader(TSharedRef<class IPropertyHa
 			[
 				SAssignNew(ComboButton, SComboButton)
 				.OnGetMenuContent(this, &FStateTreeBindableItemDetails::GenerateStructPicker)
-				.ContentPadding(0)
+				.ContentPadding(0.f)
 				.ButtonContent()
 				[
 					SNew(SHorizontalBox)
@@ -468,12 +468,12 @@ TSharedRef<SWidget> FStateTreeBindableItemDetails::GenerateStructPicker()
 	FOnStructPicked OnPicked(FOnStructPicked::CreateRaw(this, &FStateTreeBindableItemDetails::OnStructPicked));
 
 	return SNew(SBox)
-		.WidthOverride(280)
+		.WidthOverride(280.f)
 		[
 			SNew(SVerticalBox)
 			+ SVerticalBox::Slot()
 			.AutoHeight()
-			.MaxHeight(500)
+			.MaxHeight(500.f)
 			[
 				FModuleManager::LoadModuleChecked<FStructViewerModule>("StructViewer").CreateStructViewer(Options, OnPicked)
 			]
diff --git a/Engine/Plugins/Runtime/StateTree/Source/StateTreeModule/Private/StateTreeExecutionContext.cpp b/Engine/Plugins/Runtime/StateTree/Source/StateTreeModule/Private/StateTreeExecutionContext.cpp
--- a/Engine/Plugins/Runtime/StateTree/Source/StateTreeModule/Private/StateTreeExecutionContext.cpp
+++ b/Engine/Plugins/Runtime/StateTree/Source/StateTreeModule/Private/StateTreeExecutionContext.cpp
@@ -625,7 +625,7 @@ FStateTreeTransitionResult FStateTreeExecutionContext::TriggerTransitions(FState
 	for (FStateTreeHandle Handle = CurrentStatus.State; Handle.IsValid(); Handle = StateTree->States[Handle.Index].Parent)
 	{
 		const FBakedStateTreeState& State = StateTree->States[Handle.Index];
-		for (uint32 i = 0; i < State.TransitionsNum; i++)
+		for (uint8 i = 0; i < State.TransitionsNum; i++)
 		{
 			// All transition conditions must pass
 			const int16 TransitionIndex = State.TransitionsBegin + i;
diff --git a/Engine/Plugins/Runtime/ZoneGraphAnnotations/Source/ZoneGraphAnnotations/Private/Annotations/ZoneGraphDisturbanceAnnotation.cpp b/Engine/Plugins/Runtime/ZoneGraphAnnotations/Source/ZoneGraphAnnotations/Private/Annotations/ZoneGraphDisturbanceAnnotation.cpp
--- a/Engine/Plugins/Runtime/ZoneGraphAnnotations/Source/ZoneGraphAnnotations/Private/Annotations/ZoneGraphDisturbanceAnnotation.cpp
+++ b/Engine/Plugins/Runtime/ZoneGraphAnnotations/Source/ZoneGraphAnnotations/Private/Annotations/ZoneGraphDisturbanceAnnotation.cpp
@@ -184,7 +184,7 @@ void UZoneGraphDisturbanceAnnotation::CalculateEscapeGraph(FZoneGraphDataEscapeG
 		UE::ZoneGraph::Query::GetLaneLength(*ZoneStorage, LaneHandle, LaneLength);
 
 		// Split the lane up to FZoneGraphEscapeLaneAction::MaxSpans to have more fine grained Disturbance Annotation.
-		EscapeLane.SpanCount = FMath::Clamp((uint8)FMath::CeilToInt(LaneLength * InvIdealSpanLength), 2, FZoneGraphEscapeLaneAction::MaxSpans);
+		EscapeLane.SpanCount = FMath::Clamp((uint8)FMath::CeilToInt(LaneLength * InvIdealSpanLength), (uint8)2, FZoneGraphEscapeLaneAction::MaxSpans);
 		
 		FVector PrevPosition = ZoneStorage->LanePoints[Lane.PointsBegin];
 		float PrevDistance = 0.0f;
