More to make SimulationLOD optional

#rnx

diff --git a/Engine/Plugins/AI/MassAI/Source/MassAIBehavior/Private/MassStateTreeProcessors.cpp b/Engine/Plugins/AI/MassAI/Source/MassAIBehavior/Private/MassStateTreeProcessors.cpp
--- a/Engine/Plugins/AI/MassAI/Source/MassAIBehavior/Private/MassStateTreeProcessors.cpp
+++ b/Engine/Plugins/AI/MassAI/Source/MassAIBehavior/Private/MassStateTreeProcessors.cpp
@@ -229,7 +229,7 @@ void UMassStateTreeActivationProcessor::ConfigureQueries()
 {
 	EntityQuery.AddRequirement<FMassStateTreeFragment>(EMassFragmentAccess::ReadOnly);
 	EntityQuery.AddTagRequirement<FMassStateTreeActivated>(EMassFragmentPresence::None);
-	EntityQuery.AddChunkRequirement<FMassSimulationVariableTickChunkFragment>(EMassFragmentAccess::ReadOnly);
+	EntityQuery.AddChunkRequirement<FMassSimulationVariableTickChunkFragment>(EMassFragmentAccess::ReadOnly, EMassFragmentPresence::Optional);
 }
 void UMassStateTreeActivationProcessor::Execute(UMassEntitySubsystem& EntitySubsystem, FMassExecutionContext& Context)
 {
diff --git a/Engine/Plugins/AI/MassAI/Source/MassAIBehavior/Private/MassZoneGraphAnnotationProcessors.cpp b/Engine/Plugins/AI/MassAI/Source/MassAIBehavior/Private/MassZoneGraphAnnotationProcessors.cpp
--- a/Engine/Plugins/AI/MassAI/Source/MassAIBehavior/Private/MassZoneGraphAnnotationProcessors.cpp
+++ b/Engine/Plugins/AI/MassAI/Source/MassAIBehavior/Private/MassZoneGraphAnnotationProcessors.cpp
@@ -80,7 +80,7 @@ void UMassZoneGraphAnnotationTagUpdateProcessor::ConfigureQueries()
 	EntityQuery.AddRequirement<FMassZoneGraphAnnotationTagsFragment>(EMassFragmentAccess::ReadWrite);
 	EntityQuery.AddRequirement<FMassZoneGraphLaneLocationFragment>(EMassFragmentAccess::ReadOnly);
 	EntityQuery.AddChunkRequirement<FMassZoneGraphAnnotationVariableTickChunkFragment>(EMassFragmentAccess::ReadOnly);
-	EntityQuery.AddChunkRequirement<FMassSimulationVariableTickChunkFragment>(EMassFragmentAccess::ReadOnly);
+	EntityQuery.AddChunkRequirement<FMassSimulationVariableTickChunkFragment>(EMassFragmentAccess::ReadOnly, EMassFragmentPresence::Optional);
 }
 
 void UMassZoneGraphAnnotationTagUpdateProcessor::Execute(UMassEntitySubsystem& EntitySubsystem, FMassExecutionContext& Context)
diff --git a/Engine/Plugins/AI/MassAI/Source/MassAIDebug/Private/GameplayDebuggerCategory_Mass.cpp b/Engine/Plugins/AI/MassAI/Source/MassAIDebug/Private/GameplayDebuggerCategory_Mass.cpp
--- a/Engine/Plugins/AI/MassAI/Source/MassAIDebug/Private/GameplayDebuggerCategory_Mass.cpp
+++ b/Engine/Plugins/AI/MassAI/Source/MassAIDebug/Private/GameplayDebuggerCategory_Mass.cpp
@@ -306,7 +306,7 @@ void FGameplayDebuggerCategory_Mass::CollectData(APlayerController* OwnerPC, AAc
 		EntityQuery.AddRequirement<FMassMoveTargetFragment>(EMassFragmentAccess::ReadOnly);
 		EntityQuery.AddRequirement<FMassLookAtFragment>(EMassFragmentAccess::ReadOnly);
 		EntityQuery.AddRequirement<FMassLookAtTrajectoryFragment>(EMassFragmentAccess::ReadOnly);
-		EntityQuery.AddRequirement<FMassSimulationLODFragment>(EMassFragmentAccess::ReadOnly);
+		EntityQuery.AddRequirement<FMassSimulationLODFragment>(EMassFragmentAccess::ReadOnly, EMassFragmentPresence::Optional);
 		EntityQuery.AddRequirement<FMassZoneGraphShortPathFragment>(EMassFragmentAccess::ReadOnly);
 
 		const float CurrentTime = World->GetTimeSeconds();
@@ -329,6 +329,7 @@ void FGameplayDebuggerCategory_Mass::CollectData(APlayerController* OwnerPC, AAc
 				const TConstArrayView<FMassLookAtFragment> LookAtList = Context.GetFragmentView<FMassLookAtFragment>();
 				const TConstArrayView<FMassLookAtTrajectoryFragment> LookAtTrajectoryList = Context.GetFragmentView<FMassLookAtTrajectoryFragment>();
 				const TConstArrayView<FMassSimulationLODFragment> SimLODList = Context.GetFragmentView<FMassSimulationLODFragment>();
+				const bool bHasLOD = (SimLODList.Num() > 0);
 				const TConstArrayView<FMassZoneGraphShortPathFragment> ShortPathList = Context.GetFragmentView<FMassZoneGraphShortPathFragment>();
 
 				constexpr float MaxViewDistance = 25000.0f;
@@ -360,7 +361,7 @@ void FGameplayDebuggerCategory_Mass::CollectData(APlayerController* OwnerPC, AAc
 					const FMassVelocityFragment& Velocity = VelocityList[EntityIndex];
 					const FMassMoveTargetFragment& MoveTarget = MoveTargetList[EntityIndex];
 					const FMassLookAtFragment& LookAt = LookAtList[EntityIndex];
-					const FMassSimulationLODFragment& SimLOD = SimLODList[EntityIndex];
+					const FMassSimulationLODFragment& SimLOD = bHasLOD ? SimLODList[EntityIndex] : FMassSimulationLODFragment();
 					const FMassZoneGraphShortPathFragment& ShortPath = ShortPathList[EntityIndex];
 
 					const FVector EntityForward = Transform.GetTransform().GetRotation().GetForwardVector();
diff --git a/Engine/Plugins/AI/MassAI/Source/MassAIMovement/Private/MassAvoidanceProcessors.cpp b/Engine/Plugins/AI/MassAI/Source/MassAIMovement/Private/MassAvoidanceProcessors.cpp
--- a/Engine/Plugins/AI/MassAI/Source/MassAIMovement/Private/MassAvoidanceProcessors.cpp
+++ b/Engine/Plugins/AI/MassAI/Source/MassAIMovement/Private/MassAvoidanceProcessors.cpp
@@ -399,7 +399,7 @@ void UMassAvoidanceProcessor::ConfigureQueries()
 {
 	EntityQuery.AddRequirement<FMassSteeringFragment>(EMassFragmentAccess::ReadWrite);
 	EntityQuery.AddRequirement<FMassNavigationEdgesFragment>(EMassFragmentAccess::ReadOnly);
-	EntityQuery.AddRequirement<FMassSimulationLODFragment>(EMassFragmentAccess::ReadOnly);
+	EntityQuery.AddRequirement<FMassSimulationLODFragment>(EMassFragmentAccess::ReadOnly, EMassFragmentPresence::Optional);
 	EntityQuery.AddRequirement<FMassMoveTargetFragment>(EMassFragmentAccess::ReadOnly);
 	EntityQuery.AddRequirement<FDataFragment_Transform>(EMassFragmentAccess::ReadOnly);
 	EntityQuery.AddRequirement<FMassVelocityFragment>(EMassFragmentAccess::ReadOnly);
@@ -504,6 +504,7 @@ void UMassAvoidanceProcessor::Execute(UMassEntitySubsystem& EntitySubsystem, FMa
 		const TConstArrayView<FMassVelocityFragment> VelocityList = Ctx.GetFragmentView<FMassVelocityFragment>();
 		const TConstArrayView<FDataFragment_AgentRadius> RadiusList = Ctx.GetFragmentView<FDataFragment_AgentRadius>();
 		const TConstArrayView<FMassSimulationLODFragment> SimLODList = Ctx.GetFragmentView<FMassSimulationLODFragment>();
+		const bool bHasLOD = (SimLODList.Num() > 0);
 		const TConstArrayView<FMassMoveTargetFragment> MoveTargetList = Ctx.GetFragmentView<FMassMoveTargetFragment>();
 		const TConstArrayView<FMassMovementConfigFragment> MovementConfigList = Ctx.GetFragmentView<FMassMovementConfigFragment>();
 
@@ -578,7 +579,6 @@ void UMassAvoidanceProcessor::Execute(UMassEntitySubsystem& EntitySubsystem, FMa
 			const FDataFragment_Transform& Location = LocationList[EntityIndex];
 			const FMassVelocityFragment& Velocity = VelocityList[EntityIndex];
 			const FDataFragment_AgentRadius& Radius = RadiusList[EntityIndex];
-			const FMassSimulationLODFragment& SimLOD = SimLODList[EntityIndex];
 
 			// Smaller steering max accel makes the steering more "calm" but less opportunistic, may not find solution, or gets stuck.
 			// Max contact accel should be quite a big bigger than steering so that collision response is firm. 
@@ -607,7 +607,7 @@ void UMassAvoidanceProcessor::Execute(UMassEntitySubsystem& EntitySubsystem, FMa
 				const FString Text = FString::Printf(TEXT("%i"), Entity.Index);
 				DebugDrawCylinder(BaseDebugContext, AgentLocation, AgentLocation + DebugAgentHeightOffset, AgentRadius+1.f, CurrentAgentColor, Text);
 
-				const FColor AgentColor = (SimLOD.LOD == EMassLOD::High) ? CurrentAgentColor : FColor::Red;
+				const FColor AgentColor = (!bHasLOD || SimLODList[EntityIndex].LOD == EMassLOD::High) ? CurrentAgentColor : FColor::Red;
 				DebugDrawSphere(BaseDebugContext, AgentLocation, 10.f, AgentColor);
 
 				// Draw current velocity (black)
@@ -1611,7 +1611,6 @@ UMassLaneBoundaryProcessor::UMassLaneBoundaryProcessor()
 void UMassLaneBoundaryProcessor::ConfigureQueries()
 {
 	EntityQuery.AddRequirement<FDataFragment_Transform>(EMassFragmentAccess::ReadOnly);							// need agent position to get closest point on lane
-	EntityQuery.AddRequirement<FMassSimulationLODFragment>(EMassFragmentAccess::ReadOnly);						// LOD level
 	EntityQuery.AddRequirement<FMassNavigationEdgesFragment>(EMassFragmentAccess::ReadWrite);					// output edges
 	EntityQuery.AddRequirement<FMassLastUpdatePositionFragment>(EMassFragmentAccess::ReadWrite);					// to keep position when boundaries where last updated
 	EntityQuery.AddRequirement<FMassZoneGraphLaneLocationFragment>(EMassFragmentAccess::ReadOnly);				// current lane location
@@ -1891,7 +1890,6 @@ void UMassLaneCacheBoundaryProcessor::ConfigureQueries()
 	EntityQuery.AddRequirement<FMassZoneGraphCachedLaneFragment>(EMassFragmentAccess::ReadOnly);
 	EntityQuery.AddRequirement<FMassMoveTargetFragment>(EMassFragmentAccess::ReadOnly);
 	EntityQuery.AddRequirement<FMassZoneGraphLaneLocationFragment>(EMassFragmentAccess::ReadOnly);
-	EntityQuery.AddRequirement<FMassSimulationLODFragment>(EMassFragmentAccess::ReadOnly);
 	EntityQuery.AddRequirement<FMassLaneCacheBoundaryFragment>(EMassFragmentAccess::ReadWrite);
 	EntityQuery.AddRequirement<FMassNavigationEdgesFragment>(EMassFragmentAccess::ReadWrite);	// output edges
 	EntityQuery.AddTagRequirement<FMassOffLODTag>(EMassFragmentPresence::None);
diff --git a/Engine/Plugins/Runtime/MassGameplay/Source/MassRepresentation/Private/MassVisualizationTrait.cpp b/Engine/Plugins/Runtime/MassGameplay/Source/MassRepresentation/Private/MassVisualizationTrait.cpp
--- a/Engine/Plugins/Runtime/MassGameplay/Source/MassRepresentation/Private/MassVisualizationTrait.cpp
+++ b/Engine/Plugins/Runtime/MassGameplay/Source/MassRepresentation/Private/MassVisualizationTrait.cpp
@@ -36,6 +36,7 @@ void UMassVisualizationTrait::BuildTemplate(FMassEntityTemplateBuildContext& Bui
 
 		BuildContext.AddDeinitializer(*(RepresentationFragmentDeinitializerClass->GetDefaultObject<UMassProcessor>()));
 
+		BuildContext.AddFragment<FMassLODInfoFragment>();
 		BuildContext.AddFragment<FMassRepresentationLODFragment>();
 		BuildContext.AddTag<FMassVisibilityCulledByDistanceTag>();
 
