StateTree: allow task/evaluator/condition to have separate data for isntance and item in the tree.

- Changed conditions to use the same code structure as tasks and evals
- Tasks, Evaluators, and Conditions can have separate struct for instance and tree item
- Converted existing tasks/evals/conds to use the new setup
- Updated the UI to use the new setup
- Limited the property binding to only first level properties in the instance struct
- Removed StateTreeConditionItemDetails (uses the same details customization as tasks/evals)
- Removed bunch of old test tasks and evals

diff --git a/Engine/Plugins/AI/MassAI/Source/MassAIBehavior/Private/Conditions/ZoneGraphTagConditions.cpp b/Engine/Plugins/AI/MassAI/Source/MassAIBehavior/Private/Conditions/ZoneGraphTagConditions.cpp
--- a/Engine/Plugins/AI/MassAI/Source/MassAIBehavior/Private/Conditions/ZoneGraphTagConditions.cpp
+++ b/Engine/Plugins/AI/MassAI/Source/MassAIBehavior/Private/Conditions/ZoneGraphTagConditions.cpp
@@ -1,10 +1,13 @@
 ﻿// Copyright Epic Games, Inc. All Rights Reserved.
 
 #include "Conditions/ZoneGraphTagConditions.h"
+#include "StateTreeExecutionContext.h"
 #if WITH_EDITOR
 #include "ZoneGraphSettings.h"
 #include "StateTreePropertyBindings.h"
+#endif// WITH_EDITOR
 
+#if WITH_EDITOR
 #define LOCTEXT_NAMESPACE "StateTreeEditor"
 
 namespace UE::MassBehavior::ZoneGraph
@@ -74,9 +77,30 @@ namespace UE::MassBehavior::ZoneGraph
 
 }
 
-FText FZoneGraphTagFilterCondition::GetDescription(const IStateTreeBindingLookup& BindingLookup) const
+#endif// WITH_EDITOR
+
+
+//----------------------------------------------------------------------//
+//  FZoneGraphTagFilterCondition
+//----------------------------------------------------------------------//
+
+bool FZoneGraphTagFilterCondition::Link(FStateTreeLinker& Linker)
+{
+	Linker.LinkInstanceDataProperty(TagsHandle, STATETREE_INSTANCEDATA_PROPERTY(FZoneGraphTagFilterConditionInstanceData, Tags));
+
+	return true;
+}
+
+bool FZoneGraphTagFilterCondition::TestCondition(FStateTreeExecutionContext& Context) const
 {
-	const FStateTreeEditorPropertyPath LeftPath(ID, GET_MEMBER_NAME_STRING_CHECKED(FZoneGraphTagFilterCondition, Tags));
+	const FZoneGraphTagMask Tags = Context.GetInstanceData(TagsHandle);
+	return Filter.Pass(Tags) ^ bInvert;
+}
+
+#if WITH_EDITOR
+FText FZoneGraphTagFilterCondition::GetDescription(const FGuid& ID, FConstStructView InstanceData, const IStateTreeBindingLookup& BindingLookup) const
+{
+	const FStateTreeEditorPropertyPath LeftPath(ID, GET_MEMBER_NAME_STRING_CHECKED(FZoneGraphTagFilterConditionInstanceData, Tags));
 
 	FText InvertText;
 	if (bInvert)
@@ -91,7 +115,7 @@ FText FZoneGraphTagFilterCondition::GetDescription(const IStateTreeBindingLookup
 	}
 	else
 	{
-		LeftText = UE::MassBehavior::ZoneGraph::GetTagMaskName(Tags);
+		LeftText = LOCTEXT("NotBound", "Not Bound");
 	}
 
 	FText FilterParts[6] = { FText::GetEmpty(), FText::GetEmpty(), FText::GetEmpty(), FText::GetEmpty(), FText::GetEmpty(), FText::GetEmpty() };
@@ -114,14 +138,37 @@ FText FZoneGraphTagFilterCondition::GetDescription(const IStateTreeBindingLookup
 		FilterParts[PartIndex++] = UE::MassBehavior::ZoneGraph::GetTagMaskName(Filter.NotTags);
 	}
 
-	return FText::Format(LOCTEXT("CompareZoneGraphTagFilterDesc", "{0} <Details.Bold>{1}</> contains {2} <Details.Bold>{3}</> {4} <Details.Bold>{5}</> {6} <Details.Bold>{7}</>"), InvertText, LeftText, FilterParts[0], FilterParts[1], FilterParts[2], FilterParts[3], FilterParts[4], FilterParts[5]);
+	return FText::Format(LOCTEXT("CompareZoneGraphTagFilterDesc", "{0} <Details.Bold>{1}</> contains {2} <Details.Bold>{3}</> {4} <Details.Bold>{5}</> {6} <Details.Bold>{7}</>"),
+		InvertText, LeftText, FilterParts[0], FilterParts[1], FilterParts[2], FilterParts[3], FilterParts[4], FilterParts[5]);
 }
+#endif// WITH_EDITOR
 
 
-FText FZoneGraphTagMaskCondition::GetDescription(const IStateTreeBindingLookup& BindingLookup) const
+//----------------------------------------------------------------------//
+//  FZoneGraphTagMaskCondition
+//----------------------------------------------------------------------//
+
+bool FZoneGraphTagMaskCondition::Link(FStateTreeLinker& Linker)
 {
-	const FStateTreeEditorPropertyPath LeftPath(ID, GET_MEMBER_NAME_STRING_CHECKED(FZoneGraphTagMaskCondition, Left));
-	const FStateTreeEditorPropertyPath RightPath(ID, GET_MEMBER_NAME_STRING_CHECKED(FZoneGraphTagMaskCondition, Right));
+	Linker.LinkInstanceDataProperty(LeftHandle, STATETREE_INSTANCEDATA_PROPERTY(FZoneGraphTagMaskConditionInstanceData, Left));
+	Linker.LinkInstanceDataProperty(RightHandle, STATETREE_INSTANCEDATA_PROPERTY(FZoneGraphTagMaskConditionInstanceData, Right));
+
+	return true;
+}
+
+bool FZoneGraphTagMaskCondition::TestCondition(FStateTreeExecutionContext& Context) const
+{
+	const FZoneGraphTagMask Left = Context.GetInstanceData(LeftHandle);
+	const FZoneGraphTagMask Right = Context.GetInstanceData(RightHandle);
+	return Left.CompareMasks(Right, Operator) ^ bInvert;
+}
+
+#if WITH_EDITOR
+FText FZoneGraphTagMaskCondition::GetDescription(const FGuid& ID, FConstStructView InstanceData, const IStateTreeBindingLookup& BindingLookup) const
+{
+	const FZoneGraphTagMaskConditionInstanceData& Instance = InstanceData.Get<FZoneGraphTagMaskConditionInstanceData>();
+	const FStateTreeEditorPropertyPath LeftPath(ID, GET_MEMBER_NAME_STRING_CHECKED(FZoneGraphTagMaskConditionInstanceData, Left));
+	const FStateTreeEditorPropertyPath RightPath(ID, GET_MEMBER_NAME_STRING_CHECKED(FZoneGraphTagMaskConditionInstanceData, Right));
 
 	FText InvertText;
 	if (bInvert)
@@ -136,7 +183,7 @@ FText FZoneGraphTagMaskCondition::GetDescription(const IStateTreeBindingLookup&
 	}
 	else
 	{
-		LeftText = UE::MassBehavior::ZoneGraph::GetTagMaskName(Left);
+		LeftText = LOCTEXT("NotBound", "Not Bound");
 	}
 
 	FText OperatorText = UE::MassBehavior::ZoneGraph::GetMaskOperatorText(Operator);
@@ -148,17 +195,40 @@ FText FZoneGraphTagMaskCondition::GetDescription(const IStateTreeBindingLookup&
 	}
 	else
 	{
-		RightText = UE::MassBehavior::ZoneGraph::GetTagMaskName(Right);
+		RightText = UE::MassBehavior::ZoneGraph::GetTagMaskName(Instance.Right);
 	}
 
-	return FText::Format(LOCTEXT("CompareZoneGraphTagMaskDesc", "{0} <Details.Bold>{1}</> contains {2} <Details.Bold>{3}</>"), InvertText, LeftText, OperatorText, RightText);
+	return FText::Format(LOCTEXT("CompareZoneGraphTagMaskDesc", "{0} <Details.Bold>{1}</> contains {2} <Details.Bold>{3}</>"),
+		InvertText, LeftText, OperatorText, RightText);
 }
+#endif// WITH_EDITOR
+
 
+//----------------------------------------------------------------------//
+//  FZoneGraphTagCondition
+//----------------------------------------------------------------------//
 
-FText FZoneGraphTagCondition::GetDescription(const IStateTreeBindingLookup& BindingLookup) const
+bool FZoneGraphTagCondition::Link(FStateTreeLinker& Linker)
 {
-	const FStateTreeEditorPropertyPath LeftPath(ID, GET_MEMBER_NAME_STRING_CHECKED(FZoneGraphTagCondition, Left));
-	const FStateTreeEditorPropertyPath RightPath(ID, GET_MEMBER_NAME_STRING_CHECKED(FZoneGraphTagCondition, Right));
+	Linker.LinkInstanceDataProperty(LeftHandle, STATETREE_INSTANCEDATA_PROPERTY(FZoneGraphTagConditionInstanceData, Left));
+	Linker.LinkInstanceDataProperty(RightHandle, STATETREE_INSTANCEDATA_PROPERTY(FZoneGraphTagConditionInstanceData, Right));
+
+	return true;
+}
+
+bool FZoneGraphTagCondition::TestCondition(FStateTreeExecutionContext& Context) const
+{
+	const FZoneGraphTag Left = Context.GetInstanceData(LeftHandle);
+	const FZoneGraphTag Right = Context.GetInstanceData(RightHandle);
+	return (Left == Right) ^ bInvert;
+}
+
+#if WITH_EDITOR
+FText FZoneGraphTagCondition::GetDescription(const FGuid& ID, FConstStructView InstanceData, const IStateTreeBindingLookup& BindingLookup) const
+{
+	const FZoneGraphTagConditionInstanceData& Instance = InstanceData.Get<FZoneGraphTagConditionInstanceData>();
+	const FStateTreeEditorPropertyPath LeftPath(ID, GET_MEMBER_NAME_STRING_CHECKED(FZoneGraphTagConditionInstanceData, Left));
+	const FStateTreeEditorPropertyPath RightPath(ID, GET_MEMBER_NAME_STRING_CHECKED(FZoneGraphTagConditionInstanceData, Right));
 
 	FText InvertText;
 	if (bInvert)
@@ -173,7 +243,7 @@ FText FZoneGraphTagCondition::GetDescription(const IStateTreeBindingLookup& Bind
 	}
 	else
 	{
-		LeftText = UE::MassBehavior::ZoneGraph::GetTagName(Left);
+		LeftText = LOCTEXT("NotBound", "Not Bound");
 	}
 
 	FText RightText;
@@ -183,11 +253,14 @@ FText FZoneGraphTagCondition::GetDescription(const IStateTreeBindingLookup& Bind
 	}
 	else
 	{
-		RightText = UE::MassBehavior::ZoneGraph::GetTagName(Right);
+		RightText = UE::MassBehavior::ZoneGraph::GetTagName(Instance.Right);
 	}
 
-	return FText::Format(LOCTEXT("CompareZoneGraphTagDesc", "{0} <Details.Bold>{1}</> is <Details.Bold>{3}</>"), InvertText, LeftText, RightText);
+	return FText::Format(LOCTEXT("CompareZoneGraphTagDesc", "{0} <Details.Bold>{1}</> is <Details.Bold>{2}</>"),
+		InvertText, LeftText, RightText);
 }
+#endif// WITH_EDITOR
 
+#if WITH_EDITOR
 #undef LOCTEXT_NAMESPACE
 #endif // WITH_EDITOR
diff --git a/Engine/Plugins/AI/MassAI/Source/MassAIBehavior/Private/Evaluators/MassComponentHitEvaluator.cpp b/Engine/Plugins/AI/MassAI/Source/MassAIBehavior/Private/Evaluators/MassComponentHitEvaluator.cpp
--- a/Engine/Plugins/AI/MassAI/Source/MassAIBehavior/Private/Evaluators/MassComponentHitEvaluator.cpp
+++ b/Engine/Plugins/AI/MassAI/Source/MassAIBehavior/Private/Evaluators/MassComponentHitEvaluator.cpp
@@ -9,17 +9,23 @@
 
 bool FMassComponentHitEvaluator::Link(FStateTreeLinker& Linker)
 {
-	Linker.LinkExternalItem(ComponentHitSubsystemHandle);
+	Linker.LinkExternalData(ComponentHitSubsystemHandle);
+
+	Linker.LinkInstanceDataProperty(GotHitHandle, STATETREE_INSTANCEDATA_PROPERTY(FMassComponentHitEvaluatorInstanceData, bGotHit));
+	Linker.LinkInstanceDataProperty(LastHitEntityHandle, STATETREE_INSTANCEDATA_PROPERTY(FMassComponentHitEvaluatorInstanceData, LastHitEntity));
 
 	return true;
 }
 
-void FMassComponentHitEvaluator::Evaluate(FStateTreeExecutionContext &Context, const EStateTreeEvaluationType EvalType, const float DeltaTime)
+void FMassComponentHitEvaluator::Evaluate(FStateTreeExecutionContext &Context, const EStateTreeEvaluationType EvalType, const float DeltaTime) const
 {
 	// Look for recent hits
-	UMassComponentHitSubsystem& HitSubsystem = Context.GetExternalItem(ComponentHitSubsystemHandle);
+	UMassComponentHitSubsystem& HitSubsystem = Context.GetExternalData(ComponentHitSubsystemHandle);
 	const FMassHitResult* HitResult = HitSubsystem.GetLastHit(static_cast<FMassStateTreeExecutionContext&>(Context).GetEntity());
 
+	bool& bGotHit = Context.GetInstanceData(GotHitHandle);
+	FMassEntityHandle& LastHitEntity = Context.GetInstanceData(LastHitEntityHandle);
+	
 	// LastHitEntity is not reset intentionally, so that it's available the duration of the behavior reacting to it.
 	bGotHit = false;
 
diff --git a/Engine/Plugins/AI/MassAI/Source/MassAIBehavior/Private/Evaluators/MassStateTreeSmartObjectEvaluator.cpp b/Engine/Plugins/AI/MassAI/Source/MassAIBehavior/Private/Evaluators/MassStateTreeSmartObjectEvaluator.cpp
--- a/Engine/Plugins/AI/MassAI/Source/MassAIBehavior/Private/Evaluators/MassStateTreeSmartObjectEvaluator.cpp
+++ b/Engine/Plugins/AI/MassAI/Source/MassAIBehavior/Private/Evaluators/MassStateTreeSmartObjectEvaluator.cpp
@@ -20,53 +20,70 @@
 
 bool FMassStateTreeSmartObjectEvaluator::Link(FStateTreeLinker& Linker)
 {
-	Linker.LinkExternalItem(SmartObjectSubsystemHandle);
-	Linker.LinkExternalItem(MassSignalSubsystemHandle);
-	Linker.LinkExternalItem(EntityTransformHandle);
-	Linker.LinkExternalItem(SmartObjectUserHandle);
-	Linker.LinkExternalItem(LocationHandle);
+	Linker.LinkExternalData(SmartObjectSubsystemHandle);
+	Linker.LinkExternalData(MassSignalSubsystemHandle);
+	Linker.LinkExternalData(EntityTransformHandle);
+	Linker.LinkExternalData(SmartObjectUserHandle);
+	Linker.LinkExternalData(LocationHandle);
+
+	Linker.LinkInstanceDataProperty(SearchRequestResultHandle, STATETREE_INSTANCEDATA_PROPERTY(FMassStateTreeSmartObjectEvaluatorInstanceData, SearchRequestResult));
+	Linker.LinkInstanceDataProperty(SearchRequestIDHandle, STATETREE_INSTANCEDATA_PROPERTY(FMassStateTreeSmartObjectEvaluatorInstanceData, SearchRequestID));
+	Linker.LinkInstanceDataProperty(CandidatesFoundHandle, STATETREE_INSTANCEDATA_PROPERTY(FMassStateTreeSmartObjectEvaluatorInstanceData, bCandidatesFound));
+	Linker.LinkInstanceDataProperty(ClaimedHandle, STATETREE_INSTANCEDATA_PROPERTY(FMassStateTreeSmartObjectEvaluatorInstanceData, bClaimed));
+	Linker.LinkInstanceDataProperty(NextUpdateHandle, STATETREE_INSTANCEDATA_PROPERTY(FMassStateTreeSmartObjectEvaluatorInstanceData, NextUpdate));
+	Linker.LinkInstanceDataProperty(UsingZoneGraphAnnotationsHandle, STATETREE_INSTANCEDATA_PROPERTY(FMassStateTreeSmartObjectEvaluatorInstanceData, bUsingZoneGraphAnnotations));
 
 	return true;
 }
 
-void FMassStateTreeSmartObjectEvaluator::EnterState(FStateTreeExecutionContext& Context, const EStateTreeStateChangeType ChangeType, const FStateTreeTransitionResult& Transition)
+void FMassStateTreeSmartObjectEvaluator::EnterState(FStateTreeExecutionContext& Context, const EStateTreeStateChangeType ChangeType, const FStateTreeTransitionResult& Transition) const
 {
 	if (ChangeType != EStateTreeStateChangeType::Changed)
 	{
 		return;
 	}
 
-	Reset();
+	Reset(Context);
 }
 
-void FMassStateTreeSmartObjectEvaluator::ExitState(FStateTreeExecutionContext& Context, const EStateTreeStateChangeType ChangeType, const FStateTreeTransitionResult& Transition)
+void FMassStateTreeSmartObjectEvaluator::ExitState(FStateTreeExecutionContext& Context, const EStateTreeStateChangeType ChangeType, const FStateTreeTransitionResult& Transition) const
 {
 	if (ChangeType != EStateTreeStateChangeType::Changed)
 	{
 		return;
 	}
 
+	FMassSmartObjectRequestID& SearchRequestID = Context.GetInstanceData(SearchRequestIDHandle);
+
 	if (SearchRequestID.IsSet())
 	{
 		const FMassStateTreeExecutionContext& MassContext = static_cast<FMassStateTreeExecutionContext&>(Context);
-		USmartObjectSubsystem& SmartObjectSubsystem = Context.GetExternalItem(SmartObjectSubsystemHandle);
+		USmartObjectSubsystem& SmartObjectSubsystem = Context.GetExternalData(SmartObjectSubsystemHandle);
 		const FMassSmartObjectHandler MassSmartObjectHandler(MassContext.GetEntitySubsystem(), MassContext.GetEntitySubsystemExecutionContext(), SmartObjectSubsystem);
 		MassSmartObjectHandler.RemoveRequest(SearchRequestID);
 		SearchRequestID.Reset();
 	}
-	Reset();
+	Reset(Context);
 }
 
-void FMassStateTreeSmartObjectEvaluator::Reset()
+void FMassStateTreeSmartObjectEvaluator::Reset(FStateTreeExecutionContext& Context) const
 {
+	bool& bCandidatesFound = Context.GetInstanceData(CandidatesFoundHandle);
+	bool& bClaimed = Context.GetInstanceData(ClaimedHandle);
+	FMassSmartObjectRequestID& SearchRequestID = Context.GetInstanceData(SearchRequestIDHandle);
+
 	bCandidatesFound = false;
 	bClaimed = false;
 	SearchRequestID.Reset();
 }
 
-void FMassStateTreeSmartObjectEvaluator::Evaluate(FStateTreeExecutionContext& Context, const EStateTreeEvaluationType EvalType,  const float DeltaTime)
+void FMassStateTreeSmartObjectEvaluator::Evaluate(FStateTreeExecutionContext& Context, const EStateTreeEvaluationType EvalType,  const float DeltaTime) const
 {
-	const FMassSmartObjectUserFragment& SOUser = Context.GetExternalItem(SmartObjectUserHandle);
+	const FMassSmartObjectUserFragment& SOUser = Context.GetExternalData(SmartObjectUserHandle);
+
+	bool& bCandidatesFound = Context.GetInstanceData(CandidatesFoundHandle);
+	bool& bClaimed = Context.GetInstanceData(ClaimedHandle);
+	FMassSmartObjectRequestID& SearchRequestID = Context.GetInstanceData(SearchRequestIDHandle);
 
 	bCandidatesFound = false;
 	bClaimed = SOUser.GetClaimHandle().IsValid();
@@ -84,22 +101,24 @@ void FMassStateTreeSmartObjectEvaluator::Evaluate(FStateTreeExecutionContext& Co
 	}
 
 	// We need to track our next update cooldown since we can get ticked from any signals waking up the StateTree
+	float& NextUpdate = Context.GetInstanceData(NextUpdateHandle);
 	if (NextUpdate > World->GetTimeSeconds())
 	{
 		return;
 	}
 	NextUpdate = 0.f;
 
-	USmartObjectSubsystem& SmartObjectSubsystem = Context.GetExternalItem(SmartObjectSubsystemHandle);
+	USmartObjectSubsystem& SmartObjectSubsystem = Context.GetExternalData(SmartObjectSubsystemHandle);
 	const FMassStateTreeExecutionContext& MassContext = static_cast<FMassStateTreeExecutionContext&>(Context);
 	const FMassSmartObjectHandler MassSmartObjectHandler(MassContext.GetEntitySubsystem(), MassContext.GetEntitySubsystemExecutionContext(), SmartObjectSubsystem);
 
 	// Nothing claimed -> search for candidates
+	bool& bUsingZoneGraphAnnotations = Context.GetInstanceData(UsingZoneGraphAnnotationsHandle);
 	if (!SearchRequestID.IsSet())
 	{
 		// Use lanes if possible for faster queries using zone graph annotations
 		const FMassEntityHandle RequestingEntity = MassContext.GetEntity();
-		const FMassZoneGraphLaneLocationFragment* LaneLocation = Context.GetExternalItemPtr(LocationHandle);
+		const FMassZoneGraphLaneLocationFragment* LaneLocation = Context.GetExternalDataPtr(LocationHandle);
 		bUsingZoneGraphAnnotations = LaneLocation != nullptr;
 		if (bUsingZoneGraphAnnotations)
 		{
@@ -111,13 +130,14 @@ void FMassStateTreeSmartObjectEvaluator::Evaluate(FStateTreeExecutionContext& Co
 		}
 		else
 		{
-			const FDataFragment_Transform& TransformFragment = Context.GetExternalItem(EntityTransformHandle);
+			const FDataFragment_Transform& TransformFragment = Context.GetExternalData(EntityTransformHandle);
 			SearchRequestID = MassSmartObjectHandler.FindCandidatesAsync(RequestingEntity, TransformFragment.GetTransform().GetLocation());
 		}
 	}
 	else
 	{
 		// Fetch request results
+		FMassSmartObjectRequestResult& SearchRequestResult = Context.GetInstanceData(SearchRequestResultHandle);
 		SearchRequestResult = MassSmartObjectHandler.GetRequestResult(SearchRequestID);
 
 		// Check if results are ready
@@ -139,7 +159,7 @@ void FMassStateTreeSmartObjectEvaluator::Evaluate(FStateTreeExecutionContext& Co
 			{
 				const float DelayInSeconds = bCandidatesFound ? TickInterval : RetryCooldown;
 				NextUpdate = World->GetTimeSeconds() + DelayInSeconds;
-				UMassSignalSubsystem& MassSignalSubsystem = Context.GetExternalItem(MassSignalSubsystemHandle);
+				UMassSignalSubsystem& MassSignalSubsystem = Context.GetExternalData(MassSignalSubsystemHandle);
 				MassSignalSubsystem.DelaySignalEntity(UE::Mass::Signals::SmartObjectRequestCandidates, RequestingEntity, DelayInSeconds);
 			}
 		}
diff --git a/Engine/Plugins/AI/MassAI/Source/MassAIBehavior/Private/Evaluators/MassZoneGraphAnnotationEvaluator.cpp b/Engine/Plugins/AI/MassAI/Source/MassAIBehavior/Private/Evaluators/MassZoneGraphAnnotationEvaluator.cpp
--- a/Engine/Plugins/AI/MassAI/Source/MassAIBehavior/Private/Evaluators/MassZoneGraphAnnotationEvaluator.cpp
+++ b/Engine/Plugins/AI/MassAI/Source/MassAIBehavior/Private/Evaluators/MassZoneGraphAnnotationEvaluator.cpp
@@ -11,13 +11,16 @@ FMassZoneGraphAnnotationEvaluator::FMassZoneGraphAnnotationEvaluator()
 
 bool FMassZoneGraphAnnotationEvaluator::Link(FStateTreeLinker& Linker)
 {
-	Linker.LinkExternalItem(BehaviorTagsHandle);
+	Linker.LinkExternalData(AnnotationTagsFragmentHandle);
+
+	Linker.LinkInstanceDataProperty(AnnotationTagsHandle, STATETREE_INSTANCEDATA_PROPERTY(FMassZoneGraphAnnotationEvaluatorInstanceData, AnnotationTags));
 
 	return true;
 }
 
-void FMassZoneGraphAnnotationEvaluator::Evaluate(FStateTreeExecutionContext& Context, const EStateTreeEvaluationType EvalType, const float DeltaTime)
+void FMassZoneGraphAnnotationEvaluator::Evaluate(FStateTreeExecutionContext& Context, const EStateTreeEvaluationType EvalType, const float DeltaTime) const
 {
-	const FMassZoneGraphAnnotationTagsFragment& BehaviorTagsFragment = Context.GetExternalItem(BehaviorTagsHandle);
-	BehaviorTags = BehaviorTagsFragment.Tags;
+	const FMassZoneGraphAnnotationTagsFragment& AnnotationTagsFragment = Context.GetExternalData(AnnotationTagsFragmentHandle);
+	FZoneGraphTagMask& AnnotationTags = Context.GetInstanceData(AnnotationTagsHandle);
+	AnnotationTags = AnnotationTagsFragment.Tags;
 }
diff --git a/Engine/Plugins/AI/MassAI/Source/MassAIBehavior/Private/MassStateTreeProcessors.cpp b/Engine/Plugins/AI/MassAI/Source/MassAIBehavior/Private/MassStateTreeProcessors.cpp
--- a/Engine/Plugins/AI/MassAI/Source/MassAIBehavior/Private/MassStateTreeProcessors.cpp
+++ b/Engine/Plugins/AI/MassAI/Source/MassAIBehavior/Private/MassStateTreeProcessors.cpp
@@ -23,21 +23,21 @@ bool SetExternalFragments(FMassStateTreeExecutionContext& Context, const UMassEn
 {
 	bool bFoundAllFragments = true;
 	const FMassEntityView EntityView(EntitySubsystem, Context.GetEntity());
-	for (const FStateTreeExternalItemDesc& ItemDesc : Context.GetExternalItems())
+	for (const FStateTreeExternalDataDesc& DataDesc : Context.GetExternalDataDescs())
 	{
-		if (ItemDesc.Struct && ItemDesc.Struct->IsChildOf(FMassFragment::StaticStruct()))
+		if (DataDesc.Struct && DataDesc.Struct->IsChildOf(FMassFragment::StaticStruct()))
 		{
-			const UScriptStruct* ScriptStruct = Cast<const UScriptStruct>(ItemDesc.Struct);
+			const UScriptStruct* ScriptStruct = Cast<const UScriptStruct>(DataDesc.Struct);
 			FStructView Fragment = EntityView.GetFragmentDataStruct(ScriptStruct);
 			if (Fragment.IsValid())
 			{
-				Context.SetExternalItem(ItemDesc.Handle, FStateTreeItemView(Fragment));
+				Context.SetExternalData(DataDesc.Handle, FStateTreeDataView(Fragment));
 			}
 			else
 			{
-				if (ItemDesc.Requirement == EStateTreeItemRequirement::Required)
+				if (DataDesc.Requirement == EStateTreeExternalDataRequirement::Required)
 				{
-					// Note: Not breaking here, so that we can validate all missing ones in one go with FMassStateTreeExecutionContext::AreExternalItemsValid().
+					// Note: Not breaking here, so that we can validate all missing ones in one go with FMassStateTreeExecutionContext::AreExternalDataViewsValid().
 					bFoundAllFragments = false;
 				}
 			}
@@ -55,21 +55,21 @@ bool SetExternalSubsystems(FMassStateTreeExecutionContext& Context)
 	}
 
 	bool bFoundAllSubsystems = true;
-	for (const FStateTreeExternalItemDesc& ItemDesc : Context.GetExternalItems())
+	for (const FStateTreeExternalDataDesc& DataDesc : Context.GetExternalDataDescs())
 	{
-		if (ItemDesc.Struct && ItemDesc.Struct->IsChildOf(UWorldSubsystem::StaticClass()))
+		if (DataDesc.Struct && DataDesc.Struct->IsChildOf(UWorldSubsystem::StaticClass()))
 		{
-			const TSubclassOf<UWorldSubsystem> SubClass = Cast<UClass>(const_cast<UStruct*>(ItemDesc.Struct));
+			const TSubclassOf<UWorldSubsystem> SubClass = Cast<UClass>(const_cast<UStruct*>(DataDesc.Struct));
 			UWorldSubsystem* Subsystem = World->GetSubsystemBase(SubClass);
 			if (Subsystem)
 			{
-				Context.SetExternalItem(ItemDesc.Handle, FStateTreeItemView(Subsystem));
+				Context.SetExternalData(DataDesc.Handle, FStateTreeDataView(Subsystem));
 			}
 			else
 			{
-				if (ItemDesc.Requirement == EStateTreeItemRequirement::Required)
+				if (DataDesc.Requirement == EStateTreeExternalDataRequirement::Required)
 				{
-					// Note: Not breaking here, so that we can validate all missing ones in one go with FMassStateTreeExecutionContext::AreExternalItemsValid().
+					// Note: Not breaking here, so that we can validate all missing ones in one go with FMassStateTreeExecutionContext::AreExternalDataViews Valid().
 					bFoundAllSubsystems = false;
 				}
 			}
@@ -81,7 +81,7 @@ bool SetExternalSubsystems(FMassStateTreeExecutionContext& Context)
 void ForEachEntityInChunk(
 	FMassStateTreeExecutionContext& StateTreeContext,
 	UMassStateTreeSubsystem& MassStateTreeSubsystem,
-	const TFunctionRef<void(FMassStateTreeExecutionContext&, FStateTreeItemView)> ForEachEntityCallback)
+	const TFunctionRef<void(FMassStateTreeExecutionContext&, FStateTreeDataView)> ForEachEntityCallback)
 {
 	const FMassExecutionContext& Context = StateTreeContext.GetEntitySubsystemExecutionContext();
 	const TConstArrayView<FMassStateTreeFragment> StateTreeList = Context.GetFragmentView<FMassStateTreeFragment>();
@@ -106,7 +106,7 @@ void ForEachEntityInChunk(
 		}
 	}
 
-	const UScriptStruct* StorageScriptStruct = StateTree->GetRuntimeStorageStruct();
+	const UScriptStruct* StorageScriptStruct = StateTree->GetInstanceStorageStruct();
 	for (int32 i = 0; i < NumEntities; ++i)
 	{
 		const FMassEntityHandle Entity = Context.GetEntity(i);
@@ -122,11 +122,11 @@ void ForEachEntityInChunk(
 			}
 		}
 
-		// Make sure all required external items are set.
+		// Make sure all required external data are set.
 		{
-			CSV_SCOPED_TIMING_STAT_EXCLUSIVE(StateTreeProcessorExternalItemsValidation);
+			CSV_SCOPED_TIMING_STAT_EXCLUSIVE(StateTreeProcessorExternalDataValidation);
 			// TODO: disable this when not in debug.
-			if (!ensureMsgf(StateTreeContext.AreExternalItemsValid(), TEXT("StateTree will not execute due to missing external items.")))
+			if (!ensureMsgf(StateTreeContext.AreExternalDataViewsValid(), TEXT("StateTree will not execute due to missing external data.")))
 			{
 				break;
 			}
@@ -166,7 +166,7 @@ void UMassStateTreeFragmentInitializer::Execute(UMassEntitySubsystem& EntitySubs
 			UE::MassBehavior::ForEachEntityInChunk(
 				StateTreeContext,
 				*MassStateTreeSubsystem,
-				[](FMassStateTreeExecutionContext& StateTreeExecutionContext, FStateTreeItemView Storage)
+				[](FMassStateTreeExecutionContext& StateTreeExecutionContext, FStateTreeDataView Storage)
 				{
 					// Start the tree instance
 					StateTreeExecutionContext.Start(Storage);
@@ -201,7 +201,7 @@ void UMassStateTreeFragmentDestructor::Execute(UMassEntitySubsystem& EntitySubsy
 			UE::MassBehavior::ForEachEntityInChunk(
 				StateTreeContext,
 				*MassStateTreeSubsystem,
-				[](FMassStateTreeExecutionContext& StateTreeExecutionContext, FStateTreeItemView Storage)
+				[](FMassStateTreeExecutionContext& StateTreeExecutionContext, FStateTreeDataView Storage)
 				{
 					// Stop the tree instance
 					StateTreeExecutionContext.Stop(Storage);
@@ -341,7 +341,7 @@ void UMassStateTreeProcessor::SignalEntities(UMassEntitySubsystem& EntitySubsyst
 			UE::MassBehavior::ForEachEntityInChunk(
 				StateTreeContext,
 				*MassStateTreeSubsystem,
-				[&StateTreeList, TimeDelta, TimeInSeconds, &EntitiesToSignal](FMassStateTreeExecutionContext& StateTreeExecutionContext, const FStateTreeItemView Storage)
+				[&StateTreeList, TimeDelta, TimeInSeconds, &EntitiesToSignal](FMassStateTreeExecutionContext& StateTreeExecutionContext, const FStateTreeDataView Storage)
 				{
 					// Compute adjusted delta time
 					TOptional<float>& LastUpdate = StateTreeList[StateTreeExecutionContext.GetEntityIndex()].LastUpdateTimeInSeconds;
diff --git a/Engine/Plugins/AI/MassAI/Source/MassAIBehavior/Private/MassStateTreeSchema.cpp b/Engine/Plugins/AI/MassAI/Source/MassAIBehavior/Private/MassStateTreeSchema.cpp
--- a/Engine/Plugins/AI/MassAI/Source/MassAIBehavior/Private/MassStateTreeSchema.cpp
+++ b/Engine/Plugins/AI/MassAI/Source/MassAIBehavior/Private/MassStateTreeSchema.cpp
@@ -3,6 +3,7 @@
 #include "MassStateTreeSchema.h"
 #include "MassEntityTypes.h"
 #include "MassStateTreeTypes.h"
+#include "StateTreeConditionBase.h"
 #include "Subsystems/WorldSubsystem.h"
 
 UScriptStruct* UMassStateTreeSchema::GetStorageSuperStruct() const
@@ -14,7 +15,9 @@ UScriptStruct* UMassStateTreeSchema::GetStorageSuperStruct() const
 bool UMassStateTreeSchema::IsStructAllowed(const UScriptStruct* InScriptStruct) const
 {
 	// Only allow Mass evals and tasks.
-	return InScriptStruct->IsChildOf(FMassStateTreeEvaluatorBase::StaticStruct()) || InScriptStruct->IsChildOf(FMassStateTreeTaskBase::StaticStruct());
+	return InScriptStruct->IsChildOf(FMassStateTreeEvaluatorBase::StaticStruct())
+			|| InScriptStruct->IsChildOf(FMassStateTreeTaskBase::StaticStruct())
+			|| InScriptStruct->IsChildOf(FStateTreeConditionBase::StaticStruct());
 }
 
 bool UMassStateTreeSchema::IsExternalItemAllowed(const UStruct& InStruct) const
diff --git a/Engine/Plugins/AI/MassAI/Source/MassAIBehavior/Private/MassStateTreeTrait.cpp b/Engine/Plugins/AI/MassAI/Source/MassAIBehavior/Private/MassStateTreeTrait.cpp
--- a/Engine/Plugins/AI/MassAI/Source/MassAIBehavior/Private/MassStateTreeTrait.cpp
+++ b/Engine/Plugins/AI/MassAI/Source/MassAIBehavior/Private/MassStateTreeTrait.cpp
@@ -23,7 +23,7 @@ void UMassStateTreeTrait::BuildTemplate(FMassEntityTemplateBuildContext& BuildCo
 		UE_VLOG(MassStateTreeSubsystem, LogMassBehavior, Error, TEXT("StateTree asset is not set or unavailable."));
 		return;
 	}
-	if (!StateTree->GetRuntimeStorageDefaultValue().IsValid())
+	if (!StateTree->GetInstanceStorageDefaultValue().IsValid())
 	{
 		UE_VLOG(MassStateTreeSubsystem, LogMassBehavior, Error, TEXT("StateTree asset is valid but missing runtime storage type."));
 		return;
@@ -36,7 +36,7 @@ void UMassStateTreeTrait::BuildTemplate(FMassEntityTemplateBuildContext& BuildCo
 	StateTreeFragment.StateTreeHandle = Handle;
 
 	// Add runtime storage as a fragment
-	BuildContext.AddFragment(StateTree->GetRuntimeStorageDefaultValue());
+	BuildContext.AddFragment(StateTree->GetInstanceStorageDefaultValue());
 }
 
 void UMassStateTreeTrait::ValidateTemplate(FMassEntityTemplateBuildContext& BuildContext, UWorld& World) const
@@ -55,9 +55,9 @@ void UMassStateTreeTrait::ValidateTemplate(FMassEntityTemplateBuildContext& Buil
 	}
 
 	// Make sure all the required subsystems can be found.
-	for (const FStateTreeExternalItemDesc& ItemDesc : StateTree->GetExternalItems())
+	for (const FStateTreeExternalDataDesc& ItemDesc : StateTree->GetExternalDataDescs())
 	{
-		if (ensure(ItemDesc.Struct) && ItemDesc.Requirement == EStateTreeItemRequirement::Required)
+		if (ensure(ItemDesc.Struct) && ItemDesc.Requirement == EStateTreeExternalDataRequirement::Required)
 		{
 			if (ItemDesc.Struct->IsChildOf(UWorldSubsystem::StaticClass()))
 			{
diff --git a/Engine/Plugins/AI/MassAI/Source/MassAIBehavior/Private/Tasks/MassLookAtTask.cpp b/Engine/Plugins/AI/MassAI/Source/MassAIBehavior/Private/Tasks/MassLookAtTask.cpp
--- a/Engine/Plugins/AI/MassAI/Source/MassAIBehavior/Private/Tasks/MassLookAtTask.cpp
+++ b/Engine/Plugins/AI/MassAI/Source/MassAIBehavior/Private/Tasks/MassLookAtTask.cpp
@@ -8,25 +8,32 @@
 
 bool FMassLookAtTask::Link(FStateTreeLinker& Linker)
 {
-	Linker.LinkExternalItem(MassSignalSubsystemHandle);
-	Linker.LinkExternalItem(LookAtHandle);
+	Linker.LinkExternalData(MassSignalSubsystemHandle);
+	Linker.LinkExternalData(LookAtHandle);
 
+	Linker.LinkInstanceDataProperty(DurationHandle, STATETREE_INSTANCEDATA_PROPERTY(FMassLookAtTaskInstanceData, Duration));
+	Linker.LinkInstanceDataProperty(TargetEntityHandle, STATETREE_INSTANCEDATA_PROPERTY(FMassLookAtTaskInstanceData, TargetEntity));
+	Linker.LinkInstanceDataProperty(TimeHandle, STATETREE_INSTANCEDATA_PROPERTY(FMassLookAtTaskInstanceData, Time));
+	
 	return true;
 }
 
-EStateTreeRunStatus FMassLookAtTask::EnterState(FStateTreeExecutionContext& Context, const EStateTreeStateChangeType ChangeType, const FStateTreeTransitionResult& Transition)
+EStateTreeRunStatus FMassLookAtTask::EnterState(FStateTreeExecutionContext& Context, const EStateTreeStateChangeType ChangeType, const FStateTreeTransitionResult& Transition) const
 {
+	float& Time = Context.GetInstanceData(TimeHandle);
+
 	Time = 0.f;
 	
 	const FMassStateTreeExecutionContext& MassContext = static_cast<FMassStateTreeExecutionContext&>(Context);
-	FMassLookAtFragment& LookAtFragment = MassContext.GetExternalItem(LookAtHandle);
+	FMassLookAtFragment& LookAtFragment = MassContext.GetExternalData(LookAtHandle);
 
 	LookAtFragment.Reset();
 	LookAtFragment.LookAtMode = LookAtMode;
 	
 	if (LookAtMode == EMassLookAtMode::LookAtEntity)
 	{
-		if (!TargetEntity.IsSet())
+		const FMassEntityHandle* TargetEntity = Context.GetInstanceDataPtr(TargetEntityHandle); // Optional input
+		if (TargetEntity == nullptr || !TargetEntity->IsSet())
 		{
 			LookAtFragment.LookAtMode = EMassLookAtMode::LookForward;
 			MASSBEHAVIOR_LOG(Error, TEXT("Failed LookAt: invalid target entity"));
@@ -34,7 +41,7 @@ EStateTreeRunStatus FMassLookAtTask::EnterState(FStateTreeExecutionContext& Cont
 		else
 		{
 			LookAtFragment.LookAtMode = EMassLookAtMode::LookAtEntity;
-			LookAtFragment.TrackedEntity = TargetEntity;
+			LookAtFragment.TrackedEntity = *TargetEntity;
 		}
 	}
 
@@ -45,25 +52,30 @@ EStateTreeRunStatus FMassLookAtTask::EnterState(FStateTreeExecutionContext& Cont
 
 	// A Duration <= 0 indicates that the task runs until a transition in the state tree stops it.
 	// Otherwise we schedule a signal to end the task.
+	const float Duration = Context.GetInstanceData(DurationHandle);
 	if (Duration > 0.0f)
 	{
-		UMassSignalSubsystem& MassSignalSubsystem = MassContext.GetExternalItem(MassSignalSubsystemHandle);
+		UMassSignalSubsystem& MassSignalSubsystem = MassContext.GetExternalData(MassSignalSubsystemHandle);
 		MassSignalSubsystem.DelaySignalEntity(UE::Mass::Signals::LookAtFinished, MassContext.GetEntity(), Duration);
 	}
 
 	return EStateTreeRunStatus::Running;
 }
 
-void FMassLookAtTask::ExitState(FStateTreeExecutionContext& Context, const EStateTreeStateChangeType ChangeType, const FStateTreeTransitionResult& Transition)
+void FMassLookAtTask::ExitState(FStateTreeExecutionContext& Context, const EStateTreeStateChangeType ChangeType, const FStateTreeTransitionResult& Transition) const
 {
 	const FMassStateTreeExecutionContext& MassContext = static_cast<FMassStateTreeExecutionContext&>(Context);
-	FMassLookAtFragment& LookAtFragment = MassContext.GetExternalItem(LookAtHandle);
+	FMassLookAtFragment& LookAtFragment = MassContext.GetExternalData(LookAtHandle);
 	
 	LookAtFragment.Reset();
 }
 
-EStateTreeRunStatus FMassLookAtTask::Tick(FStateTreeExecutionContext& Context, const float DeltaTime)
+EStateTreeRunStatus FMassLookAtTask::Tick(FStateTreeExecutionContext& Context, const float DeltaTime) const
 {
+	float& Time = Context.GetInstanceData(TimeHandle);
+	const float Duration = Context.GetInstanceData(DurationHandle);
+	
 	Time += DeltaTime;
+	
 	return Duration <= 0.0f ? EStateTreeRunStatus::Running : (Time < Duration ? EStateTreeRunStatus::Running : EStateTreeRunStatus::Succeeded);
 }
diff --git a/Engine/Plugins/AI/MassAI/Source/MassAIBehavior/Private/Tasks/MassStateTreeSmartObjectTask.cpp b/Engine/Plugins/AI/MassAI/Source/MassAIBehavior/Private/Tasks/MassStateTreeSmartObjectTask.cpp
--- a/Engine/Plugins/AI/MassAI/Source/MassAIBehavior/Private/Tasks/MassStateTreeSmartObjectTask.cpp
+++ b/Engine/Plugins/AI/MassAI/Source/MassAIBehavior/Private/Tasks/MassStateTreeSmartObjectTask.cpp
@@ -17,19 +17,25 @@
 
 bool FMassClaimSmartObjectTask::Link(FStateTreeLinker& Linker)
 {
-	Linker.LinkExternalItem(SmartObjectUserHandle);
-	Linker.LinkExternalItem(SmartObjectSubsystemHandle);
+	Linker.LinkExternalData(SmartObjectUserHandle);
+	Linker.LinkExternalData(SmartObjectSubsystemHandle);
+
+	Linker.LinkInstanceDataProperty(SearchRequestResultHandle, STATETREE_INSTANCEDATA_PROPERTY(FMassClaimSmartObjectTaskInstanceData, SearchRequestResult));
+	Linker.LinkInstanceDataProperty(ClaimResultHandle, STATETREE_INSTANCEDATA_PROPERTY(FMassClaimSmartObjectTaskInstanceData, ClaimResult));
 
 	return true;
 }
 
-EStateTreeRunStatus FMassClaimSmartObjectTask::EnterState(FStateTreeExecutionContext& Context, const EStateTreeStateChangeType ChangeType, const FStateTreeTransitionResult& Transition)
+EStateTreeRunStatus FMassClaimSmartObjectTask::EnterState(FStateTreeExecutionContext& Context, const EStateTreeStateChangeType ChangeType, const FStateTreeTransitionResult& Transition) const
 {
 	if (ChangeType != EStateTreeStateChangeType::Changed)
 	{
 		return EStateTreeRunStatus::Running;
 	}
 
+	const FMassSmartObjectRequestResult& SearchRequestResult = Context.GetInstanceData(SearchRequestResultHandle);
+	EMassSmartObjectClaimResult& ClaimResult = Context.GetInstanceData(ClaimResultHandle);
+
 	if (!SearchRequestResult.bProcessed)
 	{
 		MASSBEHAVIOR_LOG(Error, TEXT("Failed to claim smart object: unprocessed search request results."));
@@ -45,8 +51,8 @@ EStateTreeRunStatus FMassClaimSmartObjectTask::EnterState(FStateTreeExecutionCon
 	EStateTreeRunStatus Status = EStateTreeRunStatus::Failed;
 
 	// Retrieve fragments and subsystems
-	USmartObjectSubsystem& SmartObjectSubsystem = Context.GetExternalItem(SmartObjectSubsystemHandle);
-	FMassSmartObjectUserFragment& SOUser = Context.GetExternalItem(SmartObjectUserHandle);
+	USmartObjectSubsystem& SmartObjectSubsystem = Context.GetExternalData(SmartObjectSubsystemHandle);
+	FMassSmartObjectUserFragment& SOUser = Context.GetExternalData(SmartObjectUserHandle);
 
 	// Setup MassSmartObject handler and claim
 	const FMassStateTreeExecutionContext& MassContext = static_cast<FMassStateTreeExecutionContext&>(Context);
@@ -74,16 +80,17 @@ EStateTreeRunStatus FMassClaimSmartObjectTask::EnterState(FStateTreeExecutionCon
 	return Status;
 }
 
-void FMassClaimSmartObjectTask::ExitState(FStateTreeExecutionContext& Context, const EStateTreeStateChangeType ChangeType, const FStateTreeTransitionResult& Transition)
+void FMassClaimSmartObjectTask::ExitState(FStateTreeExecutionContext& Context, const EStateTreeStateChangeType ChangeType, const FStateTreeTransitionResult& Transition) const
 {
 	if (ChangeType != EStateTreeStateChangeType::Changed)
 	{
 		return;
 	}
 
+	EMassSmartObjectClaimResult& ClaimResult = Context.GetInstanceData(ClaimResultHandle);
 	ClaimResult = EMassSmartObjectClaimResult::Unset;
 
-	FMassSmartObjectUserFragment& SOUser = Context.GetExternalItem(SmartObjectUserHandle);
+	FMassSmartObjectUserFragment& SOUser = Context.GetExternalData(SmartObjectUserHandle);
 
 	if (SOUser.ClaimHandle.IsValid())
 	{
@@ -91,7 +98,7 @@ void FMassClaimSmartObjectTask::ExitState(FStateTreeExecutionContext& Context, c
 		{
 			MASSBEHAVIOR_LOG(VeryVerbose, TEXT("Exiting state with a valid claim handle but interaction was not started: release claim on the smart object."));
 			const FMassStateTreeExecutionContext& MassContext = static_cast<FMassStateTreeExecutionContext&>(Context);
-			USmartObjectSubsystem& SmartObjectSubsystem = Context.GetExternalItem(SmartObjectSubsystemHandle);
+			USmartObjectSubsystem& SmartObjectSubsystem = Context.GetExternalData(SmartObjectSubsystemHandle);
 			const FMassSmartObjectHandler MassSmartObjectHandler(MassContext.GetEntitySubsystem(), MassContext.GetEntitySubsystemExecutionContext(), SmartObjectSubsystem);
 			MassSmartObjectHandler.ReleaseSmartObject(MassContext.GetEntity(), SOUser);
 		}
@@ -106,26 +113,21 @@ void FMassClaimSmartObjectTask::ExitState(FStateTreeExecutionContext& Context, c
 	}
 }
 
-EStateTreeRunStatus FMassClaimSmartObjectTask::Tick(FStateTreeExecutionContext& Context, const float DeltaTime)
-{
-	return EStateTreeRunStatus::Running;
-}
-
 //----------------------------------------------------------------------//
 // FMassUseSmartObjectTask
 //----------------------------------------------------------------------//
 
 bool FMassUseSmartObjectTask::Link(FStateTreeLinker& Linker)
 {
-	Linker.LinkExternalItem(SmartObjectSubsystemHandle);
-	Linker.LinkExternalItem(EntityTransformHandle);
-	Linker.LinkExternalItem(SmartObjectUserHandle);
-	Linker.LinkExternalItem(MoveTargetHandle);
+	Linker.LinkExternalData(SmartObjectSubsystemHandle);
+	Linker.LinkExternalData(EntityTransformHandle);
+	Linker.LinkExternalData(SmartObjectUserHandle);
+	Linker.LinkExternalData(MoveTargetHandle);
 
 	return true;
 }
 
-EStateTreeRunStatus FMassUseSmartObjectTask::EnterState(FStateTreeExecutionContext& Context, const EStateTreeStateChangeType ChangeType, const FStateTreeTransitionResult& Transition)
+EStateTreeRunStatus FMassUseSmartObjectTask::EnterState(FStateTreeExecutionContext& Context, const EStateTreeStateChangeType ChangeType, const FStateTreeTransitionResult& Transition) const
 {
 	if (ChangeType != EStateTreeStateChangeType::Changed)
 	{
@@ -133,10 +135,10 @@ EStateTreeRunStatus FMassUseSmartObjectTask::EnterState(FStateTreeExecutionConte
 	}
 
 	// Retrieve fragments and subsystems
-	USmartObjectSubsystem& SmartObjectSubsystem = Context.GetExternalItem(SmartObjectSubsystemHandle);
-	FMassSmartObjectUserFragment& SOUser = Context.GetExternalItem(SmartObjectUserHandle);
-	const FDataFragment_Transform& TransformFragment = Context.GetExternalItem(EntityTransformHandle);
-	FMassMoveTargetFragment& MoveTarget = Context.GetExternalItem(MoveTargetHandle);
+	USmartObjectSubsystem& SmartObjectSubsystem = Context.GetExternalData(SmartObjectSubsystemHandle);
+	FMassSmartObjectUserFragment& SOUser = Context.GetExternalData(SmartObjectUserHandle);
+	const FDataFragment_Transform& TransformFragment = Context.GetExternalData(EntityTransformHandle);
+	FMassMoveTargetFragment& MoveTarget = Context.GetExternalData(MoveTargetHandle);
 
 	// Setup MassSmartObject handler and start interaction
 	const FMassStateTreeExecutionContext& MassContext = static_cast<FMassStateTreeExecutionContext&>(Context);
@@ -157,18 +159,18 @@ EStateTreeRunStatus FMassUseSmartObjectTask::EnterState(FStateTreeExecutionConte
 	return bSuccess ? EStateTreeRunStatus::Running : EStateTreeRunStatus::Failed;
 }
 
-void FMassUseSmartObjectTask::ExitState(FStateTreeExecutionContext& Context, const EStateTreeStateChangeType ChangeType, const FStateTreeTransitionResult& Transition)
+void FMassUseSmartObjectTask::ExitState(FStateTreeExecutionContext& Context, const EStateTreeStateChangeType ChangeType, const FStateTreeTransitionResult& Transition) const
 {
 	if (ChangeType == EStateTreeStateChangeType::Changed)
 	{
-		FMassSmartObjectUserFragment& SOUser = Context.GetExternalItem(SmartObjectUserHandle);
+		FMassSmartObjectUserFragment& SOUser = Context.GetExternalData(SmartObjectUserHandle);
 
 		if (SOUser.ClaimHandle.IsValid())
 		{
 			MASSBEHAVIOR_LOG(VeryVerbose, TEXT("Exiting state with a valid ClaimHandle: stop using the smart object."));
 
 			const FMassStateTreeExecutionContext& MassContext = static_cast<FMassStateTreeExecutionContext&>(Context);
-			USmartObjectSubsystem& SmartObjectSubsystem = Context.GetExternalItem(SmartObjectSubsystemHandle);
+			USmartObjectSubsystem& SmartObjectSubsystem = Context.GetExternalData(SmartObjectSubsystemHandle);
 			const FMassSmartObjectHandler MassSmartObjectHandler(MassContext.GetEntitySubsystem(), MassContext.GetEntitySubsystemExecutionContext(), SmartObjectSubsystem);
 			MassSmartObjectHandler.ReleaseSmartObject(MassContext.GetEntity(), SOUser, EMassSmartObjectInteractionStatus::Aborted);
 		}
@@ -179,27 +181,27 @@ void FMassUseSmartObjectTask::ExitState(FStateTreeExecutionContext& Context, con
 	}
 }
 
-void FMassUseSmartObjectTask::StateCompleted(FStateTreeExecutionContext& Context, const EStateTreeRunStatus CompletionStatus, const FStateTreeHandle CompletedState)
+void FMassUseSmartObjectTask::StateCompleted(FStateTreeExecutionContext& Context, const EStateTreeRunStatus CompletionStatus, const FStateTreeHandle CompletedState) const
 {
-	FMassSmartObjectUserFragment& SOUser = Context.GetExternalItem(SmartObjectUserHandle);
+	FMassSmartObjectUserFragment& SOUser = Context.GetExternalData(SmartObjectUserHandle);
 
 	if (SOUser.ClaimHandle.IsValid())
 	{
 		MASSBEHAVIOR_LOG(VeryVerbose, TEXT("Completing state with a valid ClaimHandle: stop using the smart object."));
 
 		const FMassStateTreeExecutionContext& MassContext = static_cast<FMassStateTreeExecutionContext&>(Context);
-		USmartObjectSubsystem& SmartObjectSubsystem = Context.GetExternalItem(SmartObjectSubsystemHandle);
+		USmartObjectSubsystem& SmartObjectSubsystem = Context.GetExternalData(SmartObjectSubsystemHandle);
 		const FMassSmartObjectHandler MassSmartObjectHandler(MassContext.GetEntitySubsystem(), MassContext.GetEntitySubsystemExecutionContext(), SmartObjectSubsystem);
 		const EMassSmartObjectInteractionStatus NewStatus = CompletionStatus == EStateTreeRunStatus::Succeeded ? EMassSmartObjectInteractionStatus::Completed : EMassSmartObjectInteractionStatus::Aborted;
 		MassSmartObjectHandler.ReleaseSmartObject(MassContext.GetEntity(), SOUser, NewStatus);
 	}
 }
 
-EStateTreeRunStatus FMassUseSmartObjectTask::Tick(FStateTreeExecutionContext& Context, const float DeltaTime)
+EStateTreeRunStatus FMassUseSmartObjectTask::Tick(FStateTreeExecutionContext& Context, const float DeltaTime) const
 {
 	EStateTreeRunStatus Status = EStateTreeRunStatus::Failed;
 
-	FMassSmartObjectUserFragment& SOUser = Context.GetExternalItem(SmartObjectUserHandle);
+	FMassSmartObjectUserFragment& SOUser = Context.GetExternalData(SmartObjectUserHandle);
 	switch (SOUser.InteractionStatus)
 	{
 	case EMassSmartObjectInteractionStatus::InProgress:
diff --git a/Engine/Plugins/AI/MassAI/Source/MassAIBehavior/Private/Tasks/MassZoneGraphFindEscapeTarget.cpp b/Engine/Plugins/AI/MassAI/Source/MassAIBehavior/Private/Tasks/MassZoneGraphFindEscapeTarget.cpp
--- a/Engine/Plugins/AI/MassAI/Source/MassAIBehavior/Private/Tasks/MassZoneGraphFindEscapeTarget.cpp
+++ b/Engine/Plugins/AI/MassAI/Source/MassAIBehavior/Private/Tasks/MassZoneGraphFindEscapeTarget.cpp
@@ -17,28 +17,30 @@ FMassZoneGraphFindEscapeTarget::FMassZoneGraphFindEscapeTarget()
 
 bool FMassZoneGraphFindEscapeTarget::Link(FStateTreeLinker& Linker)
 {
-	Linker.LinkExternalItem(LocationHandle);
-	Linker.LinkExternalItem(ZoneGraphSubsystemHandle);
-	Linker.LinkExternalItem(ZoneGraphAnnotationSubsystemHandle);
+	Linker.LinkExternalData(LocationHandle);
+	Linker.LinkExternalData(ZoneGraphSubsystemHandle);
+	Linker.LinkExternalData(ZoneGraphAnnotationSubsystemHandle);
+
+	Linker.LinkInstanceDataProperty(EscapeTargetLocationHandle, STATETREE_INSTANCEDATA_PROPERTY(FMassZoneGraphFindEscapeTargetInstangeData, EscapeTargetLocation));
 
 	return true;
 }
 
-EStateTreeRunStatus FMassZoneGraphFindEscapeTarget::EnterState(FStateTreeExecutionContext& Context, const EStateTreeStateChangeType ChangeType, const FStateTreeTransitionResult& Transition)
+EStateTreeRunStatus FMassZoneGraphFindEscapeTarget::EnterState(FStateTreeExecutionContext& Context, const EStateTreeStateChangeType ChangeType, const FStateTreeTransitionResult& Transition) const
 {
 	FMassStateTreeExecutionContext& MassContext = static_cast<FMassStateTreeExecutionContext&>(Context);
 	bool bDisplayDebug = false;
 
-	EscapeTargetLocation = nullptr;
-
 #if WITH_MASSGAMEPLAY_DEBUG
 	bDisplayDebug = UE::Mass::Debug::IsDebuggingEntity(MassContext.GetEntity());
 #endif // WITH_MASSGAMEPLAY_DEBUG
 
-	const FMassZoneGraphLaneLocationFragment& LaneLocation = Context.GetExternalItem(LocationHandle);
-	UZoneGraphSubsystem& ZoneGraphSubsystem = Context.GetExternalItem(ZoneGraphSubsystemHandle);
-	UZoneGraphAnnotationSubsystem& ZoneGraphAnnotationSubsystem = Context.GetExternalItem(ZoneGraphAnnotationSubsystemHandle);
+	const FMassZoneGraphLaneLocationFragment& LaneLocation = Context.GetExternalData(LocationHandle);
+	UZoneGraphSubsystem& ZoneGraphSubsystem = Context.GetExternalData(ZoneGraphSubsystemHandle);
+	UZoneGraphAnnotationSubsystem& ZoneGraphAnnotationSubsystem = Context.GetExternalData(ZoneGraphAnnotationSubsystemHandle);
 
+	FMassZoneGraphTargetLocation& EscapeTargetLocation = Context.GetInstanceData(EscapeTargetLocationHandle);
+	
 	if (!LaneLocation.LaneHandle.IsValid())
 	{
 		MASSBEHAVIOR_LOG(Error, TEXT("Invalid lane handle."));
@@ -90,62 +92,48 @@ EStateTreeRunStatus FMassZoneGraphFindEscapeTarget::EnterState(FStateTreeExecuti
 		// Could also try to sample few locations along the lane to see which is closest.
 		
 		// Small move, and goto adjacent lane
-		TargetLocation.LaneHandle = LaneLocation.LaneHandle;
-		TargetLocation.TargetDistance = FMath::Clamp(LaneLocation.DistanceAlongLane + AdjacentMoveDistance * MoveDir, 0.f, LaneLocation.LaneLength);
-		TargetLocation.NextExitLinkType = EZoneLaneLinkType::Adjacent;
-		TargetLocation.NextLaneHandle = FZoneGraphLaneHandle(EscapeSpan.ExitLaneIndex, ZoneGraphStorage->DataHandle);
-		TargetLocation.bMoveReverse = EscapeSpan.bReverseLaneDirection;
-		TargetLocation.EndOfPathIntent = EMassMovementAction::Move;
+		EscapeTargetLocation.LaneHandle = LaneLocation.LaneHandle;
+		EscapeTargetLocation.TargetDistance = FMath::Clamp(LaneLocation.DistanceAlongLane + AdjacentMoveDistance * MoveDir, 0.f, LaneLocation.LaneLength);
+		EscapeTargetLocation.NextExitLinkType = EZoneLaneLinkType::Adjacent;
+		EscapeTargetLocation.NextLaneHandle = FZoneGraphLaneHandle(EscapeSpan.ExitLaneIndex, ZoneGraphStorage->DataHandle);
+		EscapeTargetLocation.bMoveReverse = EscapeSpan.bReverseLaneDirection;
+		EscapeTargetLocation.EndOfPathIntent = EMassMovementAction::Move;
 
 		MASSBEHAVIOR_CLOG(bDisplayDebug, Log, TEXT("Switching from lane %s to adjacent lane %s."),
-			*LaneLocation.LaneHandle.ToString(), *TargetLocation.NextLaneHandle.ToString());
+			*LaneLocation.LaneHandle.ToString(), *EscapeTargetLocation.NextLaneHandle.ToString());
 	}
 	else
 	{
 		// Forward or backwards on current lane.
-		TargetLocation.LaneHandle = LaneLocation.LaneHandle;
-		TargetLocation.TargetDistance = LaneLocation.DistanceAlongLane + MoveDistance * MoveDir;
-		TargetLocation.NextExitLinkType = EZoneLaneLinkType::None;
-		TargetLocation.NextLaneHandle.Reset();
-		TargetLocation.bMoveReverse = EscapeSpan.bReverseLaneDirection;
-		TargetLocation.EndOfPathIntent = EMassMovementAction::Move;
+		EscapeTargetLocation.LaneHandle = LaneLocation.LaneHandle;
+		EscapeTargetLocation.TargetDistance = LaneLocation.DistanceAlongLane + MoveDistance * MoveDir;
+		EscapeTargetLocation.NextExitLinkType = EZoneLaneLinkType::None;
+		EscapeTargetLocation.NextLaneHandle.Reset();
+		EscapeTargetLocation.bMoveReverse = EscapeSpan.bReverseLaneDirection;
+		EscapeTargetLocation.EndOfPathIntent = EMassMovementAction::Move;
 
 		// When close to end of a lane, choose next lane.
-		const bool bPastStart = TargetLocation.TargetDistance < 0.0f;
-		const bool bPastEnd = TargetLocation.TargetDistance > LaneLocation.LaneLength;
+		const bool bPastStart = EscapeTargetLocation.TargetDistance < 0.0f;
+		const bool bPastEnd = EscapeTargetLocation.TargetDistance > LaneLocation.LaneLength;
 		if (bPastStart || bPastEnd)
 		{
-			TargetLocation.TargetDistance = FMath::Clamp(TargetLocation.TargetDistance, 0.0f, LaneLocation.LaneLength);
+			EscapeTargetLocation.TargetDistance = FMath::Clamp(EscapeTargetLocation.TargetDistance, 0.0f, LaneLocation.LaneLength);
 
-			TargetLocation.NextExitLinkType = EscapeSpan.ExitLinkType;
-			TargetLocation.NextLaneHandle = FZoneGraphLaneHandle(EscapeSpan.ExitLaneIndex, ZoneGraphStorage->DataHandle);
+			EscapeTargetLocation.NextExitLinkType = EscapeSpan.ExitLinkType;
+			EscapeTargetLocation.NextLaneHandle = FZoneGraphLaneHandle(EscapeSpan.ExitLaneIndex, ZoneGraphStorage->DataHandle);
 
 			MASSBEHAVIOR_CLOG(bDisplayDebug, Log, TEXT("Advancing %s along flee lane %s to next lane %s at distance %.1f."),
-				TargetLocation.bMoveReverse ? TEXT("forward") : TEXT("reverse"),
-				*TargetLocation.LaneHandle.ToString(), *TargetLocation.NextLaneHandle.ToString(),
-				TargetLocation.TargetDistance);
+				EscapeTargetLocation.bMoveReverse ? TEXT("forward") : TEXT("reverse"),
+				*EscapeTargetLocation.LaneHandle.ToString(), *EscapeTargetLocation.NextLaneHandle.ToString(),
+				EscapeTargetLocation.TargetDistance);
 		}
 		else
 		{
 			MASSBEHAVIOR_CLOG(bDisplayDebug, Log, TEXT("Advancing %s along flee lane %s to distance %.1f."),
-				TargetLocation.bMoveReverse ? TEXT("forward") : TEXT("reverse"),
-				*TargetLocation.LaneHandle.ToString(), TargetLocation.TargetDistance);
+				EscapeTargetLocation.bMoveReverse ? TEXT("forward") : TEXT("reverse"),
+				*EscapeTargetLocation.LaneHandle.ToString(), EscapeTargetLocation.TargetDistance);
 		}
 	}
 
-	EscapeTargetLocation = &TargetLocation;
-
-	return EStateTreeRunStatus::Running;
-}
-
-void FMassZoneGraphFindEscapeTarget::ExitState(FStateTreeExecutionContext& Context, const EStateTreeStateChangeType ChangeType, const FStateTreeTransitionResult& Transition)
-{
-	EscapeTargetLocation = &TargetLocation;
-}
-
-EStateTreeRunStatus FMassZoneGraphFindEscapeTarget::Tick(FStateTreeExecutionContext& Context, const float DeltaTime)
-{
-	EscapeTargetLocation = &TargetLocation;
-	
 	return EStateTreeRunStatus::Running;
 }
diff --git a/Engine/Plugins/AI/MassAI/Source/MassAIBehavior/Private/Tasks/MassZoneGraphFindSmartObjectTarget.cpp b/Engine/Plugins/AI/MassAI/Source/MassAIBehavior/Private/Tasks/MassZoneGraphFindSmartObjectTarget.cpp
--- a/Engine/Plugins/AI/MassAI/Source/MassAIBehavior/Private/Tasks/MassZoneGraphFindSmartObjectTarget.cpp
+++ b/Engine/Plugins/AI/MassAI/Source/MassAIBehavior/Private/Tasks/MassZoneGraphFindSmartObjectTarget.cpp
@@ -11,21 +11,24 @@
 
 bool FMassZoneGraphFindSmartObjectTarget::Link(FStateTreeLinker& Linker)
 {
-	Linker.LinkExternalItem(SmartObjectUserHandle);
-	Linker.LinkExternalItem(LocationHandle);
-	Linker.LinkExternalItem(AnnotationSubsystemHandle);
+	Linker.LinkExternalData(SmartObjectUserHandle);
+	Linker.LinkExternalData(LocationHandle);
+	Linker.LinkExternalData(AnnotationSubsystemHandle);
 
+	Linker.LinkInstanceDataProperty(SmartObjectLocationHandle, STATETREE_INSTANCEDATA_PROPERTY(FMassZoneGraphFindSmartObjectTargetInstanceData, SmartObjectLocation));
+	
 	return true;
 }
 
-EStateTreeRunStatus FMassZoneGraphFindSmartObjectTarget::EnterState(FStateTreeExecutionContext& Context, const EStateTreeStateChangeType ChangeType, const FStateTreeTransitionResult& Transition)
+EStateTreeRunStatus FMassZoneGraphFindSmartObjectTarget::EnterState(FStateTreeExecutionContext& Context, const EStateTreeStateChangeType ChangeType, const FStateTreeTransitionResult& Transition) const
 {
-	const UZoneGraphAnnotationSubsystem& AnnotationSubsystem = Context.GetExternalItem(AnnotationSubsystemHandle);
-	const FMassSmartObjectUserFragment& SOUser = Context.GetExternalItem(SmartObjectUserHandle);
-	const FMassZoneGraphLaneLocationFragment& LaneLocation = Context.GetExternalItem(LocationHandle);
+	const UZoneGraphAnnotationSubsystem& AnnotationSubsystem = Context.GetExternalData(AnnotationSubsystemHandle);
+	const FMassSmartObjectUserFragment& SOUser = Context.GetExternalData(SmartObjectUserHandle);
+	const FMassZoneGraphLaneLocationFragment& LaneLocation = Context.GetExternalData(LocationHandle);
 	const FZoneGraphLaneHandle LaneHandle(LaneLocation.LaneHandle);
 
-	TargetLocationRef = nullptr;
+	FMassZoneGraphTargetLocation& SmartObjectLocation = Context.GetInstanceData(SmartObjectLocationHandle);
+	SmartObjectLocation.Reset();
 
 	if (Context.GetEnterStateStatus() == EStateTreeRunStatus::Failed)
 	{
@@ -48,14 +51,15 @@ EStateTreeRunStatus FMassZoneGraphFindSmartObjectTarget::EnterState(FStateTreeEx
 	const FZoneGraphTag SmartObjectTag = GetDefault<UMassSmartObjectSettings>()->SmartObjectTag;
 	const USmartObjectZoneAnnotations* SOAnnotations = Cast<USmartObjectZoneAnnotations>(AnnotationSubsystem.GetFirstAnnotationForTag(SmartObjectTag));
 
-	TargetLocation.LaneHandle = LaneHandle;
-	TargetLocation.NextExitLinkType = EZoneLaneLinkType::None;
-	TargetLocation.NextLaneHandle.Reset();
-	TargetLocation.bMoveReverse = false;
-	TargetLocation.EndOfPathIntent = EMassMovementAction::Stand;
-	TargetLocation.EndOfPathPosition = SOUser.GetTargetLocation();
+	
+	SmartObjectLocation.LaneHandle = LaneHandle;
+	SmartObjectLocation.NextExitLinkType = EZoneLaneLinkType::None;
+	SmartObjectLocation.NextLaneHandle.Reset();
+	SmartObjectLocation.bMoveReverse = false;
+	SmartObjectLocation.EndOfPathIntent = EMassMovementAction::Stand;
+	SmartObjectLocation.EndOfPathPosition = SOUser.GetTargetLocation();
 	// Let's start moving toward the interaction a bit before the entry point.
-	TargetLocation.AnticipationDistance.Set(100.f);
+	SmartObjectLocation.AnticipationDistance.Set(100.f);
 
 	if (SOAnnotations != nullptr)
 	{
@@ -66,30 +70,14 @@ EStateTreeRunStatus FMassZoneGraphFindSmartObjectTarget::EnterState(FStateTreeEx
 
 		// Request path along current lane to reach entry point on lane
 		MASSBEHAVIOR_LOG(Log, TEXT("Claim successful: create path along lane to reach interaction location."));
-		TargetLocation.TargetDistance = EntryPoint.DistanceAlongLane;
+		SmartObjectLocation.TargetDistance = EntryPoint.DistanceAlongLane;
 	}
 	else
 	{
 		// Request path from current lane location directly to interaction location
 		MASSBEHAVIOR_LOG(Warning, TEXT("Claim successful: create path from current lane location directly to interaction location since SmartObject zone annotations weren't found."));
-		TargetLocation.TargetDistance = LaneLocation.DistanceAlongLane;
+		SmartObjectLocation.TargetDistance = LaneLocation.DistanceAlongLane;
 	}
-
-	TargetLocationRef = &TargetLocation;
-	
-	return EStateTreeRunStatus::Running;
-}
-
-void FMassZoneGraphFindSmartObjectTarget::ExitState(FStateTreeExecutionContext& Context, const EStateTreeStateChangeType ChangeType, const FStateTreeTransitionResult& Transition)
-{
-	// Temp solution to make sure the target location is always up to date (will be replaced with automatic update).
-	TargetLocationRef = &TargetLocation;
-}
-
-EStateTreeRunStatus FMassZoneGraphFindSmartObjectTarget::Tick(FStateTreeExecutionContext& Context, const float DeltaTime)
-{
-	// Temp solution to make sure the target location is always up to date (will be replaced with automatic update).
-	TargetLocationRef = &TargetLocation;
 	
 	return EStateTreeRunStatus::Running;
 }
diff --git a/Engine/Plugins/AI/MassAI/Source/MassAIBehavior/Private/Tasks/MassZoneGraphPathFollowTask.cpp b/Engine/Plugins/AI/MassAI/Source/MassAIBehavior/Private/Tasks/MassZoneGraphPathFollowTask.cpp
--- a/Engine/Plugins/AI/MassAI/Source/MassAIBehavior/Private/Tasks/MassZoneGraphPathFollowTask.cpp
+++ b/Engine/Plugins/AI/MassAI/Source/MassAIBehavior/Private/Tasks/MassZoneGraphPathFollowTask.cpp
@@ -13,28 +13,32 @@
 
 bool FMassZoneGraphPathFollowTask::Link(FStateTreeLinker& Linker)
 {
-	Linker.LinkExternalItem(LocationHandle);
-	Linker.LinkExternalItem(MoveTargetHandle);
-	Linker.LinkExternalItem(MovementConfigHandle);
-	Linker.LinkExternalItem(PathRequestHandle);
-	Linker.LinkExternalItem(ShortPathHandle);
-	Linker.LinkExternalItem(CachedLaneHandle);
-	Linker.LinkExternalItem(AgentRadiusHandle);
-	Linker.LinkExternalItem(ZoneGraphSubsystemHandle);
+	Linker.LinkExternalData(LocationHandle);
+	Linker.LinkExternalData(MoveTargetHandle);
+	Linker.LinkExternalData(MovementConfigHandle);
+	Linker.LinkExternalData(PathRequestHandle);
+	Linker.LinkExternalData(ShortPathHandle);
+	Linker.LinkExternalData(CachedLaneHandle);
+	Linker.LinkExternalData(AgentRadiusHandle);
+	Linker.LinkExternalData(ZoneGraphSubsystemHandle);
+
+	Linker.LinkInstanceDataProperty(TargetLocationHandle, STATETREE_INSTANCEDATA_PROPERTY(FMassZoneGraphPathFollowTaskInstanceData, TargetLocation));
+	Linker.LinkInstanceDataProperty(MovementStyleHandle, STATETREE_INSTANCEDATA_PROPERTY(FMassZoneGraphPathFollowTaskInstanceData, MovementStyle));
+	Linker.LinkInstanceDataProperty(SpeedScaleHandle, STATETREE_INSTANCEDATA_PROPERTY(FMassZoneGraphPathFollowTaskInstanceData, SpeedScale));
 
 	return true;
 }
 
 bool FMassZoneGraphPathFollowTask::RequestPath(FMassStateTreeExecutionContext& Context, const FMassZoneGraphTargetLocation& RequestedTargetLocation) const
 {
-	const UZoneGraphSubsystem& ZoneGraphSubsystem = Context.GetExternalItem(ZoneGraphSubsystemHandle);
-	const FMassZoneGraphLaneLocationFragment& LaneLocation = Context.GetExternalItem(LocationHandle);
-	const FMassMovementConfigFragment& MovementConfig = Context.GetExternalItem(MovementConfigHandle);
-	const FDataFragment_AgentRadius& AgentRadius = Context.GetExternalItem(AgentRadiusHandle);
-	FMassZoneGraphShortPathFragment& ShortPath = Context.GetExternalItem(ShortPathHandle);
-	FMassZoneGraphCachedLaneFragment& CachedLane = Context.GetExternalItem(CachedLaneHandle);
-	FMassMoveTargetFragment& MoveTarget = Context.GetExternalItem(MoveTargetHandle);
-	FMassZoneGraphPathRequestFragment& RequestFragment = Context.GetExternalItem(PathRequestHandle);
+	const UZoneGraphSubsystem& ZoneGraphSubsystem = Context.GetExternalData(ZoneGraphSubsystemHandle);
+	const FMassZoneGraphLaneLocationFragment& LaneLocation = Context.GetExternalData(LocationHandle);
+	const FMassMovementConfigFragment& MovementConfig = Context.GetExternalData(MovementConfigHandle);
+	const FDataFragment_AgentRadius& AgentRadius = Context.GetExternalData(AgentRadiusHandle);
+	FMassZoneGraphShortPathFragment& ShortPath = Context.GetExternalData(ShortPathHandle);
+	FMassZoneGraphCachedLaneFragment& CachedLane = Context.GetExternalData(CachedLaneHandle);
+	FMassMoveTargetFragment& MoveTarget = Context.GetExternalData(MoveTargetHandle);
+	FMassZoneGraphPathRequestFragment& RequestFragment = Context.GetExternalData(PathRequestHandle);
 
 	const UMassMovementSettings* Settings = GetDefault<UMassMovementSettings>();
 	check(Settings);
@@ -82,6 +86,9 @@ bool FMassZoneGraphPathFollowTask::RequestPath(FMassStateTreeExecutionContext& C
 	PathRequest.AnticipationDistance = RequestedTargetLocation.AnticipationDistance;
 	PathRequest.EndOfPathOffset.Set(FMath::RandRange(-AgentRadius.Radius, AgentRadius.Radius));
 
+	const FMassMovementStyleRef& MovementStyle = Context.GetInstanceData(MovementStyleHandle);
+	const float SpeedScale = Context.GetInstanceData(SpeedScaleHandle);
+
 	const float DesiredSpeed = FMath::Min(Config->GenerateDesiredSpeed(MovementStyle, Context.GetEntity().Index) * SpeedScale, Config->MaximumSpeed);
 	const UWorld* World = Context.GetWorld();
 	checkf(World != nullptr, TEXT("A valid world is expected from the execution context"));
@@ -90,7 +97,7 @@ bool FMassZoneGraphPathFollowTask::RequestPath(FMassStateTreeExecutionContext& C
 	return UE::MassMovement::ActivateActionMove(*World, Context.GetOwner(), Context.GetEntity(), ZoneGraphSubsystem, LaneLocation, PathRequest, AgentRadius.Radius, DesiredSpeed, MoveTarget, ShortPath, CachedLane);
 }
 
-EStateTreeRunStatus FMassZoneGraphPathFollowTask::EnterState(FStateTreeExecutionContext& Context, const EStateTreeStateChangeType ChangeType, const FStateTreeTransitionResult& Transition)
+EStateTreeRunStatus FMassZoneGraphPathFollowTask::EnterState(FStateTreeExecutionContext& Context, const EStateTreeStateChangeType ChangeType, const FStateTreeTransitionResult& Transition) const
 {
 	// Do not reset of the state if current state is still active after transition, unless transitioned specifically to this state.
 	if (ChangeType == EStateTreeStateChangeType::Sustained && Transition.Current != Transition.Next)
@@ -115,22 +122,16 @@ EStateTreeRunStatus FMassZoneGraphPathFollowTask::EnterState(FStateTreeExecution
 		return EStateTreeRunStatus::Failed;
 	}
 
-	const FMassZoneGraphLaneLocationFragment& LaneLocation = Context.GetExternalItem(LocationHandle);
-
-	const FMassZoneGraphTargetLocation* TargetLocationPtr = TargetLocation.GetPtr<FMassZoneGraphTargetLocation>();
-	if (!TargetLocationPtr)
-	{
-		MASSBEHAVIOR_LOG(Error, TEXT("Target location not set."));
-		return EStateTreeRunStatus::Failed;
-	}
+	const FMassZoneGraphLaneLocationFragment& LaneLocation = Context.GetExternalData(LocationHandle);
+	const FMassZoneGraphTargetLocation& TargetLocation = Context.GetInstanceData(TargetLocationHandle);
 	
-	if (TargetLocationPtr->LaneHandle != LaneLocation.LaneHandle)
+	if (TargetLocation.LaneHandle != LaneLocation.LaneHandle)
 	{
-		MASSBEHAVIOR_LOG(Error, TEXT("Target is not on current lane, target lane is %s expected %s."), *TargetLocationPtr->LaneHandle.ToString(), *LaneLocation.LaneHandle.ToString());
+		MASSBEHAVIOR_LOG(Error, TEXT("Target is not on current lane, target lane is %s expected %s."), *TargetLocation.LaneHandle.ToString(), *LaneLocation.LaneHandle.ToString());
 		return EStateTreeRunStatus::Failed;
 	}
 
-	if (!RequestPath(MassContext, *TargetLocationPtr))
+	if (!RequestPath(MassContext, TargetLocation))
 	{
 		return EStateTreeRunStatus::Failed;
 	}
@@ -138,7 +139,7 @@ EStateTreeRunStatus FMassZoneGraphPathFollowTask::EnterState(FStateTreeExecution
 	return EStateTreeRunStatus::Running;
 }
 
-EStateTreeRunStatus FMassZoneGraphPathFollowTask::Tick(FStateTreeExecutionContext& Context, const float DeltaTime)
+EStateTreeRunStatus FMassZoneGraphPathFollowTask::Tick(FStateTreeExecutionContext& Context, const float DeltaTime) const
 {
 	FMassStateTreeExecutionContext& MassContext = static_cast<FMassStateTreeExecutionContext&>(Context);
 
@@ -151,26 +152,20 @@ EStateTreeRunStatus FMassZoneGraphPathFollowTask::Tick(FStateTreeExecutionContex
 		MASSBEHAVIOR_LOG(Verbose, TEXT("tick"));
 	}
 
-	const FMassZoneGraphShortPathFragment& ShortPath = Context.GetExternalItem(ShortPathHandle);
-	const FMassZoneGraphLaneLocationFragment& LaneLocation = Context.GetExternalItem(LocationHandle);
-
-	const FMassZoneGraphTargetLocation* TargetLocationPtr = TargetLocation.GetPtr<FMassZoneGraphTargetLocation>();
-	if (!TargetLocationPtr)
-	{
-		MASSBEHAVIOR_LOG(Error, TEXT("Target location not set."));
-		return EStateTreeRunStatus::Failed;
-	}
+	const FMassZoneGraphShortPathFragment& ShortPath = Context.GetExternalData(ShortPathHandle);
+	const FMassZoneGraphLaneLocationFragment& LaneLocation = Context.GetExternalData(LocationHandle);
+	const FMassZoneGraphTargetLocation& TargetLocation = Context.GetInstanceData(TargetLocationHandle);
 
 	// Current path follow is done, but it was partial (i.e. many points on a curve), try again until we get there.
 	if (ShortPath.IsDone() && ShortPath.bPartialResult)
 	{
-		if (TargetLocationPtr->LaneHandle != LaneLocation.LaneHandle)
+		if (TargetLocation.LaneHandle != LaneLocation.LaneHandle)
 		{
-			MASSBEHAVIOR_LOG(Error, TEXT("Target is not on current lane, target lane is %s expected %s."), *TargetLocationPtr->LaneHandle.ToString(), *LaneLocation.LaneHandle.ToString());
+			MASSBEHAVIOR_LOG(Error, TEXT("Target is not on current lane, target lane is %s expected %s."), *TargetLocation.LaneHandle.ToString(), *LaneLocation.LaneHandle.ToString());
 			return EStateTreeRunStatus::Failed;
 		}
 		
-		if (!RequestPath(MassContext, *TargetLocationPtr))
+		if (!RequestPath(MassContext, TargetLocation))
 		{
 			MASSBEHAVIOR_LOG(Error, TEXT("Failed to request path."));
 			return EStateTreeRunStatus::Failed;
diff --git a/Engine/Plugins/AI/MassAI/Source/MassAIBehavior/Private/Tasks/MassZoneGraphStandTask.cpp b/Engine/Plugins/AI/MassAI/Source/MassAIBehavior/Private/Tasks/MassZoneGraphStandTask.cpp
--- a/Engine/Plugins/AI/MassAI/Source/MassAIBehavior/Private/Tasks/MassZoneGraphStandTask.cpp
+++ b/Engine/Plugins/AI/MassAI/Source/MassAIBehavior/Private/Tasks/MassZoneGraphStandTask.cpp
@@ -14,18 +14,21 @@
 
 bool FMassZoneGraphStandTask::Link(FStateTreeLinker& Linker)
 {
-	Linker.LinkExternalItem(LocationHandle);
-	Linker.LinkExternalItem(MoveTargetHandle);
-	Linker.LinkExternalItem(ShortPathHandle);
-	Linker.LinkExternalItem(CachedLaneHandle);
-	Linker.LinkExternalItem(ZoneGraphSubsystemHandle);
-	Linker.LinkExternalItem(MassSignalSubsystemHandle);
-	Linker.LinkExternalItem(MovementConfigHandle);
+	Linker.LinkExternalData(LocationHandle);
+	Linker.LinkExternalData(MoveTargetHandle);
+	Linker.LinkExternalData(ShortPathHandle);
+	Linker.LinkExternalData(CachedLaneHandle);
+	Linker.LinkExternalData(ZoneGraphSubsystemHandle);
+	Linker.LinkExternalData(MassSignalSubsystemHandle);
+	Linker.LinkExternalData(MovementConfigHandle);
+
+	Linker.LinkInstanceDataProperty(DurationHandle, STATETREE_INSTANCEDATA_PROPERTY(FMassZoneGraphStandTaskInstanceData, Duration));
+	Linker.LinkInstanceDataProperty(TimeHandle, STATETREE_INSTANCEDATA_PROPERTY(FMassZoneGraphStandTaskInstanceData, Time));
 
 	return true;
 }
 
-EStateTreeRunStatus FMassZoneGraphStandTask::EnterState(FStateTreeExecutionContext& Context, const EStateTreeStateChangeType ChangeType, const FStateTreeTransitionResult& Transition)
+EStateTreeRunStatus FMassZoneGraphStandTask::EnterState(FStateTreeExecutionContext& Context, const EStateTreeStateChangeType ChangeType, const FStateTreeTransitionResult& Transition) const
 {
 	// Do not reset of the state if current state is still active after transition, unless transitioned specifically to this state.
 	if (ChangeType == EStateTreeStateChangeType::Sustained && Transition.Current != Transition.Next)
@@ -35,9 +38,9 @@ EStateTreeRunStatus FMassZoneGraphStandTask::EnterState(FStateTreeExecutionConte
 
 	const FMassStateTreeExecutionContext& MassContext = static_cast<FMassStateTreeExecutionContext&>(Context);
 
-	const FMassZoneGraphLaneLocationFragment& LaneLocation = Context.GetExternalItem(LocationHandle);
-	const UZoneGraphSubsystem& ZoneGraphSubsystem = Context.GetExternalItem(ZoneGraphSubsystemHandle);
-	const FMassMovementConfigFragment& MovementConfig = Context.GetExternalItem(MovementConfigHandle);
+	const FMassZoneGraphLaneLocationFragment& LaneLocation = Context.GetExternalData(LocationHandle);
+	const UZoneGraphSubsystem& ZoneGraphSubsystem = Context.GetExternalData(ZoneGraphSubsystemHandle);
+	const FMassMovementConfigFragment& MovementConfig = Context.GetExternalData(MovementConfigHandle);
 
 	const UMassMovementSettings* Settings = GetDefault<UMassMovementSettings>();
 	check(Settings);
@@ -55,9 +58,9 @@ EStateTreeRunStatus FMassZoneGraphStandTask::EnterState(FStateTreeExecutionConte
 		return EStateTreeRunStatus::Failed;
 	}
 
-	FMassZoneGraphShortPathFragment& ShortPath = Context.GetExternalItem(ShortPathHandle);
-	FMassZoneGraphCachedLaneFragment& CachedLane = Context.GetExternalItem(CachedLaneHandle);
-	FMassMoveTargetFragment& MoveTarget = Context.GetExternalItem(MoveTargetHandle);
+	FMassZoneGraphShortPathFragment& ShortPath = Context.GetExternalData(ShortPathHandle);
+	FMassZoneGraphCachedLaneFragment& CachedLane = Context.GetExternalData(CachedLaneHandle);
+	FMassMoveTargetFragment& MoveTarget = Context.GetExternalData(MoveTargetHandle);
 
 	// TODO: This could be smarter too, like having a stand location/direction, or even make a small path to stop, if we're currently running.
 
@@ -71,21 +74,26 @@ EStateTreeRunStatus FMassZoneGraphStandTask::EnterState(FStateTreeExecutionConte
 		return EStateTreeRunStatus::Failed;
 	}
 
+	const float Duration = Context.GetInstanceData(DurationHandle);
+	float& Time = Context.GetInstanceData(TimeHandle);
 	Time = 0.0f;
 
 	// A Duration <= 0 indicates that the task runs until a transition in the state tree stops it.
 	// Otherwise we schedule a signal to end the task.
 	if (Duration > 0.0f)
 	{
-		UMassSignalSubsystem& MassSignalSubsystem = Context.GetExternalItem(MassSignalSubsystemHandle);
+		UMassSignalSubsystem& MassSignalSubsystem = Context.GetExternalData(MassSignalSubsystemHandle);
 		MassSignalSubsystem.DelaySignalEntity(UE::Mass::Signals::StandTaskFinished, MassContext.GetEntity(), Duration);
 	}
 
 	return EStateTreeRunStatus::Running;
 }
 
-EStateTreeRunStatus FMassZoneGraphStandTask::Tick(FStateTreeExecutionContext& Context, const float DeltaTime)
+EStateTreeRunStatus FMassZoneGraphStandTask::Tick(FStateTreeExecutionContext& Context, const float DeltaTime) const
 {
+	const float Duration = Context.GetInstanceData(DurationHandle);
+	float& Time = Context.GetInstanceData(TimeHandle);
+	
 	Time += DeltaTime;
 	return Duration <= 0.0f ? EStateTreeRunStatus::Running : (Time < Duration ? EStateTreeRunStatus::Running : EStateTreeRunStatus::Succeeded);
 }
diff --git a/Engine/Plugins/AI/MassAI/Source/MassAIBehavior/Public/Conditions/ZoneGraphTagConditions.h b/Engine/Plugins/AI/MassAI/Source/MassAIBehavior/Public/Conditions/ZoneGraphTagConditions.h
--- a/Engine/Plugins/AI/MassAI/Source/MassAIBehavior/Public/Conditions/ZoneGraphTagConditions.h
+++ b/Engine/Plugins/AI/MassAI/Source/MassAIBehavior/Public/Conditions/ZoneGraphTagConditions.h
@@ -10,24 +10,32 @@
 /**
 * ZoneGraph Tag condition.
 */
+
+USTRUCT()
+struct MASSAIBEHAVIOR_API FZoneGraphTagFilterConditionInstanceData
+{
+	GENERATED_BODY()
+	
+	UPROPERTY(EditAnywhere, Category = Input)
+	FZoneGraphTagMask Tags = FZoneGraphTagMask::None;
+};
+
 USTRUCT(DisplayName="ZoneGraphTagFilter Compare")
 struct MASSAIBEHAVIOR_API FZoneGraphTagFilterCondition : public FStateTreeConditionBase
 {
 	GENERATED_BODY()
 
-	FZoneGraphTagFilterCondition() : FStateTreeConditionBase() { }
+	FZoneGraphTagFilterCondition() = default;
+	
+	virtual bool Link(FStateTreeLinker& Linker) override;
+	virtual const UStruct* GetInstanceDataType() const override { return FZoneGraphTagFilterConditionInstanceData::StaticStruct(); }
+	virtual bool TestCondition(FStateTreeExecutionContext& Context) const override;
 
 #if WITH_EDITOR
-	virtual FText GetDescription(const IStateTreeBindingLookup& BindingLookup) const override;
+	virtual FText GetDescription(const FGuid& ID, FConstStructView InstanceData, const IStateTreeBindingLookup& BindingLookup) const override;
 #endif
 
-	virtual bool TestCondition() const override
-	{
-		return Filter.Pass(Tags) ^ bInvert;
-	}
-
-	UPROPERTY(EditAnywhere, Category = Condition, meta = (Bindable))
-	FZoneGraphTagMask Tags = FZoneGraphTagMask::None;
+	TStateTreeInstanceDataPropertyHandle<FZoneGraphTagMask> TagsHandle;
 
 	UPROPERTY(EditAnywhere, Category = Condition)
 	FZoneGraphTagFilter Filter;
@@ -39,30 +47,39 @@ struct MASSAIBEHAVIOR_API FZoneGraphTagFilterCondition : public FStateTreeCondit
 /**
 * ZoneGraph Tag condition.
 */
+
+USTRUCT()
+struct MASSAIBEHAVIOR_API FZoneGraphTagMaskConditionInstanceData
+{
+	GENERATED_BODY()
+
+	UPROPERTY(EditAnywhere, Category = Input)
+	FZoneGraphTagMask Left = FZoneGraphTagMask::None;
+
+	UPROPERTY(EditAnywhere, Category = Parameter)
+	FZoneGraphTagMask Right = FZoneGraphTagMask::None;
+};
+
 USTRUCT(DisplayName="ZoneGraphTagMask Compare")
 struct MASSAIBEHAVIOR_API FZoneGraphTagMaskCondition : public FStateTreeConditionBase
 {
 	GENERATED_BODY()
 
-	FZoneGraphTagMaskCondition() : FStateTreeConditionBase() { }
+	FZoneGraphTagMaskCondition() = default;
+
+	virtual bool Link(FStateTreeLinker& Linker) override;
+	virtual const UStruct* GetInstanceDataType() const override { return FZoneGraphTagMaskConditionInstanceData::StaticStruct(); }
+	virtual bool TestCondition(FStateTreeExecutionContext& Context) const override;
 
 #if WITH_EDITOR
-	virtual FText GetDescription(const IStateTreeBindingLookup& BindingLookup) const override;
+	virtual FText GetDescription(const FGuid& ID, FConstStructView InstanceData, const IStateTreeBindingLookup& BindingLookup) const override;
 #endif
 
-	virtual bool TestCondition() const override
-	{
-		return Left.CompareMasks(Right, Operator) ^ bInvert;
-	}
-
-	UPROPERTY(EditAnywhere, Category = Condition, meta = (Bindable))
-	FZoneGraphTagMask Left = FZoneGraphTagMask::None;
+	TStateTreeInstanceDataPropertyHandle<FZoneGraphTagMask> LeftHandle;
+	TStateTreeInstanceDataPropertyHandle<FZoneGraphTagMask> RightHandle;
 
 	UPROPERTY(EditAnywhere, Category = Condition)
 	EZoneLaneTagMaskComparison Operator = EZoneLaneTagMaskComparison::Any;
-	
-	UPROPERTY(EditAnywhere, Category = Condition, meta = (Bindable))
-	FZoneGraphTagMask Right = FZoneGraphTagMask::None;
 
 	UPROPERTY(EditAnywhere, Category = Condition)
 	bool bInvert = false;
@@ -71,27 +88,36 @@ struct MASSAIBEHAVIOR_API FZoneGraphTagMaskCondition : public FStateTreeConditio
 /**
 * ZoneGraph Tag condition.
 */
+
+USTRUCT()
+struct MASSAIBEHAVIOR_API FZoneGraphTagConditionInstanceData
+{
+	GENERATED_BODY()
+
+	UPROPERTY(EditAnywhere, Category = Input)
+	FZoneGraphTag Left = FZoneGraphTag::None;
+
+	UPROPERTY(EditAnywhere, Category = Parameter)
+	FZoneGraphTag Right = FZoneGraphTag::None;
+};
+
 USTRUCT(DisplayName="ZoneGraphTag Compare")
 struct MASSAIBEHAVIOR_API FZoneGraphTagCondition : public FStateTreeConditionBase
 {
 	GENERATED_BODY()
 
-	FZoneGraphTagCondition() : FStateTreeConditionBase() { }
+	FZoneGraphTagCondition() = default;
+
+	virtual bool Link(FStateTreeLinker& Linker) override;
+	virtual const UStruct* GetInstanceDataType() const override { return FZoneGraphTagConditionInstanceData::StaticStruct(); }
+	virtual bool TestCondition(FStateTreeExecutionContext& Context) const override;
 
 #if WITH_EDITOR
-	virtual FText GetDescription(const IStateTreeBindingLookup& BindingLookup) const override;
+	virtual FText GetDescription(const FGuid& ID, FConstStructView InstanceData, const IStateTreeBindingLookup& BindingLookup) const override;
 #endif
 
-	virtual bool TestCondition() const override
-	{
-		return (Left == Right) ^ bInvert;
-	}
-
-	UPROPERTY(EditAnywhere, Category = Condition, meta = (Bindable))
-	FZoneGraphTag Left = FZoneGraphTag::None;
-
-	UPROPERTY(EditAnywhere, Category = Condition, meta = (Bindable))
-	FZoneGraphTag Right = FZoneGraphTag::None;
+	TStateTreeInstanceDataPropertyHandle<FZoneGraphTag> LeftHandle;
+	TStateTreeInstanceDataPropertyHandle<FZoneGraphTag> RightHandle;
 
 	UPROPERTY(EditAnywhere, Category = Condition)
 	bool bInvert = false;
diff --git a/Engine/Plugins/AI/MassAI/Source/MassAIBehavior/Public/Evaluators/MassComponentHitEvaluator.h b/Engine/Plugins/AI/MassAI/Source/MassAIBehavior/Public/Evaluators/MassComponentHitEvaluator.h
--- a/Engine/Plugins/AI/MassAI/Source/MassAIBehavior/Public/Evaluators/MassComponentHitEvaluator.h
+++ b/Engine/Plugins/AI/MassAI/Source/MassAIBehavior/Public/Evaluators/MassComponentHitEvaluator.h
@@ -11,6 +11,19 @@ class UMassComponentHitSubsystem;
 /**
  * Evaluator to extract last hit from the MassComponentHitSubsystem and expose it for tasks and transitions
  */
+
+USTRUCT()
+struct MASSAIBEHAVIOR_API FMassComponentHitEvaluatorInstanceData
+{
+	GENERATED_BODY()
+
+	UPROPERTY(VisibleAnywhere, Category = Output)
+	bool bGotHit = false;
+
+	UPROPERTY(VisibleAnywhere, Category = Output)
+	FMassEntityHandle LastHitEntity;
+};
+
 USTRUCT(meta = (DisplayName = "Mass ComponentHit Eval"))
 struct MASSAIBEHAVIOR_API FMassComponentHitEvaluator : public FMassStateTreeEvaluatorBase
 {
@@ -18,13 +31,11 @@ struct MASSAIBEHAVIOR_API FMassComponentHitEvaluator : public FMassStateTreeEval
 
 protected:
 	virtual bool Link(FStateTreeLinker& Linker) override;
-	virtual void Evaluate(FStateTreeExecutionContext& Context, const EStateTreeEvaluationType EvalType, const float DeltaTime) override;
-
-	TStateTreeItemHandle<UMassComponentHitSubsystem> ComponentHitSubsystemHandle;
+	virtual const UStruct* GetInstanceDataType() const override { return FMassComponentHitEvaluatorInstanceData::StaticStruct(); }
+	virtual void Evaluate(FStateTreeExecutionContext& Context, const EStateTreeEvaluationType EvalType, const float DeltaTime) const override;
 
-	UPROPERTY(VisibleAnywhere, Category = Output)
-	bool bGotHit = false;
+	TStateTreeExternalDataHandle<UMassComponentHitSubsystem> ComponentHitSubsystemHandle;
 
-	UPROPERTY(VisibleAnywhere, Category = Output)
-	FMassEntityHandle LastHitEntity;
+	TStateTreeInstanceDataPropertyHandle<bool> GotHitHandle;
+	TStateTreeInstanceDataPropertyHandle<FMassEntityHandle> LastHitEntityHandle;
 };
diff --git a/Engine/Plugins/AI/MassAI/Source/MassAIBehavior/Public/Evaluators/MassStateTreeSmartObjectEvaluator.h b/Engine/Plugins/AI/MassAI/Source/MassAIBehavior/Public/Evaluators/MassStateTreeSmartObjectEvaluator.h
--- a/Engine/Plugins/AI/MassAI/Source/MassAIBehavior/Public/Evaluators/MassStateTreeSmartObjectEvaluator.h
+++ b/Engine/Plugins/AI/MassAI/Source/MassAIBehavior/Public/Evaluators/MassStateTreeSmartObjectEvaluator.h
@@ -14,51 +14,67 @@ struct FMassZoneGraphLaneLocationFragment;
 /**
  * Evaluator that will keep track if there is one or more smart object(s) that can be used.
  */
-USTRUCT(meta = (DisplayName = "Mass SmartObject Eval"))
-struct MASSAIBEHAVIOR_API FMassStateTreeSmartObjectEvaluator : public FMassStateTreeEvaluatorBase
+USTRUCT()
+struct MASSAIBEHAVIOR_API FMassStateTreeSmartObjectEvaluatorInstanceData
 {
 	GENERATED_BODY()
 
-protected:
-	virtual bool Link(FStateTreeLinker& Linker) override;
-	virtual void EnterState(FStateTreeExecutionContext& Context, const EStateTreeStateChangeType ChangeType, const FStateTreeTransitionResult& Transition) override;
-	virtual void ExitState(FStateTreeExecutionContext& Context, const EStateTreeStateChangeType ChangeType, const FStateTreeTransitionResult& Transition) override;
-	virtual void Evaluate(FStateTreeExecutionContext& Context, const EStateTreeEvaluationType EvalType, const float DeltaTime) override;
-	void Reset();
-
-	TStateTreeItemHandle<USmartObjectSubsystem> SmartObjectSubsystemHandle;
-	TStateTreeItemHandle<UMassSignalSubsystem> MassSignalSubsystemHandle;
-	TStateTreeItemHandle<FDataFragment_Transform> EntityTransformHandle;
-	TStateTreeItemHandle<FMassSmartObjectUserFragment> SmartObjectUserHandle;
-	TStateTreeItemHandle<FMassZoneGraphLaneLocationFragment, EStateTreeItemRequirement::Optional> LocationHandle;
-
 	/** The identifier of the search request send by the evaluator to find candidates */
-	UPROPERTY(VisibleAnywhere, Category = SmartObject)
+	UPROPERTY(EditAnywhere, Category = Output)
 	FMassSmartObjectRequestResult SearchRequestResult;
 
 	/** The identifier of the search request send by the evaluator to find candidates */
-	UPROPERTY(VisibleAnywhere, Category = SmartObject)
+	UPROPERTY(EditAnywhere, Category = Output)
 	FMassSmartObjectRequestID SearchRequestID;
 
-	/** The delay that the evaluator should wait before trying to find a smart object after a failed attempt */
-	UPROPERTY(EditAnywhere, Category = SmartObject)
-	float RetryCooldown = 0.f;
-
-	/** The default delay that the evaluator should wait before evaluating again */
-	UPROPERTY(EditAnywhere, Category = SmartObject)
-	float TickInterval = 0.f;
-
-	/** Next update time; evaluator will not do anything when Evaluate gets called before that time */
-	float NextUpdate = 0.f;
-
 	/** Indicates that the result of the candidates search is ready and contains some candidates */
-	UPROPERTY(VisibleAnywhere, Category = SmartObject)
+	UPROPERTY(EditAnywhere, Category = Output)
 	bool bCandidatesFound = false;
 	
 	/** Indicates that an object has been claimed */
-	UPROPERTY(VisibleAnywhere, Category = SmartObject)
+	UPROPERTY(EditAnywhere, Category = Output)
 	bool bClaimed = false;
 
+	/** Next update time; evaluator will not do anything when Evaluate gets called before that time */
+	UPROPERTY()
+	float NextUpdate = 0.f;
+
 	/** Indicates that the query was able to use annotations on zone graph lanes instead of a spatial query. */
+	UPROPERTY()
 	bool bUsingZoneGraphAnnotations = false;
 };
+
+USTRUCT(meta = (DisplayName = "Mass SmartObject Eval"))
+struct MASSAIBEHAVIOR_API FMassStateTreeSmartObjectEvaluator : public FMassStateTreeEvaluatorBase
+{
+	GENERATED_BODY()
+
+protected:
+	virtual bool Link(FStateTreeLinker& Linker) override;
+	virtual const UStruct* GetInstanceDataType() const override { return FMassStateTreeSmartObjectEvaluatorInstanceData::StaticStruct(); }
+	virtual void EnterState(FStateTreeExecutionContext& Context, const EStateTreeStateChangeType ChangeType, const FStateTreeTransitionResult& Transition) const override;
+	virtual void ExitState(FStateTreeExecutionContext& Context, const EStateTreeStateChangeType ChangeType, const FStateTreeTransitionResult& Transition) const override;
+	virtual void Evaluate(FStateTreeExecutionContext& Context, const EStateTreeEvaluationType EvalType, const float DeltaTime) const override;
+	void Reset(FStateTreeExecutionContext& Context) const;
+
+	TStateTreeExternalDataHandle<USmartObjectSubsystem> SmartObjectSubsystemHandle;
+	TStateTreeExternalDataHandle<UMassSignalSubsystem> MassSignalSubsystemHandle;
+	TStateTreeExternalDataHandle<FDataFragment_Transform> EntityTransformHandle;
+	TStateTreeExternalDataHandle<FMassSmartObjectUserFragment> SmartObjectUserHandle;
+	TStateTreeExternalDataHandle<FMassZoneGraphLaneLocationFragment, EStateTreeExternalDataRequirement::Optional> LocationHandle;
+
+	TStateTreeInstanceDataPropertyHandle<FMassSmartObjectRequestResult> SearchRequestResultHandle;
+	TStateTreeInstanceDataPropertyHandle<FMassSmartObjectRequestID> SearchRequestIDHandle;
+	TStateTreeInstanceDataPropertyHandle<bool> CandidatesFoundHandle;
+	TStateTreeInstanceDataPropertyHandle<bool> ClaimedHandle;
+	TStateTreeInstanceDataPropertyHandle<float> NextUpdateHandle;
+	TStateTreeInstanceDataPropertyHandle<bool> UsingZoneGraphAnnotationsHandle;
+
+	/** The delay that the evaluator should wait before trying to find a smart object after a failed attempt */
+	UPROPERTY(EditAnywhere, Category = Parameter)
+	float RetryCooldown = 0.f;
+
+	/** The default delay that the evaluator should wait before evaluating again */
+	UPROPERTY(EditAnywhere, Category = Parameter)
+	float TickInterval = 0.f;
+};
diff --git a/Engine/Plugins/AI/MassAI/Source/MassAIBehavior/Public/Evaluators/MassZoneGraphAnnotationEvaluator.h b/Engine/Plugins/AI/MassAI/Source/MassAIBehavior/Public/Evaluators/MassZoneGraphAnnotationEvaluator.h
--- a/Engine/Plugins/AI/MassAI/Source/MassAIBehavior/Public/Evaluators/MassZoneGraphAnnotationEvaluator.h
+++ b/Engine/Plugins/AI/MassAI/Source/MassAIBehavior/Public/Evaluators/MassZoneGraphAnnotationEvaluator.h
@@ -11,6 +11,15 @@ struct FMassZoneGraphAnnotationTagsFragment;
 /**
  * Evaluator to expose ZoneGraph Annotation Tags for decision making.
  */
+USTRUCT()
+struct MASSAIBEHAVIOR_API FMassZoneGraphAnnotationEvaluatorInstanceData
+{
+	GENERATED_BODY()
+
+	UPROPERTY(EditAnywhere, Category = Output)
+	FZoneGraphTagMask AnnotationTags = FZoneGraphTagMask::None;
+};
+
 USTRUCT(meta = (DisplayName = "ZG Annotation Tags"))
 struct MASSAIBEHAVIOR_API FMassZoneGraphAnnotationEvaluator : public FMassStateTreeEvaluatorBase
 {
@@ -20,10 +29,10 @@ struct MASSAIBEHAVIOR_API FMassZoneGraphAnnotationEvaluator : public FMassStateT
 
 protected:
 	virtual bool Link(FStateTreeLinker& Linker) override;
-	virtual void Evaluate(FStateTreeExecutionContext& Context, const EStateTreeEvaluationType EvalType, const float DeltaTime) override;
+	virtual const UStruct* GetInstanceDataType() const override { return FMassZoneGraphAnnotationEvaluatorInstanceData::StaticStruct(); }
+	virtual void Evaluate(FStateTreeExecutionContext& Context, const EStateTreeEvaluationType EvalType, const float DeltaTime) const override;
 
-	TStateTreeItemHandle<FMassZoneGraphAnnotationTagsFragment> BehaviorTagsHandle;
+	TStateTreeExternalDataHandle<FMassZoneGraphAnnotationTagsFragment> AnnotationTagsFragmentHandle;
 
-	UPROPERTY(EditAnywhere, Category = Parameters)
-	FZoneGraphTagMask BehaviorTags = FZoneGraphTagMask::None;
+	TStateTreeInstanceDataPropertyHandle<FZoneGraphTagMask> AnnotationTagsHandle;
 };
diff --git a/Engine/Plugins/AI/MassAI/Source/MassAIBehavior/Public/MassComponentHitTypes.h b/Engine/Plugins/AI/MassAI/Source/MassAIBehavior/Public/MassComponentHitTypes.h
--- a/Engine/Plugins/AI/MassAI/Source/MassAIBehavior/Public/MassComponentHitTypes.h
+++ b/Engine/Plugins/AI/MassAI/Source/MassAIBehavior/Public/MassComponentHitTypes.h
@@ -27,6 +27,8 @@ struct FMassHitResult
 	{
 	}
 
+	bool IsValid() const { return OtherEntity.IsValid(); }
+	
 	FMassEntityHandle OtherEntity;
 
 	/** Time when first hit was received. */
@@ -35,19 +37,3 @@ struct FMassHitResult
 	/** Time used for filtering frequent hits. */
 	float LastFilteredHitTime = 0.f;
 };
-
-USTRUCT()
-struct FComponentHitStateTreeResult : public FStateTreeResult
-{
-	GENERATED_BODY()
-
-protected:
-	virtual const UScriptStruct& GetStruct() const override { return *StaticStruct(); }
-
-public:
-	UPROPERTY()
-	FMassHitResult MassHitResult;
-
-	UPROPERTY()
-	EStateTreeResultStatus Status = EStateTreeResultStatus::Unset;
-};
diff --git a/Engine/Plugins/AI/MassAI/Source/MassAIBehavior/Public/Tasks/MassLookAtTask.h b/Engine/Plugins/AI/MassAI/Source/MassAIBehavior/Public/Tasks/MassLookAtTask.h
--- a/Engine/Plugins/AI/MassAI/Source/MassAIBehavior/Public/Tasks/MassLookAtTask.h
+++ b/Engine/Plugins/AI/MassAI/Source/MassAIBehavior/Public/Tasks/MassLookAtTask.h
@@ -12,6 +12,24 @@ class UMassSignalSubsystem;
 /**
  * Task to assign a LookAt target for mass processing
  */
+USTRUCT()
+struct MASSAIBEHAVIOR_API FMassLookAtTaskInstanceData
+{
+	GENERATED_BODY()
+
+	/** Entity to set as the target for the LookAt behavior. */
+	UPROPERTY(EditAnywhere, Category = Input, meta = (Optional))
+	FMassEntityHandle TargetEntity;
+ 
+	/** Delay before the task ends. Default (0 or any negative) will run indefinitely so it requires a transition in the state tree to stop it. */
+	UPROPERTY(EditAnywhere, Category = Parameter)
+	float Duration = 0.f;
+
+	/** Accumulated time used to stop task if duration is set */
+	UPROPERTY()
+	float Time = 0.f;
+};
+
 USTRUCT(meta = (DisplayName = "Mass LookAt Task"))
 struct MASSAIBEHAVIOR_API FMassLookAtTask : public FMassStateTreeTaskBase
 {
@@ -19,41 +37,35 @@ struct MASSAIBEHAVIOR_API FMassLookAtTask : public FMassStateTreeTaskBase
 
 protected:
 	virtual bool Link(FStateTreeLinker& Linker) override;
-	virtual EStateTreeRunStatus EnterState(FStateTreeExecutionContext& Context, const EStateTreeStateChangeType ChangeType, const FStateTreeTransitionResult& Transition) override;
-	virtual void ExitState(FStateTreeExecutionContext& Context, const EStateTreeStateChangeType ChangeType, const FStateTreeTransitionResult& Transition) override;
-	virtual EStateTreeRunStatus Tick(FStateTreeExecutionContext& Context, const float DeltaTime) override;
+	virtual const UStruct* GetInstanceDataType() const override { return FMassLookAtTaskInstanceData::StaticStruct(); }
+	virtual EStateTreeRunStatus EnterState(FStateTreeExecutionContext& Context, const EStateTreeStateChangeType ChangeType, const FStateTreeTransitionResult& Transition) const override;
+	virtual void ExitState(FStateTreeExecutionContext& Context, const EStateTreeStateChangeType ChangeType, const FStateTreeTransitionResult& Transition) const override;
+	virtual EStateTreeRunStatus Tick(FStateTreeExecutionContext& Context, const float DeltaTime) const override;
 
-	TStateTreeItemHandle<UMassSignalSubsystem> MassSignalSubsystemHandle;
-	TStateTreeItemHandle<FMassLookAtFragment> LookAtHandle;
+	TStateTreeExternalDataHandle<UMassSignalSubsystem> MassSignalSubsystemHandle;
+	TStateTreeExternalDataHandle<FMassLookAtFragment> LookAtHandle;
 
-	/** Delay before the task ends. Default (0 or any negative) will run indefinitely so it requires a transition in the state tree to stop it. */
-	UPROPERTY(EditAnywhere, Category = Parameters, meta = (Bindable))
-	float Duration = 0.f;
+	TStateTreeInstanceDataPropertyHandle<float> DurationHandle;
+	TStateTreeInstanceDataPropertyHandle<FMassEntityHandle> TargetEntityHandle;
+	TStateTreeInstanceDataPropertyHandle<float> TimeHandle;
 
 	/** Look At Mode */
-	UPROPERTY(EditAnywhere, Category = Params, meta = (Bindable))
+	UPROPERTY(EditAnywhere, Category = Parameter)
 	EMassLookAtMode LookAtMode = EMassLookAtMode::LookForward; 
 
-	/** Entity to set as the target for the LookAt behavior. */
-	UPROPERTY(VisibleAnywhere, Category = Params, meta = (Bindable))
-	FMassEntityHandle TargetEntity;
-
 	/** Random gaze Mode */
-	UPROPERTY(EditAnywhere, Category = Params, meta = (Bindable))
+	UPROPERTY(EditAnywhere, Category = Parameter)
 	EMassLookAtGazeMode RandomGazeMode = EMassLookAtGazeMode::None;
 	
 	/** Random gaze yaw angle added to the look direction determined by the look at mode. */
-	UPROPERTY(EditAnywhere, Category = Params, meta = (Bindable, UIMin = 0.0, ClampMin = 0.0, UIMax = 180.0, ClampMax = 180.0))
+	UPROPERTY(EditAnywhere, Category = Parameter, meta = (UIMin = 0.0, ClampMin = 0.0, UIMax = 180.0, ClampMax = 180.0))
 	uint8 RandomGazeYawVariation = 0;
 
 	/** Random gaze pitch angle added to the look direction determined by the look at mode. */
-	UPROPERTY(EditAnywhere, Category = Params, meta = (Bindable, UIMin = 0.0, ClampMin = 0.0, UIMax = 180.0, ClampMax = 180.0))
+	UPROPERTY(EditAnywhere, Category = Parameter, meta = (UIMin = 0.0, ClampMin = 0.0, UIMax = 180.0, ClampMax = 180.0))
 	uint8 RandomGazePitchVariation = 0;
 
 	/** If true, allow random gaze to look at other entities too. */
-	UPROPERTY(EditAnywhere, Category = Params, meta = (Bindable))
+	UPROPERTY(EditAnywhere, Category = Parameter)
 	bool bRandomGazeEntities = false;
-
-	/** Accumulated time used to stop task if duration is set */
-	float Time = 0.f;
 };
diff --git a/Engine/Plugins/AI/MassAI/Source/MassAIBehavior/Public/Tasks/MassStateTreeSmartObjectTask.h b/Engine/Plugins/AI/MassAI/Source/MassAIBehavior/Public/Tasks/MassStateTreeSmartObjectTask.h
--- a/Engine/Plugins/AI/MassAI/Source/MassAIBehavior/Public/Tasks/MassStateTreeSmartObjectTask.h
+++ b/Engine/Plugins/AI/MassAI/Source/MassAIBehavior/Public/Tasks/MassStateTreeSmartObjectTask.h
@@ -17,6 +17,20 @@ struct FMassMoveTargetFragment;
 /**
  * Tasks to claim a smart object from search results and release it when done.
  */
+USTRUCT()
+struct MASSAIBEHAVIOR_API FMassClaimSmartObjectTaskInstanceData
+{
+	GENERATED_BODY()
+
+	/** Result of the candidates search request (Input) */
+	UPROPERTY(VisibleAnywhere, Category = Input)
+	FMassSmartObjectRequestResult SearchRequestResult;
+
+	/** Result of the claim on potential candidates from the search results (Output) */
+	UPROPERTY(VisibleAnywhere, Category = Output)
+	EMassSmartObjectClaimResult ClaimResult = EMassSmartObjectClaimResult::Unset;
+};
+
 USTRUCT(meta = (DisplayName = "Claim SmartObject"))
 struct MASSAIBEHAVIOR_API FMassClaimSmartObjectTask : public FMassStateTreeTaskBase
 {
@@ -24,25 +38,26 @@ struct MASSAIBEHAVIOR_API FMassClaimSmartObjectTask : public FMassStateTreeTaskB
 
 protected:
 	virtual bool Link(FStateTreeLinker& Linker) override;
-	virtual EStateTreeRunStatus EnterState(FStateTreeExecutionContext& Context, const EStateTreeStateChangeType ChangeType, const FStateTreeTransitionResult& Transition) override;
-	virtual void ExitState(FStateTreeExecutionContext& Context, const EStateTreeStateChangeType ChangeType, const FStateTreeTransitionResult& Transition) override;
-	virtual EStateTreeRunStatus Tick(FStateTreeExecutionContext& Context, const float DeltaTime) override;
+	virtual const UStruct* GetInstanceDataType() const override { return FMassClaimSmartObjectTaskInstanceData::StaticStruct(); }
+	virtual EStateTreeRunStatus EnterState(FStateTreeExecutionContext& Context, const EStateTreeStateChangeType ChangeType, const FStateTreeTransitionResult& Transition) const override;
+	virtual void ExitState(FStateTreeExecutionContext& Context, const EStateTreeStateChangeType ChangeType, const FStateTreeTransitionResult& Transition) const override;
 
-	TStateTreeItemHandle<FMassSmartObjectUserFragment> SmartObjectUserHandle;
-	TStateTreeItemHandle<USmartObjectSubsystem> SmartObjectSubsystemHandle;
+	TStateTreeExternalDataHandle<FMassSmartObjectUserFragment> SmartObjectUserHandle;
+	TStateTreeExternalDataHandle<USmartObjectSubsystem> SmartObjectSubsystemHandle;
 
-	/** Result of the candidates search request (Input) */
-	UPROPERTY(VisibleAnywhere, Category = SmartObject, meta = (Bindable), Transient)
-	FMassSmartObjectRequestResult SearchRequestResult;
-
-	/** Result of the claim on potential candidates from the search results (Output) */
-	UPROPERTY(VisibleAnywhere, Category = SmartObject, Transient)
-	EMassSmartObjectClaimResult ClaimResult = EMassSmartObjectClaimResult::Unset;
+	TStateTreeInstanceDataPropertyHandle<FMassSmartObjectRequestResult> SearchRequestResultHandle;
+	TStateTreeInstanceDataPropertyHandle<EMassSmartObjectClaimResult> ClaimResultHandle;
 };
 
 /**
  * Task to tell an entity to start using a claimed smart object.
  */
+USTRUCT()
+struct MASSAIBEHAVIOR_API FMassUseSmartObjectTaskInstanceData
+{
+	GENERATED_BODY()
+};
+
 USTRUCT(meta = (DisplayName = "Mass Use SmartObject Task"))
 struct MASSAIBEHAVIOR_API FMassUseSmartObjectTask : public FMassStateTreeTaskBase
 {
@@ -50,18 +65,19 @@ struct MASSAIBEHAVIOR_API FMassUseSmartObjectTask : public FMassStateTreeTaskBas
 
 protected:
 	virtual bool Link(FStateTreeLinker& Linker) override;
-	virtual EStateTreeRunStatus EnterState(FStateTreeExecutionContext& Context, const EStateTreeStateChangeType ChangeType, const FStateTreeTransitionResult& Transition) override;
-	virtual void ExitState(FStateTreeExecutionContext& Context, const EStateTreeStateChangeType ChangeType, const FStateTreeTransitionResult& Transition) override;
-	virtual void StateCompleted(FStateTreeExecutionContext& Context, const EStateTreeRunStatus CompletionStatus, const FStateTreeHandle CompletedState) override;
-	virtual EStateTreeRunStatus Tick(FStateTreeExecutionContext& Context, const float DeltaTime) override;
+	virtual const UStruct* GetInstanceDataType() const override { return FMassUseSmartObjectTaskInstanceData::StaticStruct(); }
+	virtual EStateTreeRunStatus EnterState(FStateTreeExecutionContext& Context, const EStateTreeStateChangeType ChangeType, const FStateTreeTransitionResult& Transition) const override;
+	virtual void ExitState(FStateTreeExecutionContext& Context, const EStateTreeStateChangeType ChangeType, const FStateTreeTransitionResult& Transition) const override;
+	virtual void StateCompleted(FStateTreeExecutionContext& Context, const EStateTreeRunStatus CompletionStatus, const FStateTreeHandle CompletedState) const override;
+	virtual EStateTreeRunStatus Tick(FStateTreeExecutionContext& Context, const float DeltaTime) const override;
 
-	TStateTreeItemHandle<USmartObjectSubsystem> SmartObjectSubsystemHandle;
-	TStateTreeItemHandle<FDataFragment_Transform> EntityTransformHandle;
-	TStateTreeItemHandle<FMassSmartObjectUserFragment> SmartObjectUserHandle;
-	TStateTreeItemHandle<FMassMoveTargetFragment> MoveTargetHandle;
+	TStateTreeExternalDataHandle<USmartObjectSubsystem> SmartObjectSubsystemHandle;
+	TStateTreeExternalDataHandle<FDataFragment_Transform> EntityTransformHandle;
+	TStateTreeExternalDataHandle<FMassSmartObjectUserFragment> SmartObjectUserHandle;
+	TStateTreeExternalDataHandle<FMassMoveTargetFragment> MoveTargetHandle;
 
 	/** Delay in seconds before trying to find & use another smart object */
-	UPROPERTY(EditAnywhere, Category = SmartObject)
+	UPROPERTY(EditAnywhere, Category = Parameter)
 	float Cooldown = 0.f;
 };
 
