GameplayDebugger Mass category will now read UE::Mass::Debug's entity index range and pick entities from that range, if no debug actor is picked.
Added a couple of debug-time functions to support that.
Also cleaned up the entity selection code in the debugger category and massaged the code a tiny bit (like painting disabled options gray).

diff --git a/Engine/Plugins/AI/MassAI/Source/MassAIDebug/Private/GameplayDebuggerCategory_Mass.cpp b/Engine/Plugins/AI/MassAI/Source/MassAIDebug/Private/GameplayDebuggerCategory_Mass.cpp
--- a/Engine/Plugins/AI/MassAI/Source/MassAIDebug/Private/GameplayDebuggerCategory_Mass.cpp
+++ b/Engine/Plugins/AI/MassAI/Source/MassAIDebug/Private/GameplayDebuggerCategory_Mass.cpp
@@ -3,6 +3,7 @@
 #include "GameplayDebuggerCategory_Mass.h"
 
 #if WITH_GAMEPLAY_DEBUGGER && WITH_MASSGAMEPLAY_DEBUG
+#include "MassGameplayDebugTypes.h"
 #include "MassEntityView.h"
 #include "GameplayDebuggerCategoryReplicator.h"
 #include "GameplayDebuggerPlayerManager.h"
@@ -21,6 +22,61 @@
 #include "CanvasItem.h"
 #include "Engine/World.h"
 
+namespace UE::Mass::Debug
+{
+	FMassEntityHandle GetEntityFromActor(const AActor& Actor, const UMassAgentComponent*& OutMassAgentComponent)
+	{
+		FMassEntityHandle EntityHandle;
+		if (UMassAgentComponent* AgentComp = Actor.FindComponentByClass<UMassAgentComponent>())
+		{
+			EntityHandle = AgentComp->GetEntityHandle();
+			OutMassAgentComponent = AgentComp;
+		}
+		else if (UMassActorSubsystem* ActorManager = UWorld::GetSubsystem<UMassActorSubsystem>(Actor.GetWorld()))
+		{
+			EntityHandle = ActorManager->GetEntityHandleFromActor(&Actor);
+		}
+		return EntityHandle;
+	};
+	
+	FMassEntityHandle GetBestEntity(const FVector ViewLocation, const FVector ViewDirection, const TConstArrayView<FMassEntityHandle> Entities, const TConstArrayView<FVector> Locations, const bool bLimitAngle)
+	{
+		// Reusing similar algorithm as UGameplayDebuggerLocalController for now 
+		constexpr float MaxScanDistanceSq = 25000.0f * 25000.0f;
+		constexpr float MinViewDirDot = 0.707f; // 45 degrees
+
+		checkf(Entities.Num() == Locations.Num(), TEXT("Both Entities and Locations lists are expected to be of the same size: %d vs %d"), Entities.Num(), Locations.Num());
+		
+		float BestScore = bLimitAngle ? MinViewDirDot : (-1.f - KINDA_SMALL_NUMBER);	
+		FMassEntityHandle BestEntity;
+
+		for (int i = 0; i < Entities.Num(); ++i)
+		{
+			if (Entities[i].IsSet() == false)
+			{
+				continue;
+			}
+			
+			const FVector DirToEntity = (Locations[i] - ViewLocation);
+			const float DistToEntitySq = DirToEntity.SizeSquared();
+			if (DistToEntitySq > MaxScanDistanceSq)
+			{
+				continue;
+			}
+
+			const FVector DirToEntityNormal = (FMath::IsNearlyZero(DistToEntitySq)) ? ViewDirection : (DirToEntity / FMath::Sqrt(DistToEntitySq));
+			const float ViewDot = FVector::DotProduct(ViewDirection, DirToEntityNormal);
+			if (ViewDot > BestScore)
+			{
+				BestScore = ViewDot;
+				BestEntity = Entities[i];
+			}
+		}
+
+		return BestEntity;
+	}
+} // namespace UE::Mass:Debug
+
 //----------------------------------------------------------------------//
 //  FGameplayDebuggerCategory_Mass
 //----------------------------------------------------------------------//
@@ -38,6 +94,7 @@ FGameplayDebuggerCategory_Mass::FGameplayDebuggerCategory_Mass()
 	bShowNearEntityOverview = true;
 	bShowNearEntityAvoidance = false;
 	bShowNearEntityPath = false;
+	bMarkEntityBeingDebugged = true;
 
 	BindKeyPress(EKeys::A.GetFName(), FGameplayDebuggerInputModifier::Shift, this, &FGameplayDebuggerCategory_Mass::OnToggleArchetypes, EGameplayDebuggerInputMode::Replicated);
 	BindKeyPress(EKeys::S.GetFName(), FGameplayDebuggerInputModifier::Shift, this, &FGameplayDebuggerCategory_Mass::OnToggleShapes, EGameplayDebuggerInputMode::Replicated);
@@ -55,56 +112,40 @@ void FGameplayDebuggerCategory_Mass::SetCachedEntity(FMassEntityHandle Entity, U
 	Debugger.SetSelectedEntity(Entity);
 }
 
-void FGameplayDebuggerCategory_Mass::PickEntity(const APlayerController& OwnerPC, UWorld& World, UMassDebuggerSubsystem& Debugger)
+void FGameplayDebuggerCategory_Mass::PickEntity(const APlayerController& OwnerPC, UWorld& World, UMassDebuggerSubsystem& Debugger, const bool bLimitAngle)
 {
 	FVector ViewLocation = FVector::ZeroVector;
 	FVector ViewDirection = FVector::ForwardVector;
 	AGameplayDebuggerPlayerManager::GetViewPoint(OwnerPC, ViewLocation, ViewDirection);
 
-	// Reusing similar algorithm as UGameplayDebuggerLocalController for now 
-	const float MaxScanDistance = 25000.0f;
-	const float MinViewDirDot = 0.707f; // 45 degrees
-	
-	float BestScore = MinViewDirDot;
 	FMassEntityHandle BestEntity;
-	FVector BestLocation = FVector::ZeroVector;
-
-	TConstArrayView<FMassEntityHandle> Entities = Debugger.GetEntities();
-	TConstArrayView<FVector> Locations = Debugger.GetLocations();
-	checkf(Entities.Num() == Locations.Num(), TEXT("Both Entities and Locations lists are expected to be of the same size: %d vs %d"), Entities.Num(), Locations.Num());
-
-	for (int32 i = 0; i < Locations.Num(); ++i)
+	// entities indicated by UE::Mass::Debug take precedence 
+    if (UE::Mass::Debug::HasDebugEntities())
+    {
+	    if (const UMassEntitySubsystem* EntitySystem = UWorld::GetSubsystem<UMassEntitySubsystem>(&World))
+	    {
+	    	TArray<FMassEntityHandle> Entities;
+	    	TArray<FVector> Locations;
+	    	UE::Mass::Debug::GetDebugEntitiesAndLocations(*EntitySystem, Entities, Locations);
+	    	BestEntity = UE::Mass::Debug::GetBestEntity(ViewLocation, ViewDirection, Entities, Locations, bLimitAngle);
+	    }
+    }
+	else
 	{
-		const FVector& Location = Locations[i];
-		FVector DirToEntity = (Location - ViewLocation);
-		const float DistToEntity = DirToEntity.Size();
-		if (DistToEntity > MaxScanDistance)
-		{
-			continue;
-		}
-
-		DirToEntity = (FMath::IsNearlyZero(DistToEntity)) ? ViewDirection : DirToEntity /= DistToEntity;
-		const float ViewDot = FVector::DotProduct(ViewDirection, DirToEntity);
-		if (ViewDot > BestScore)
-		{
-			BestScore = ViewDot;
-			BestEntity = Entities[i];
-			BestLocation = Location;
-		}
+		BestEntity = UE::Mass::Debug::GetBestEntity(ViewLocation, ViewDirection, Debugger.GetEntities(), Debugger.GetLocations(), bLimitAngle);
 	}
 
 	AActor* BestActor = nullptr;
 	if (BestEntity.IsSet())
 	{
-		// Use this new entity
-		SetCachedEntity(BestEntity, Debugger);
-		UMassActorSubsystem* ActorManager = World.GetSubsystem<UMassActorSubsystem>();
-		if (ActorManager != nullptr)
+		if (const UMassActorSubsystem* ActorManager = World.GetSubsystem<UMassActorSubsystem>())
 		{
-			BestActor = ActorManager->GetActorFromHandle(FMassEntityHandle(CachedEntity));
-			CachedDebugActor = BestActor;
+			BestActor = ActorManager->GetActorFromHandle(FMassEntityHandle(BestEntity));
 		}
 	}
+
+	SetCachedEntity(BestEntity, Debugger);
+	CachedDebugActor = BestActor;
 	GetReplicator()->SetDebugActor(BestActor);
 }
 
@@ -125,6 +166,18 @@ void FGameplayDebuggerCategory_Mass::CollectData(APlayerController* OwnerPC, AAc
 	}
 	Debugger->SetCollectingData();
 
+	const UMassAgentComponent* AgentComp = nullptr;
+	if (DebugActor)
+	{
+		const FMassEntityHandle EntityHandle = UE::Mass::Debug::GetEntityFromActor(*DebugActor, AgentComp);	
+		SetCachedEntity(EntityHandle, *Debugger);
+		CachedDebugActor = DebugActor;
+	}
+	else if (CachedDebugActor)
+	{
+		SetCachedEntity(FMassEntityHandle(), *Debugger);
+		CachedDebugActor = nullptr;
+	}
 	// Ideally we would have a way to register in the main picking flow but that would require more changes to
 	// also support client-server picking. For now, we handle explicit mass picking requests on the authority
 	if (bPickEntity)
@@ -132,37 +185,11 @@ void FGameplayDebuggerCategory_Mass::CollectData(APlayerController* OwnerPC, AAc
 		PickEntity(*OwnerPC, *World, *Debugger);
 		bPickEntity = false;
 	}
-
-	auto GetEntityFromActorFunc = [](const AActor& Actor, const UMassAgentComponent** OutMassAgentComponent = nullptr)
+	// if we're debugging based on UE::Mass::Debug and the range changed
+	else if (CachedDebugActor == nullptr && UE::Mass::Debug::HasDebugEntities() && UE::Mass::Debug::IsDebuggingEntity(CachedEntity) == false)
 	{
-		FMassEntityHandle EntityHandle;
-		UMassAgentComponent* AgentComp = Actor.FindComponentByClass<UMassAgentComponent>();
-		if (AgentComp)
-		{
-			EntityHandle = AgentComp->GetEntityHandle();
-			if (OutMassAgentComponent)
-			{
-				*OutMassAgentComponent = AgentComp;
-			}
-		}
-		else
-		{
-			UMassActorSubsystem* ActorManager = UWorld::GetSubsystem<UMassActorSubsystem>(Actor.GetWorld());
-			if (ActorManager != nullptr)
-			{
-				EntityHandle = ActorManager->GetEntityHandleFromActor(&Actor);
-			}
-		}
-		return EntityHandle;
-	};
-
-	if (CachedDebugActor != DebugActor)
-	{
-		CachedDebugActor = DebugActor;
-		if (DebugActor != nullptr)
-		{
-			SetCachedEntity(GetEntityFromActorFunc(*DebugActor), *Debugger);
-		}
+		// using bLimitAngle = false to not limit the selection to only the things in from of the player
+		PickEntity(*OwnerPC, *World, *Debugger, /*bLimitAngle=*/false);
 	}
 
 	UMassEntitySubsystem* EntitySystem = UWorld::GetSubsystem<UMassEntitySubsystem>(World);
@@ -171,6 +198,13 @@ void FGameplayDebuggerCategory_Mass::CollectData(APlayerController* OwnerPC, AAc
 		AddTextLine(FString::Printf(TEXT("{Green}Entities count active{grey}/all: {white}%d{grey}/%d"), EntitySystem->DebugGetEntityCount(), EntitySystem->DebugGetEntityCount()));
 		AddTextLine(FString::Printf(TEXT("{Green}Registered Archetypes count: {white}%d {green}data ver: {white}%d"), EntitySystem->DebugGetArchetypesCount(), EntitySystem->GetArchetypeDataVersion()));
 
+		if (UE::Mass::Debug::HasDebugEntities())
+		{
+			int32 RangeBegin, RangeEnd;
+			UE::Mass::Debug::GetDebugEntitiesRange(RangeBegin, RangeEnd);
+			AddTextLine(FString::Printf(TEXT("{Green}Debugged entity range: {orange}%d-%d"), RangeBegin, RangeEnd));
+		}
+
 		if (bShowArchetypes)
 		{
 			FStringOutputDevice Ar;
@@ -179,6 +213,16 @@ void FGameplayDebuggerCategory_Mass::CollectData(APlayerController* OwnerPC, AAc
 
 			AddTextLine(Ar);
 		}
+
+		if (CachedEntity.IsSet() && bMarkEntityBeingDebugged)
+		{
+			if (const FDataFragment_Transform* TransformFragment = EntitySystem->GetFragmentDataPtr<FDataFragment_Transform>(CachedEntity))
+			{
+				const FVector Location = TransformFragment->GetTransform().GetLocation();
+				AddShape(FGameplayDebuggerShape::MakeBox(Location, FVector(3,3,250), FColor::Purple));
+				AddShape(FGameplayDebuggerShape::MakePoint(Location, 10, FColor::Purple));
+			}
+		}
 	}
 	else
 	{
@@ -222,12 +266,6 @@ void FGameplayDebuggerCategory_Mass::CollectData(APlayerController* OwnerPC, AAc
 
 	if (bShowAgentFragments)
 	{
-		const UMassAgentComponent* AgentComp = nullptr;
-		if (!CachedEntity.IsSet() && DebugActor != nullptr)
-		{
-			SetCachedEntity(GetEntityFromActorFunc(*DebugActor, &AgentComp), *Debugger);
-		}
-
 		if (CachedEntity.IsSet() && EntitySystem)
 		{
 			// CachedEntity can become invalid if the entity "dies" or in editor mode when related actor gets moved 
@@ -368,7 +406,7 @@ void FGameplayDebuggerCategory_Mass::CollectData(APlayerController* OwnerPC, AAc
 
 					const float Height = 180.0f; // @todo: add height to agent.
 					const float EyeHeight = 160.0f; // @todo: add eye height to agent.
-					
+
 					// Draw entity position and orientation.
 					FVector BasePos = EntityLocation + FVector(0.0f ,0.0f ,25.0f );
 
@@ -397,13 +435,13 @@ void FGameplayDebuggerCategory_Mass::CollectData(APlayerController* OwnerPC, AAc
 							FVector TargetPosition = TargetTransform->GetTransform().GetLocation();
 							TargetPosition.Z = BasePos.Z;
 							AddShape(FGameplayDebuggerShape::MakeCircle(TargetPosition, FVector::UpVector, Radius.Radius, FColor::Red));
-							
+
 							const float TargetDistance = FMath::Max(LookArrowLength, FVector::DotProduct(WorldLookDirection, TargetPosition - BasePos));
 							AddShape(FGameplayDebuggerShape::MakeSegment(BasePos, BasePos + WorldLookDirection * TargetDistance, FColorList::LightGrey));
 							bLookArrowDrawn = true;
 						}
 					}
-					
+
 					if (LookAt.bRandomGazeEntities && EntitySystem->IsEntityValid(LookAt.GazeTrackedEntity))
 					{
 						if (const FDataFragment_Transform* TargetTransform = EntitySystem->GetFragmentDataPtr<FDataFragment_Transform>(LookAt.GazeTrackedEntity))
@@ -588,8 +626,15 @@ void FGameplayDebuggerCategory_Mass::DrawData(APlayerController* OwnerPC, FGamep
 	CanvasContext.Printf(TEXT("\n[{yellow}%s{white}] %s Archetypes"), *GetInputHandlerDescription(0), bShowArchetypes ? TEXT("Hide") : TEXT("Show"));
 	CanvasContext.Printf(TEXT("[{yellow}%s{white}] %s Shapes"), *GetInputHandlerDescription(1), bShowShapes ? TEXT("Hide") : TEXT("Show"));
 	CanvasContext.Printf(TEXT("[{yellow}%s{white}] %s Agent Fragments"), *GetInputHandlerDescription(2), bShowAgentFragments ? TEXT("Hide") : TEXT("Show"));
+	if (bShowAgentFragments)
+	{
+		CanvasContext.Printf(TEXT("[{yellow}%s{white}] %s Entity details"), *GetInputHandlerDescription(4), bShowEntityDetails ? TEXT("Hide") : TEXT("Show"));
+	}
+	else
+	{
+		CanvasContext.Printf(TEXT("{grey}[%s] Entity details [enable Agent Fragments]{white}"), *GetInputHandlerDescription(4));
+	}
 	CanvasContext.Printf(TEXT("[{yellow}%s{white}] Pick Entity"), *GetInputHandlerDescription(3));
-	CanvasContext.Printf(TEXT("[{yellow}%s{white}] %s Entity details"), *GetInputHandlerDescription(4), bShowEntityDetails ? TEXT("Hide") : TEXT("Show"));
 	CanvasContext.Printf(TEXT("[{yellow}%s{white}] %s Entity overview"), *GetInputHandlerDescription(5), bShowNearEntityOverview ? TEXT("Hide") : TEXT("Show"));
 	CanvasContext.Printf(TEXT("[{yellow}%s{white}] %s Entity avoidance"), *GetInputHandlerDescription(6), bShowNearEntityAvoidance ? TEXT("Hide") : TEXT("Show"));
 	CanvasContext.Printf(TEXT("[{yellow}%s{white}] %s Entity path"), *GetInputHandlerDescription(7), bShowNearEntityPath ? TEXT("Hide") : TEXT("Show"));
diff --git a/Engine/Plugins/AI/MassAI/Source/MassAIDebug/Public/GameplayDebuggerCategory_Mass.h b/Engine/Plugins/AI/MassAI/Source/MassAIDebug/Public/GameplayDebuggerCategory_Mass.h
--- a/Engine/Plugins/AI/MassAI/Source/MassAIDebug/Public/GameplayDebuggerCategory_Mass.h
+++ b/Engine/Plugins/AI/MassAI/Source/MassAIDebug/Public/GameplayDebuggerCategory_Mass.h
@@ -35,7 +35,7 @@ protected:
 	void OnToggleNearEntityAvoidance() { bShowNearEntityAvoidance = !bShowNearEntityAvoidance; }
 	void OnToggleNearEntityPath() { bShowNearEntityPath = !bShowNearEntityPath; }
 	
-	void PickEntity(const APlayerController& OwnerPC, UWorld& World, UMassDebuggerSubsystem& Debugger);
+	void PickEntity(const APlayerController& OwnerPC, UWorld& World, UMassDebuggerSubsystem& Debugger, const bool bLimitAngle = true);
 
 protected:
 	AActor* CachedDebugActor;
@@ -48,6 +48,7 @@ protected:
 	bool bShowNearEntityOverview;
 	bool bShowNearEntityAvoidance;
 	bool bShowNearEntityPath;
+	bool bMarkEntityBeingDebugged;
 
 	struct FEntityDescription
 	{
diff --git a/Engine/Plugins/Runtime/MassGameplay/Source/MassCommon/Private/MassCommonTypes.cpp b/Engine/Plugins/Runtime/MassGameplay/Source/MassCommon/Private/MassCommonTypes.cpp
--- a/Engine/Plugins/Runtime/MassGameplay/Source/MassCommon/Private/MassCommonTypes.cpp
+++ b/Engine/Plugins/Runtime/MassGameplay/Source/MassCommon/Private/MassCommonTypes.cpp
@@ -83,15 +83,11 @@ bool HasDebugEntities()
 	return DebugEntityBegin != INDEX_NONE && DebugEntityEnd != INDEX_NONE;
 }
 
-void GetDebugEntitiesRange(int32& OutBegin, int32& OutEnd)
+bool GetDebugEntitiesRange(int32& OutBegin, int32& OutEnd)
 {
 	OutBegin = DebugEntityBegin;
 	OutEnd = DebugEntityEnd;
-}
-	
-FMassEntityHandle ConvertEntityIndexToHandle(const UMassEntitySubsystem& EntitySystem, const int32 EntityIndex)
-{
-	return EntitySystem.DebugGetEntityIndexHandle(EntityIndex);
+	return DebugEntityBegin != INDEX_NONE && DebugEntityEnd != INDEX_NONE;  
 }
 	
 bool IsDebuggingEntity(FMassEntityHandle Entity, FColor* OutEntityColor)
diff --git a/Engine/Plugins/Runtime/MassGameplay/Source/MassCommon/Public/MassCommonTypes.h b/Engine/Plugins/Runtime/MassGameplay/Source/MassCommon/Public/MassCommonTypes.h
--- a/Engine/Plugins/Runtime/MassGameplay/Source/MassCommon/Public/MassCommonTypes.h
+++ b/Engine/Plugins/Runtime/MassGameplay/Source/MassCommon/Public/MassCommonTypes.h
@@ -18,8 +18,12 @@ class UMassEntitySubsystem;
 namespace UE::Mass::Debug
 {
 	MASSCOMMON_API extern bool HasDebugEntities();
-	MASSCOMMON_API extern void GetDebugEntitiesRange(int32& OutBegin, int32& OutEnd);
-	MASSCOMMON_API extern FMassEntityHandle ConvertEntityIndexToHandle(const UMassEntitySubsystem& EntitySystem, const int32 EntityIndex);
+	/**
+	 * Populates OutBegin and OutEnd with entity index ranges as set by ai.debug.mass.SetDebugEntityRange or
+	 * ai.debug.mass.DebugEntity console commands.
+	 * @return whether any range has been configured.
+	 */
+	MASSCOMMON_API extern bool GetDebugEntitiesRange(int32& OutBegin, int32& OutEnd);
 	MASSCOMMON_API extern bool IsDebuggingEntity(FMassEntityHandle Entity, FColor* OutEntityColor = nullptr);
 	MASSCOMMON_API extern FColor GetEntityDebugColor(FMassEntityHandle Entity);
 } // namespace UE::Mass::Debug
diff --git a/Engine/Plugins/Runtime/MassGameplay/Source/MassGameplayDebug/Private/MassGameplayDebugTypes.cpp b/Engine/Plugins/Runtime/MassGameplay/Source/MassGameplayDebug/Private/MassGameplayDebugTypes.cpp
new file mode 100644
--- /dev/null
+++ b/Engine/Plugins/Runtime/MassGameplay/Source/MassGameplayDebug/Private/MassGameplayDebugTypes.cpp
@@ -0,0 +1,45 @@
+// Copyright Epic Games, Inc. All Rights Reserved.
+
+#include "MassGameplayDebugTypes.h"
+#include "MassCommonTypes.h"
+#include "MassCommonFragments.h"
+#include "MassEntitySubsystem.h"
+
+
+#if WITH_MASSGAMEPLAY_DEBUG
+
+namespace UE::Mass::Debug
+{
+
+void GetDebugEntitiesAndLocations(const UMassEntitySubsystem& EntitySystem, TArray<FMassEntityHandle>& OutEntities, TArray<FVector>& OutLocations)
+{
+	int32 DebugEntityEnd, DebugEntityBegin;
+	if (GetDebugEntitiesRange(DebugEntityEnd, DebugEntityBegin) == false)
+	{
+		return;
+	}
+
+	OutEntities.Reserve(DebugEntityEnd - DebugEntityBegin);
+	OutLocations.Reserve(DebugEntityEnd - DebugEntityBegin);
+
+	for (int32 i = DebugEntityBegin; i <= DebugEntityEnd; ++i)
+	{
+		const FMassEntityHandle EntityHandle = ConvertEntityIndexToHandle(EntitySystem, i);
+		if (EntityHandle.IsSet())
+		{
+			if (const FDataFragment_Transform* TransformFragment = EntitySystem.GetFragmentDataPtr<FDataFragment_Transform>(EntityHandle))
+			{
+				OutEntities.Add(EntityHandle);
+				OutLocations.Add(TransformFragment->GetTransform().GetLocation());
+			}
+		}
+	}
+}
+
+FMassEntityHandle ConvertEntityIndexToHandle(const UMassEntitySubsystem& EntitySystem, const int32 EntityIndex)
+{
+	return EntitySystem.DebugGetEntityIndexHandle(EntityIndex);
+}
+
+} // namespace UE::Mass::Debug
+#endif // WITH_MASSGAMEPLAY_DEBUG
\ No newline at end of file
diff --git a/Engine/Plugins/Runtime/MassGameplay/Source/MassGameplayDebug/Public/MassGameplayDebugTypes.h b/Engine/Plugins/Runtime/MassGameplay/Source/MassGameplayDebug/Public/MassGameplayDebugTypes.h
--- a/Engine/Plugins/Runtime/MassGameplay/Source/MassGameplayDebug/Public/MassGameplayDebugTypes.h
+++ b/Engine/Plugins/Runtime/MassGameplay/Source/MassGameplayDebug/Public/MassGameplayDebugTypes.h
@@ -3,10 +3,8 @@
 #pragma once
 
 #include "CoreMinimal.h"
-#include "Stats/Stats.h"
 #include "UObject/ObjectMacros.h"
 #include "Logging/LogMacros.h"
-#include "Templates/SubclassOf.h"
 #include "Engine/DataTable.h"
 #include "MassEntityTypes.h"
 #include "MassSpawnerTypes.h"
@@ -21,6 +19,17 @@ class UBillboardComponent;
 class UStaticMesh;
 class UMaterialInterface;
 
+#if WITH_MASSGAMEPLAY_DEBUG
+namespace UE::Mass::Debug
+{
+	/**
+	 * Fetches entity handles and their locations for entities indicated by index range as set by
+	 * ai.debug.mass.SetDebugEntityRange or ai.debug.mass.DebugEntity console commands.
+	 */
+	MASSGAMEPLAYDEBUG_API extern void GetDebugEntitiesAndLocations(const UMassEntitySubsystem& EntitySystem, TArray<FMassEntityHandle>& OutEntities, TArray<FVector>& OutLocations);
+	MASSGAMEPLAYDEBUG_API extern FMassEntityHandle ConvertEntityIndexToHandle(const UMassEntitySubsystem& EntitySystem, const int32 EntityIndex);
+} // namespace UE::Mass::Debug
+#endif // WITH_MASSGAMEPLAY_DEBUG
 
 USTRUCT()
 struct FSimDebugDataRow : public FTableRowBase
