Fix crash opening maps when gameplay debugger is enabled for editor world in editor settings

#rnx

diff --git a/Engine/Plugins/AI/MassAI/Source/MassAIDebug/Private/GameplayDebuggerCategory_Mass.cpp b/Engine/Plugins/AI/MassAI/Source/MassAIDebug/Private/GameplayDebuggerCategory_Mass.cpp
--- a/Engine/Plugins/AI/MassAI/Source/MassAIDebug/Private/GameplayDebuggerCategory_Mass.cpp
+++ b/Engine/Plugins/AI/MassAI/Source/MassAIDebug/Private/GameplayDebuggerCategory_Mass.cpp
@@ -178,18 +178,22 @@ void FGameplayDebuggerCategory_Mass::CollectData(APlayerController* OwnerPC, AAc
 		SetCachedEntity(FMassEntityHandle(), *Debugger);
 		CachedDebugActor = nullptr;
 	}
-	// Ideally we would have a way to register in the main picking flow but that would require more changes to
-	// also support client-server picking. For now, we handle explicit mass picking requests on the authority
-	if (bPickEntity)
-	{
-		PickEntity(*OwnerPC, *World, *Debugger);
-		bPickEntity = false;
-	}
-	// if we're debugging based on UE::Mass::Debug and the range changed
-	else if (CachedDebugActor == nullptr && UE::Mass::Debug::HasDebugEntities() && UE::Mass::Debug::IsDebuggingEntity(CachedEntity) == false)
+
+	if (OwnerPC)
 	{
-		// using bLimitAngle = false to not limit the selection to only the things in from of the player
-		PickEntity(*OwnerPC, *World, *Debugger, /*bLimitAngle=*/false);
+		// Ideally we would have a way to register in the main picking flow but that would require more changes to
+		// also support client-server picking. For now, we handle explicit mass picking requests on the authority
+		if (bPickEntity)
+		{
+			PickEntity(*OwnerPC, *World, *Debugger);
+			bPickEntity = false;
+		}
+		// if we're debugging based on UE::Mass::Debug and the range changed
+		else if (CachedDebugActor == nullptr && UE::Mass::Debug::HasDebugEntities() && UE::Mass::Debug::IsDebuggingEntity(CachedEntity) == false)
+		{
+			// using bLimitAngle = false to not limit the selection to only the things in from of the player
+			PickEntity(*OwnerPC, *World, *Debugger, /*bLimitAngle=*/false);
+		}
 	}
 
 	UMassEntitySubsystem* EntitySystem = UWorld::GetSubsystem<UMassEntitySubsystem>(World);
@@ -328,7 +332,7 @@ void FGameplayDebuggerCategory_Mass::CollectData(APlayerController* OwnerPC, AAc
 	}
 
 	NearEntityDescriptions.Reset();
-	if (bShowNearEntityOverview)
+	if (bShowNearEntityOverview && OwnerPC)
 	{
 		FVector ViewLocation = FVector::ZeroVector;
 		FVector ViewDirection = FVector::ForwardVector;
diff --git a/Engine/Source/Developer/GameplayDebugger/Private/Editor/GameplayDebuggerEdMode.cpp b/Engine/Source/Developer/GameplayDebugger/Private/Editor/GameplayDebuggerEdMode.cpp
--- a/Engine/Source/Developer/GameplayDebugger/Private/Editor/GameplayDebuggerEdMode.cpp
+++ b/Engine/Source/Developer/GameplayDebugger/Private/Editor/GameplayDebuggerEdMode.cpp
@@ -109,6 +109,12 @@ void FGameplayDebuggerEdMode::Tick(FEditorViewportClient* ViewportClient, float
 	}
 }
 
+void FGameplayDebuggerEdMode::SafeOpenMode()
+{
+	// By calling this get, it make sure that the singleton is created and not during the SafeCloseMode whitch can be called during garbage collect.
+	GLevelEditorModeTools();
+}
+
 void FGameplayDebuggerEdMode::SafeCloseMode()
 {
 	// this may be called on closing editor during PIE (~viewport -> teardown PIE -> debugger's cleanup on game end)
diff --git a/Engine/Source/Developer/GameplayDebugger/Private/Editor/GameplayDebuggerEdMode.h b/Engine/Source/Developer/GameplayDebugger/Private/Editor/GameplayDebuggerEdMode.h
--- a/Engine/Source/Developer/GameplayDebugger/Private/Editor/GameplayDebuggerEdMode.h
+++ b/Engine/Source/Developer/GameplayDebugger/Private/Editor/GameplayDebuggerEdMode.h
@@ -25,6 +25,7 @@ public:
 	virtual void Tick(FEditorViewportClient* ViewportClient, float DeltaTime);
 	// End of FEdMode interface
 
+	static void SafeOpenMode();
 	static void SafeCloseMode();
 
 	static const FName EM_GameplayDebugger;
diff --git a/Engine/Source/Developer/GameplayDebugger/Private/GameplayDebuggerLocalController.cpp b/Engine/Source/Developer/GameplayDebugger/Private/GameplayDebuggerLocalController.cpp
--- a/Engine/Source/Developer/GameplayDebugger/Private/GameplayDebuggerLocalController.cpp
+++ b/Engine/Source/Developer/GameplayDebugger/Private/GameplayDebuggerLocalController.cpp
@@ -59,6 +59,11 @@ void UGameplayDebuggerLocalController::Initialize(AGameplayDebuggerCategoryRepli
 	UDebugDrawService::Register(bSimulateMode ? TEXT("DebugAI") : TEXT("Game"), FDebugDrawDelegate::CreateUObject(this, &UGameplayDebuggerLocalController::OnDebugDraw));
 
 #if WITH_EDITOR
+	if (bSimulateMode)
+	{
+		FGameplayDebuggerEdMode::SafeOpenMode();
+	}
+
 	if (GIsEditor)
 	{
 		USelection::SelectionChangedEvent.AddUObject(this, &UGameplayDebuggerLocalController::OnSelectionChanged);
