Made avoidance processors work with no LOD tags

diff --git a/Engine/Plugins/AI/MassAI/Source/MassAIMovement/Private/MassAvoidanceProcessors.cpp b/Engine/Plugins/AI/MassAI/Source/MassAIMovement/Private/MassAvoidanceProcessors.cpp
--- a/Engine/Plugins/AI/MassAI/Source/MassAIMovement/Private/MassAvoidanceProcessors.cpp
+++ b/Engine/Plugins/AI/MassAI/Source/MassAIMovement/Private/MassAvoidanceProcessors.cpp
@@ -405,7 +405,9 @@ void UMassAvoidanceProcessor::ConfigureQueries()
 	EntityQuery.AddRequirement<FMassVelocityFragment>(EMassFragmentAccess::ReadOnly);
 	EntityQuery.AddRequirement<FDataFragment_AgentRadius>(EMassFragmentAccess::ReadOnly);
 	EntityQuery.AddRequirement<FMassMovementConfigFragment>(EMassFragmentAccess::ReadOnly);
-	EntityQuery.AddTagRequirement<FMassHighLODTag>(EMassFragmentPresence::All);
+	EntityQuery.AddTagRequirement<FMassMediumLODTag>(EMassFragmentPresence::None);
+	EntityQuery.AddTagRequirement<FMassLowLODTag>(EMassFragmentPresence::None);
+	EntityQuery.AddTagRequirement<FMassOffLODTag>(EMassFragmentPresence::None);
 }
 
 void UMassAvoidanceProcessor::Initialize(UObject& Owner)
@@ -1153,7 +1155,9 @@ void UMassStandingAvoidanceProcessor::ConfigureQueries()
 	EntityQuery.AddRequirement<FMassMoveTargetFragment>(EMassFragmentAccess::ReadOnly);
 	EntityQuery.AddRequirement<FDataFragment_Transform>(EMassFragmentAccess::ReadOnly);
 	EntityQuery.AddRequirement<FDataFragment_AgentRadius>(EMassFragmentAccess::ReadOnly);
-	EntityQuery.AddTagRequirement<FMassHighLODTag>(EMassFragmentPresence::All);
+	EntityQuery.AddTagRequirement<FMassMediumLODTag>(EMassFragmentPresence::None);
+	EntityQuery.AddTagRequirement<FMassLowLODTag>(EMassFragmentPresence::None);
+	EntityQuery.AddTagRequirement<FMassOffLODTag>(EMassFragmentPresence::None);
 }
 
 void UMassStandingAvoidanceProcessor::Initialize(UObject& Owner)
@@ -1615,7 +1619,9 @@ void UMassLaneBoundaryProcessor::ConfigureQueries()
 	EntityQuery.AddRequirement<FMassLastUpdatePositionFragment>(EMassFragmentAccess::ReadWrite);					// to keep position when boundaries where last updated
 	EntityQuery.AddRequirement<FMassZoneGraphLaneLocationFragment>(EMassFragmentAccess::ReadOnly);				// current lane location
 	EntityQuery.AddRequirement<FMassAvoidanceBoundaryLastLaneHandleFragment>(EMassFragmentAccess::ReadWrite);	// keep track of the last used lane
-	EntityQuery.AddTagRequirement<FMassHighLODTag>(EMassFragmentPresence::All);
+	EntityQuery.AddTagRequirement<FMassMediumLODTag>(EMassFragmentPresence::None);
+	EntityQuery.AddTagRequirement<FMassLowLODTag>(EMassFragmentPresence::None);
+	EntityQuery.AddTagRequirement<FMassOffLODTag>(EMassFragmentPresence::None);
 }
 
 void UMassLaneBoundaryProcessor::Initialize(UObject& Owner)
