Made FGameplayDebuggerCategory_Mass's FMassLookAtFragment use optional

diff --git a/Engine/Plugins/AI/MassAI/Source/MassAIDebug/Private/GameplayDebuggerCategory_Mass.cpp b/Engine/Plugins/AI/MassAI/Source/MassAIDebug/Private/GameplayDebuggerCategory_Mass.cpp
--- a/Engine/Plugins/AI/MassAI/Source/MassAIDebug/Private/GameplayDebuggerCategory_Mass.cpp
+++ b/Engine/Plugins/AI/MassAI/Source/MassAIDebug/Private/GameplayDebuggerCategory_Mass.cpp
@@ -347,8 +347,7 @@ void FGameplayDebuggerCategory_Mass::CollectData(APlayerController* OwnerPC, AAc
 		EntityQuery.AddRequirement<FMassSteeringGhostFragment>(EMassFragmentAccess::ReadOnly);
 		EntityQuery.AddRequirement<FMassVelocityFragment>(EMassFragmentAccess::ReadOnly);
 		EntityQuery.AddRequirement<FMassMoveTargetFragment>(EMassFragmentAccess::ReadOnly);
-		EntityQuery.AddRequirement<FMassLookAtFragment>(EMassFragmentAccess::ReadOnly);
-		EntityQuery.AddRequirement<FMassLookAtTrajectoryFragment>(EMassFragmentAccess::ReadOnly);
+		EntityQuery.AddRequirement<FMassLookAtFragment>(EMassFragmentAccess::ReadOnly, EMassFragmentPresence::Optional);
 		EntityQuery.AddRequirement<FMassSimulationLODFragment>(EMassFragmentAccess::ReadOnly, EMassFragmentPresence::Optional);
 		EntityQuery.AddRequirement<FMassZoneGraphShortPathFragment>(EMassFragmentAccess::ReadOnly);
 
@@ -372,7 +371,7 @@ void FGameplayDebuggerCategory_Mass::CollectData(APlayerController* OwnerPC, AAc
 				const TConstArrayView<FMassVelocityFragment> VelocityList = Context.GetFragmentView<FMassVelocityFragment>();
 				const TConstArrayView<FMassMoveTargetFragment> MoveTargetList = Context.GetFragmentView<FMassMoveTargetFragment>();
 				const TConstArrayView<FMassLookAtFragment> LookAtList = Context.GetFragmentView<FMassLookAtFragment>();
-				const TConstArrayView<FMassLookAtTrajectoryFragment> LookAtTrajectoryList = Context.GetFragmentView<FMassLookAtTrajectoryFragment>();
+				const bool bHasLookAt = (LookAtList.Num() > 0);
 				const TConstArrayView<FMassSimulationLODFragment> SimLODList = Context.GetFragmentView<FMassSimulationLODFragment>();
 				const bool bHasLOD = (SimLODList.Num() > 0);
 				const TConstArrayView<FMassZoneGraphShortPathFragment> ShortPathList = Context.GetFragmentView<FMassZoneGraphShortPathFragment>();
@@ -405,7 +404,6 @@ void FGameplayDebuggerCategory_Mass::CollectData(APlayerController* OwnerPC, AAc
 					const FMassSteeringGhostFragment& Ghost = GhostList[EntityIndex];
 					const FMassVelocityFragment& Velocity = VelocityList[EntityIndex];
 					const FMassMoveTargetFragment& MoveTarget = MoveTargetList[EntityIndex];
-					const FMassLookAtFragment& LookAt = LookAtList[EntityIndex];
 					const FMassSimulationLODFragment& SimLOD = bHasLOD ? SimLODList[EntityIndex] : FMassSimulationLODFragment();
 					const FMassZoneGraphShortPathFragment& ShortPath = ShortPathList[EntityIndex];
 
@@ -433,38 +431,42 @@ void FGameplayDebuggerCategory_Mass::CollectData(APlayerController* OwnerPC, AAc
 					// Look at
 					constexpr float LookArrowLength = 100.0f;
 					BasePos = EntityLocation + FVector(0,0,EyeHeight);
-					const FVector WorldLookDirection = Transform.GetTransform().TransformVector(LookAt.Direction);
-					bool bLookArrowDrawn = false;
-					if (LookAt.LookAtMode == EMassLookAtMode::LookAtEntity && EntitySystem->IsEntityValid(LookAt.TrackedEntity))
+
+					if (bHasLookAt)
 					{
-						if (const FDataFragment_Transform* TargetTransform = EntitySystem->GetFragmentDataPtr<FDataFragment_Transform>(LookAt.TrackedEntity))
+						const FMassLookAtFragment& LookAt = LookAtList[EntityIndex];
+						const FVector WorldLookDirection = Transform.GetTransform().TransformVector(LookAt.Direction);
+						bool bLookArrowDrawn = false;
+						if (LookAt.LookAtMode == EMassLookAtMode::LookAtEntity && EntitySystem->IsEntityValid(LookAt.TrackedEntity))
 						{
-							FVector TargetPosition = TargetTransform->GetTransform().GetLocation();
-							TargetPosition.Z = BasePos.Z;
-							AddShape(FGameplayDebuggerShape::MakeCircle(TargetPosition, FVector::UpVector, Radius.Radius, FColor::Red));
-
-							const float TargetDistance = FMath::Max(LookArrowLength, FVector::DotProduct(WorldLookDirection, TargetPosition - BasePos));
-							AddShape(FGameplayDebuggerShape::MakeSegment(BasePos, BasePos + WorldLookDirection * TargetDistance, FColorList::LightGrey));
-							bLookArrowDrawn = true;
+							if (const FDataFragment_Transform* TargetTransform = EntitySystem->GetFragmentDataPtr<FDataFragment_Transform>(LookAt.TrackedEntity))
+							{
+								FVector TargetPosition = TargetTransform->GetTransform().GetLocation();
+								TargetPosition.Z = BasePos.Z;
+								AddShape(FGameplayDebuggerShape::MakeCircle(TargetPosition, FVector::UpVector, Radius.Radius, FColor::Red));
+
+								const float TargetDistance = FMath::Max(LookArrowLength, FVector::DotProduct(WorldLookDirection, TargetPosition - BasePos));
+								AddShape(FGameplayDebuggerShape::MakeSegment(BasePos, BasePos + WorldLookDirection * TargetDistance, FColorList::LightGrey));
+								bLookArrowDrawn = true;
+							}
 						}
-					}
 
-					if (LookAt.bRandomGazeEntities && EntitySystem->IsEntityValid(LookAt.GazeTrackedEntity))
-					{
-						if (const FDataFragment_Transform* TargetTransform = EntitySystem->GetFragmentDataPtr<FDataFragment_Transform>(LookAt.GazeTrackedEntity))
+						if (LookAt.bRandomGazeEntities && EntitySystem->IsEntityValid(LookAt.GazeTrackedEntity))
 						{
-							FVector TargetPosition = TargetTransform->GetTransform().GetLocation();
-							TargetPosition.Z = BasePos.Z;
-							AddShape(FGameplayDebuggerShape::MakeCircle(TargetPosition, FVector::UpVector, Radius.Radius, FColor::Turquoise));
+							if (const FDataFragment_Transform* TargetTransform = EntitySystem->GetFragmentDataPtr<FDataFragment_Transform>(LookAt.GazeTrackedEntity))
+							{
+								FVector TargetPosition = TargetTransform->GetTransform().GetLocation();
+								TargetPosition.Z = BasePos.Z;
+								AddShape(FGameplayDebuggerShape::MakeCircle(TargetPosition, FVector::UpVector, Radius.Radius, FColor::Turquoise));
+							}
 						}
-					}
 
-					if (!bLookArrowDrawn)
-					{
-						AddShape(FGameplayDebuggerShape::MakeArrow(BasePos, BasePos + WorldLookDirection * LookArrowLength, 10.0f, 1.0f, FColor::Turquoise));
+						if (!bLookArrowDrawn)
+						{
+							AddShape(FGameplayDebuggerShape::MakeArrow(BasePos, BasePos + WorldLookDirection * LookArrowLength, 10.0f, 1.0f, FColor::Turquoise));
+						}
 					}
 
-
 					// Path
 					if (bShowNearEntityPath)
 					{
@@ -549,9 +551,13 @@ void FGameplayDebuggerCategory_Mass::CollectData(APlayerController* OwnerPC, AAc
 							*UEnum::GetDisplayValueAsText(MoveTarget.IntentAtGoal).ToString(), MoveTarget.DistanceToGoal);
 
 						// Look
-						const float RemainingTime = LookAt.GazeDuration - (CurrentTime - LookAt.GazeStartTime);
-						Status += FString::Printf(TEXT("{turquoise}%s/%s {lightgrey}%.1f\n"),
-							*UEnum::GetDisplayValueAsText(LookAt.LookAtMode).ToString(), *UEnum::GetDisplayValueAsText(LookAt.RandomGazeMode).ToString(), RemainingTime);
+						if (bHasLookAt)
+						{
+							const FMassLookAtFragment& LookAt = LookAtList[EntityIndex];
+							const float RemainingTime = LookAt.GazeDuration - (CurrentTime - LookAt.GazeStartTime);
+							Status += FString::Printf(TEXT("{turquoise}%s/%s {lightgrey}%.1f\n"),
+								*UEnum::GetDisplayValueAsText(LookAt.LookAtMode).ToString(), *UEnum::GetDisplayValueAsText(LookAt.RandomGazeMode).ToString(), RemainingTime);
+						}
 						
 						if (!Status.IsEmpty())
 						{
