Fix Mass look at when simulation LOD is not enabled.

#rnx

diff --git a/Engine/Plugins/AI/MassAI/Source/MassAIBehavior/Private/MassLookAtProcessors.cpp b/Engine/Plugins/AI/MassAI/Source/MassAIBehavior/Private/MassLookAtProcessors.cpp
--- a/Engine/Plugins/AI/MassAI/Source/MassAIBehavior/Private/MassLookAtProcessors.cpp
+++ b/Engine/Plugins/AI/MassAI/Source/MassAIBehavior/Private/MassLookAtProcessors.cpp
@@ -112,7 +112,9 @@ void UMassLookAtProcessor::ConfigureQueries()
 	EntityQuery_Conditional.AddRequirement<FMassZoneGraphLaneLocationFragment>(EMassFragmentAccess::ReadOnly, EMassFragmentPresence::Optional);
 	EntityQuery_Conditional.AddRequirement<FMassLookAtTrajectoryFragment>(EMassFragmentAccess::ReadWrite, EMassFragmentPresence::Optional);
 	EntityQuery_Conditional.AddRequirement<FMassZoneGraphShortPathFragment>(EMassFragmentAccess::ReadOnly, EMassFragmentPresence::Optional);
-	EntityQuery_Conditional.AddTagRequirement<FMassHighLODTag>(EMassFragmentPresence::All);
+	EntityQuery_Conditional.AddTagRequirement<FMassMediumLODTag>(EMassFragmentPresence::None);
+	EntityQuery_Conditional.AddTagRequirement<FMassLowLODTag>(EMassFragmentPresence::None);
+	EntityQuery_Conditional.AddTagRequirement<FMassOffLODTag>(EMassFragmentPresence::None);
 	EntityQuery_Conditional.AddChunkRequirement<FMassVisualizationChunkFragment>(EMassFragmentAccess::ReadOnly);
 	EntityQuery_Conditional.SetChunkFilter(&FMassVisualizationChunkFragment::AreAnyEntitiesVisibleInChunk);
 }
