Prevent MassAI and StateTree from crashing when invalid StateTree assets are used.
Added a check to make usre StateTrees have valid instances, and prevented adding Mass Behavior subsystems when an ExecutionContext fails to initialize.

diff --git a/Engine/Plugins/AI/MassAI/Source/MassAIBehavior/Private/MassStateTreeProcessors.cpp b/Engine/Plugins/AI/MassAI/Source/MassAIBehavior/Private/MassStateTreeProcessors.cpp
--- a/Engine/Plugins/AI/MassAI/Source/MassAIBehavior/Private/MassStateTreeProcessors.cpp
+++ b/Engine/Plugins/AI/MassAI/Source/MassAIBehavior/Private/MassStateTreeProcessors.cpp
@@ -94,9 +94,8 @@ void ForEachEntityInChunk(
 	// Initialize the execution context if changed between chunks.
 	if (StateTreeContext.GetStateTree() != StateTree)
 	{
-		StateTreeContext.Init(MassStateTreeSubsystem, *StateTree, EStateTreeStorage::External);
-
 		// Gather subsystems.
+		if (StateTreeContext.Init(MassStateTreeSubsystem, *StateTree, EStateTreeStorage::External))
 		{
 			CSV_SCOPED_TIMING_STAT_EXCLUSIVE(StateTreeProcessorExternalSubsystems);
 			if (!ensureMsgf(UE::MassBehavior::SetExternalSubsystems(StateTreeContext), TEXT("StateTree will not execute due to missing subsystem requirements.")))
diff --git a/Engine/Plugins/Runtime/StateTree/Source/StateTreeModule/Private/StateTreeExecutionContext.cpp b/Engine/Plugins/Runtime/StateTree/Source/StateTreeModule/Private/StateTreeExecutionContext.cpp
--- a/Engine/Plugins/Runtime/StateTree/Source/StateTreeModule/Private/StateTreeExecutionContext.cpp
+++ b/Engine/Plugins/Runtime/StateTree/Source/StateTreeModule/Private/StateTreeExecutionContext.cpp
@@ -39,6 +39,14 @@ bool FStateTreeExecutionContext::Init(UObject& InOwner, const UStateTree& InStat
 		Reset();
 		return false;
 	}
+	
+	if (InStateTree.Instances.Num() == 0)
+	{
+		STATETREE_LOG(Error, TEXT("%s: StateTree asset '%s' has no valid instances."), ANSI_TO_TCHAR(__FUNCTION__), *InStateTree.GetName());
+		Reset();
+		return false;
+	}
+	
 	StateTree = &InStateTree;
 
 	StorageType = InStorageType;
