Prevent evaluation of StateTreeContexts if they can't be initialized.

diff --git a/Engine/Plugins/AI/MassAI/Source/MassAIBehavior/Private/MassStateTreeProcessors.cpp b/Engine/Plugins/AI/MassAI/Source/MassAIBehavior/Private/MassStateTreeProcessors.cpp
--- a/Engine/Plugins/AI/MassAI/Source/MassAIBehavior/Private/MassStateTreeProcessors.cpp
+++ b/Engine/Plugins/AI/MassAI/Source/MassAIBehavior/Private/MassStateTreeProcessors.cpp
@@ -103,6 +103,10 @@ void ForEachEntityInChunk(
 				return;
 			}
 		}
+		else
+		{
+			return;
+		}
 	}
 
 	const UScriptStruct* StorageScriptStruct = StateTree->GetInstanceStorageStruct();
