Improved readability of tag information displayed for an entity being viewed via Mass' Gameplay Debugger

diff --git a/Engine/Plugins/AI/MassAI/Source/MassAIDebug/Private/GameplayDebuggerCategory_Mass.cpp b/Engine/Plugins/AI/MassAI/Source/MassAIDebug/Private/GameplayDebuggerCategory_Mass.cpp
--- a/Engine/Plugins/AI/MassAI/Source/MassAIDebug/Private/GameplayDebuggerCategory_Mass.cpp
+++ b/Engine/Plugins/AI/MassAI/Source/MassAIDebug/Private/GameplayDebuggerCategory_Mass.cpp
@@ -296,15 +296,20 @@ void FGameplayDebuggerCategory_Mass::CollectData(APlayerController* OwnerPC, AAc
 					EntitySystem->DebugGetArchetypeStrings(Archetype, ComponentNames, TagNames);
 
 					FString Tags;
+					int i = 0;
 					for (const FName Name : TagNames)
 					{
 						Tags += FString::Printf(TEXT("%s, "), *Name.ToString());
+						if (i++ % 2)
+						{
+							Tags += TEXT("\n");
+						}
 					}
 					AddTextLine(FString::Printf(TEXT("{Green}Tags:\n{White}%s"), *Tags));
 
 					AddTextLine(FString::Printf(TEXT("{Green}Fragments:{White}")));
 					constexpr int ColumnsCount = 2;
-					int i = 0;
+					i = 0;
 					while (i + ColumnsCount < ComponentNames.Num())
 					{
 						AddTextLine(FString::Printf(TEXT("%-42s, %-42s"), *ComponentNames[i].ToString(), *ComponentNames[i + 1].ToString()));
