[MassSmartObject] fixed Request Id that was reset without destroying the associated entity

diff --git a/Engine/Plugins/AI/MassAI/Source/MassAIBehavior/Private/Evaluators/MassStateTreeSmartObjectEvaluator.cpp b/Engine/Plugins/AI/MassAI/Source/MassAIBehavior/Private/Evaluators/MassStateTreeSmartObjectEvaluator.cpp
--- a/Engine/Plugins/AI/MassAI/Source/MassAIBehavior/Private/Evaluators/MassStateTreeSmartObjectEvaluator.cpp
+++ b/Engine/Plugins/AI/MassAI/Source/MassAIBehavior/Private/Evaluators/MassStateTreeSmartObjectEvaluator.cpp
@@ -36,7 +36,7 @@ bool FMassStateTreeSmartObjectEvaluator::Link(FStateTreeLinker& Linker)
 	return true;
 }
 
-void FMassStateTreeSmartObjectEvaluator::EnterState(FStateTreeExecutionContext& Context, const EStateTreeStateChangeType ChangeType, const FStateTreeTransitionResult& Transition) const
+void FMassStateTreeSmartObjectEvaluator::ExitState(FStateTreeExecutionContext& Context, const EStateTreeStateChangeType ChangeType, const FStateTreeTransitionResult& Transition) const
 {
 	if (ChangeType != EStateTreeStateChangeType::Changed)
 	{
@@ -46,15 +46,11 @@ void FMassStateTreeSmartObjectEvaluator::EnterState(FStateTreeExecutionContext&
 	Reset(Context);
 }
 
-void FMassStateTreeSmartObjectEvaluator::ExitState(FStateTreeExecutionContext& Context, const EStateTreeStateChangeType ChangeType, const FStateTreeTransitionResult& Transition) const
+void FMassStateTreeSmartObjectEvaluator::Reset(FStateTreeExecutionContext& Context) const
 {
-	if (ChangeType != EStateTreeStateChangeType::Changed)
-	{
-		return;
-	}
-
+	bool& bCandidatesFound = Context.GetInstanceData(CandidatesFoundHandle);
+	bool& bClaimed = Context.GetInstanceData(ClaimedHandle);
 	FMassSmartObjectRequestID& SearchRequestID = Context.GetInstanceData(SearchRequestIDHandle);
-
 	if (SearchRequestID.IsSet())
 	{
 		const FMassStateTreeExecutionContext& MassContext = static_cast<FMassStateTreeExecutionContext&>(Context);
@@ -63,18 +59,9 @@ void FMassStateTreeSmartObjectEvaluator::ExitState(FStateTreeExecutionContext& C
 		MassSmartObjectHandler.RemoveRequest(SearchRequestID);
 		SearchRequestID.Reset();
 	}
-	Reset(Context);
-}
-
-void FMassStateTreeSmartObjectEvaluator::Reset(FStateTreeExecutionContext& Context) const
-{
-	bool& bCandidatesFound = Context.GetInstanceData(CandidatesFoundHandle);
-	bool& bClaimed = Context.GetInstanceData(ClaimedHandle);
-	FMassSmartObjectRequestID& SearchRequestID = Context.GetInstanceData(SearchRequestIDHandle);
 
 	bCandidatesFound = false;
 	bClaimed = false;
-	SearchRequestID.Reset();
 }
 
 void FMassStateTreeSmartObjectEvaluator::Evaluate(FStateTreeExecutionContext& Context, const EStateTreeEvaluationType EvalType,  const float DeltaTime) const
diff --git a/Engine/Plugins/AI/MassAI/Source/MassAIBehavior/Public/Evaluators/MassStateTreeSmartObjectEvaluator.h b/Engine/Plugins/AI/MassAI/Source/MassAIBehavior/Public/Evaluators/MassStateTreeSmartObjectEvaluator.h
--- a/Engine/Plugins/AI/MassAI/Source/MassAIBehavior/Public/Evaluators/MassStateTreeSmartObjectEvaluator.h
+++ b/Engine/Plugins/AI/MassAI/Source/MassAIBehavior/Public/Evaluators/MassStateTreeSmartObjectEvaluator.h
@@ -52,7 +52,6 @@ struct MASSAIBEHAVIOR_API FMassStateTreeSmartObjectEvaluator : public FMassState
 protected:
 	virtual bool Link(FStateTreeLinker& Linker) override;
 	virtual const UStruct* GetInstanceDataType() const override { return FMassStateTreeSmartObjectEvaluatorInstanceData::StaticStruct(); }
-	virtual void EnterState(FStateTreeExecutionContext& Context, const EStateTreeStateChangeType ChangeType, const FStateTreeTransitionResult& Transition) const override;
 	virtual void ExitState(FStateTreeExecutionContext& Context, const EStateTreeStateChangeType ChangeType, const FStateTreeTransitionResult& Transition) const override;
 	virtual void Evaluate(FStateTreeExecutionContext& Context, const EStateTreeEvaluationType EvalType, const float DeltaTime) const override;
 	void Reset(FStateTreeExecutionContext& Context) const;
