[MassMovement] use direction to end of path point instead of tangent from last lane point for last point outside of a lane.
SmartObjectUser fragment now holds desired entity orientation

diff --git a/Engine/Plugins/AI/MassAI/Source/MassAIBehavior/Private/Tasks/MassZoneGraphFindSmartObjectTarget.cpp b/Engine/Plugins/AI/MassAI/Source/MassAIBehavior/Private/Tasks/MassZoneGraphFindSmartObjectTarget.cpp
--- a/Engine/Plugins/AI/MassAI/Source/MassAIBehavior/Private/Tasks/MassZoneGraphFindSmartObjectTarget.cpp
+++ b/Engine/Plugins/AI/MassAI/Source/MassAIBehavior/Private/Tasks/MassZoneGraphFindSmartObjectTarget.cpp
@@ -51,6 +51,9 @@ EStateTreeRunStatus FMassZoneGraphFindSmartObjectTarget::EnterState(FStateTreeEx
 	SmartObjectLocation.bMoveReverse = false;
 	SmartObjectLocation.EndOfPathIntent = EMassMovementAction::Stand;
 	SmartObjectLocation.EndOfPathPosition = SOUser.TargetLocation;
+	// Can't set direction at the moment since it seems problematic if it's opposite to the steering direction
+	//SmartObjectLocation.EndOfPathDirection = SOUser.TargetDirection;
+
 	// Let's start moving toward the interaction a bit before the entry point.
 	SmartObjectLocation.AnticipationDistance.Set(100.f);
 
diff --git a/Engine/Plugins/AI/MassAI/Source/MassAIMovement/Private/MassZoneGraphMovementFragments.cpp b/Engine/Plugins/AI/MassAI/Source/MassAIMovement/Private/MassZoneGraphMovementFragments.cpp
--- a/Engine/Plugins/AI/MassAI/Source/MassAIMovement/Private/MassZoneGraphMovementFragments.cpp
+++ b/Engine/Plugins/AI/MassAI/Source/MassAIMovement/Private/MassZoneGraphMovementFragments.cpp
@@ -373,13 +373,19 @@ bool FMassZoneGraphShortPathFragment::RequestPath(const FMassZoneGraphCachedLane
 		}
 	}
 
+	checkf(NumPoints >= 1, TEXT("Path should have at least 1 point at this stage but has none."));
+
 	// Add end of path point if set.
 	if (bHasEndOfPathPoint)
 	{
 		if (NumPoints < MaxPoints)
 		{
 			const FVector EndPosition = Request.EndOfPathPosition;
-			const FVector EndDirection = Request.bIsEndOfPathDirectionSet ? Request.EndOfPathDirection.Get() : (EndLaneTangent * TangentSign);
+
+			// Use provided direction if set, otherwise use direction from last point on lane to end of path point
+			const FVector EndDirection = (Request.bIsEndOfPathDirectionSet) ?
+				Request.EndOfPathDirection.Get() :
+				(EndPosition - Points[NumPoints-1].Position).GetSafeNormal();
 			
 			FMassZoneGraphPathPoint& Point = Points[NumPoints++];
 			Point.DistanceAlongLane = FMassInt16Real10(TargetDistanceAlongLane);
diff --git a/Engine/Plugins/AI/MassAI/Source/MassAIMovement/Private/MassZoneGraphMovementProcessors.cpp b/Engine/Plugins/AI/MassAI/Source/MassAIMovement/Private/MassZoneGraphMovementProcessors.cpp
--- a/Engine/Plugins/AI/MassAI/Source/MassAIMovement/Private/MassZoneGraphMovementProcessors.cpp
+++ b/Engine/Plugins/AI/MassAI/Source/MassAIMovement/Private/MassZoneGraphMovementProcessors.cpp
@@ -270,11 +270,12 @@ void UMassZoneGraphPathFollowProcessor::Execute(UMassEntitySubsystem& EntitySubs
 						MoveTarget.DistanceToGoal = ShortPath.Points[LastPointIndex].Distance.Get();
 						MoveTarget.bOffBoundaries = ShortPath.Points[0].bOffLane;
 
-						UE_CVLOG(bDisplayDebug,this, LogMassNavigation, Verbose, TEXT("Entity [%s] before start of lane %s at distance %.1f. Distance to goal: %.1f."),
+						UE_CVLOG(bDisplayDebug,this, LogMassNavigation, Verbose, TEXT("Entity [%s] before start of lane %s at distance %.1f. Distance to goal: %.1f. Off Boundaries: %s"),
 							*Entity.DebugGetDescription(),
 							*LaneLocation.LaneHandle.ToString(),
 							LaneLocation.DistanceAlongLane,
-							MoveTarget.DistanceToGoal);
+							MoveTarget.DistanceToGoal,
+							*LexToString((bool)MoveTarget.bOffBoundaries));
 					}
 					else if (ShortPath.ProgressDistance <= ShortPath.Points[LastPointIndex].Distance.Get())
 					{
@@ -301,11 +302,12 @@ void UMassZoneGraphPathFollowProcessor::Execute(UMassEntitySubsystem& EntitySubs
 						MoveTarget.DistanceToGoal = ShortPath.Points[LastPointIndex].Distance.Get() - FMath::Lerp(CurrPoint.Distance.Get(), NextPoint.Distance.Get(), T);
 						MoveTarget.bOffBoundaries = CurrPoint.bOffLane || NextPoint.bOffLane;
 
-						UE_CVLOG(bDisplayDebug, this, LogMassNavigation, Verbose, TEXT("Entity [%s] along lane %s at distance %.1f. Distance to goal: %.1f."),
+						UE_CVLOG(bDisplayDebug, this, LogMassNavigation, Verbose, TEXT("Entity [%s] along lane %s at distance %.1f. Distance to goal: %.1f. Off Boundaries: %s"),
 							*Entity.DebugGetDescription(),
 							*LaneLocation.LaneHandle.ToString(),
 							LaneLocation.DistanceAlongLane,
-							MoveTarget.DistanceToGoal);
+							MoveTarget.DistanceToGoal,
+							*LexToString((bool)MoveTarget.bOffBoundaries));
 					}
 					else
 					{
@@ -317,8 +319,8 @@ void UMassZoneGraphPathFollowProcessor::Execute(UMassEntitySubsystem& EntitySubs
 						MoveTarget.DistanceToGoal = 0.0f;
 						MoveTarget.bOffBoundaries = ShortPath.Points[LastPointIndex].bOffLane;
 
-						UE_CVLOG(bDisplayDebug, this, LogMassNavigation, Log, TEXT("Entity [%s] Finished path follow on lane %s at distance %f."),
-							*Entity.DebugGetDescription(), *LaneLocation.LaneHandle.ToString(), LaneLocation.DistanceAlongLane);
+						UE_CVLOG(bDisplayDebug, this, LogMassNavigation, Log, TEXT("Entity [%s] Finished path follow on lane %s at distance %f. Off Boundaries: %s"),
+							*Entity.DebugGetDescription(), *LaneLocation.LaneHandle.ToString(), LaneLocation.DistanceAlongLane, *LexToString((bool)MoveTarget.bOffBoundaries));
 
 						if (bDisplayDebug)
 						{
diff --git a/Engine/Plugins/Runtime/MassGameplay/Source/MassSmartObjects/Private/MassSmartObjectHandler.cpp b/Engine/Plugins/Runtime/MassGameplay/Source/MassSmartObjects/Private/MassSmartObjectHandler.cpp
--- a/Engine/Plugins/Runtime/MassGameplay/Source/MassSmartObjects/Private/MassSmartObjectHandler.cpp
+++ b/Engine/Plugins/Runtime/MassGameplay/Source/MassSmartObjects/Private/MassSmartObjectHandler.cpp
@@ -122,7 +122,7 @@ EMassSmartObjectClaimResult FMassSmartObjectHandler::ClaimCandidate(const FMassE
 	}
 
 	checkf(!User.ClaimHandle.IsValid(), TEXT("User should not already have a valid handle."));
-	TConstArrayView<FSmartObjectCandidate> View = MakeArrayView(SearchRequestResult.Candidates.GetData(), SearchRequestResult.NumCandidates);
+	const TConstArrayView<FSmartObjectCandidate> View = MakeArrayView(SearchRequestResult.Candidates.GetData(), SearchRequestResult.NumCandidates);
 
 	for (const FSmartObjectCandidate& Candidate : View)
 	{
@@ -155,7 +155,9 @@ bool FMassSmartObjectHandler::ClaimSmartObject(const FMassEntityHandle Entity, F
 	{
 		User.ClaimHandle = ClaimHandle;
 		User.InteractionStatus = EMassSmartObjectInteractionStatus::Unset;
-		User.TargetLocation = SmartObjectSubsystem.GetSlotLocation(User.ClaimHandle).Get(FVector::ZeroVector);
+		const FTransform Transform =SmartObjectSubsystem.GetSlotTransform(User.ClaimHandle).Get(FTransform::Identity);
+		User.TargetLocation = Transform.GetLocation();
+		User.TargetDirection = Transform.GetRotation().Vector();
 
 		// @todo SO: need to unregister callback if entity or fragment gets removed. Or at least support getting called with a no longer valid FMassEntityHandle.
 		UE::Mass::SmartObject::FPayload Payload;
diff --git a/Engine/Plugins/Runtime/MassGameplay/Source/MassSmartObjects/Public/MassSmartObjectFragments.h b/Engine/Plugins/Runtime/MassGameplay/Source/MassSmartObjects/Public/MassSmartObjectFragments.h
--- a/Engine/Plugins/Runtime/MassGameplay/Source/MassSmartObjects/Public/MassSmartObjectFragments.h
+++ b/Engine/Plugins/Runtime/MassGameplay/Source/MassSmartObjects/Public/MassSmartObjectFragments.h
@@ -19,6 +19,9 @@ struct MASSSMARTOBJECTS_API FMassSmartObjectUserFragment : public FMassFragment
 	UPROPERTY(Transient)
 	FVector TargetLocation = FVector::ZeroVector;
 
+	UPROPERTY(Transient)
+	FVector TargetDirection = FVector::ForwardVector;
+
 	UPROPERTY(Transient)
 	EMassSmartObjectInteractionStatus InteractionStatus = EMassSmartObjectInteractionStatus::Unset;
 
