MassStateTree: Support shared fragments

diff --git a/Engine/Plugins/AI/MassAI/Source/MassAIBehavior/Private/MassStateTreeProcessors.cpp b/Engine/Plugins/AI/MassAI/Source/MassAIBehavior/Private/MassStateTreeProcessors.cpp
--- a/Engine/Plugins/AI/MassAI/Source/MassAIBehavior/Private/MassStateTreeProcessors.cpp
+++ b/Engine/Plugins/AI/MassAI/Source/MassAIBehavior/Private/MassStateTreeProcessors.cpp
@@ -25,7 +25,12 @@ bool SetExternalFragments(FMassStateTreeExecutionContext& Context, const UMassEn
 	const FMassEntityView EntityView(EntitySubsystem, Context.GetEntity());
 	for (const FStateTreeExternalDataDesc& DataDesc : Context.GetExternalDataDescs())
 	{
-		if (DataDesc.Struct && DataDesc.Struct->IsChildOf(FMassFragment::StaticStruct()))
+		if (DataDesc.Struct == nullptr)
+		{
+			continue;
+		}
+		
+		if (DataDesc.Struct->IsChildOf(FMassFragment::StaticStruct()))
 		{
 			const UScriptStruct* ScriptStruct = Cast<const UScriptStruct>(DataDesc.Struct);
 			FStructView Fragment = EntityView.GetFragmentDataStruct(ScriptStruct);
@@ -42,10 +47,27 @@ bool SetExternalFragments(FMassStateTreeExecutionContext& Context, const UMassEn
 				}
 			}
 		}
+		else if (DataDesc.Struct->IsChildOf(FMassSharedFragment::StaticStruct()))
+		{
+			const UScriptStruct* ScriptStruct = Cast<const UScriptStruct>(DataDesc.Struct);
+			FConstStructView Fragment = EntityView.GetConstSharedFragmentDataStruct(ScriptStruct);
+			if (Fragment.IsValid())
+			{
+				Context.SetExternalData(DataDesc.Handle, FStateTreeDataView(Fragment.GetScriptStruct(), const_cast<uint8*>(Fragment.GetMemory())));
+			}
+			else
+			{
+				if (DataDesc.Requirement == EStateTreeExternalDataRequirement::Required)
+				{
+					// Note: Not breaking here, so that we can validate all missing ones in one go with FMassStateTreeExecutionContext::AreExternalDataViewsValid().
+					bFoundAllFragments = false;
+				}
+			}
+		}
 	}
 	return bFoundAllFragments;
 }
-
+	
 bool SetExternalSubsystems(FMassStateTreeExecutionContext& Context)
 {
 	const UWorld* World = Context.GetWorld();
