Mass Movement and Navigation cleanup
- Removed ununed processors and fragments
- Separated ZoneGraph specific things from common movement and navigation fragments and processors
- First attempt at cleaner separation between Mass Movement and Mass Navigation
- Moved configs and hard coded values to shared fragments
- Simplified and clarified parameters

diff --git a/Engine/Plugins/AI/MassAI/Config/BaseMassAI.ini b/Engine/Plugins/AI/MassAI/Config/BaseMassAI.ini
--- a/Engine/Plugins/AI/MassAI/Config/BaseMassAI.ini
+++ b/Engine/Plugins/AI/MassAI/Config/BaseMassAI.ini
@@ -6,6 +6,7 @@
 +ClassRedirects=(OldName="/Script/MassMovement.MassDynamicObstacleTrait",NewName="/Script/MassAIMovement.MassAvoidanceObstacleTrait")
 +ClassRedirects=(OldName="/Script/MassMovement.MassAvoidanceTrait",NewName="/Script/MassAIMovement.MassAvoidanceTrait")
 +ClassRedirects=(OldName="/Script/MassMovement.MassZoneGraphMovementTrait",NewName="/Script/MassAIMovement.MassZoneGraphMovementTrait")
++ClassRedirects=(OldName="/Script/MassAIMovement.MassMovementSettings",NewName="/Script/MassMovement.MassMovementSettings")
 
 ; these are potentially not needed anymore
 +StructRedirects=(OldName="MassLaneHandleFragment",NewName="MassLaneNavigationFragment")
diff --git a/Engine/Plugins/AI/MassAI/MassAI.uplugin b/Engine/Plugins/AI/MassAI/MassAI.uplugin
--- a/Engine/Plugins/AI/MassAI/MassAI.uplugin
+++ b/Engine/Plugins/AI/MassAI/MassAI.uplugin
@@ -17,11 +17,15 @@
 	"Installed": false,
 	"Modules": [
 		{
-			"Name": "MassAIMovement",
+			"Name": "MassNavigation",
 			"Type": "Runtime"
 		},
 		{
-			"Name": "MassAIMovementEditor",
+			"Name": "MassZoneGraphNavigation",
+			"Type": "Runtime"
+		},
+		{
+			"Name": "MassNavigationEditor",
 			"Type": "Editor",
 			"LoadingPhase": "PostEngineInit"
 		},
diff --git a/Engine/Plugins/AI/MassAI/Source/MassAIBehavior/MassAIBehavior.Build.cs b/Engine/Plugins/AI/MassAI/Source/MassAIBehavior/MassAIBehavior.Build.cs
--- a/Engine/Plugins/AI/MassAI/Source/MassAIBehavior/MassAIBehavior.Build.cs
+++ b/Engine/Plugins/AI/MassAI/Source/MassAIBehavior/MassAIBehavior.Build.cs
@@ -31,7 +31,8 @@ namespace UnrealBuildTool.Rules
 					"MassCommon",
 					"MassLOD",
 					"MassMovement",
-					"MassAIMovement",
+					"MassNavigation",
+					"MassZoneGraphNavigation",
 					"MassRepresentation",
 					"MassSignals",
 					"MassSmartObjects",
diff --git a/Engine/Plugins/AI/MassAI/Source/MassAIBehavior/Private/Evaluators/MassStateTreeSmartObjectEvaluator.cpp b/Engine/Plugins/AI/MassAI/Source/MassAIBehavior/Private/Evaluators/MassStateTreeSmartObjectEvaluator.cpp
--- a/Engine/Plugins/AI/MassAI/Source/MassAIBehavior/Private/Evaluators/MassStateTreeSmartObjectEvaluator.cpp
+++ b/Engine/Plugins/AI/MassAI/Source/MassAIBehavior/Private/Evaluators/MassStateTreeSmartObjectEvaluator.cpp
@@ -8,7 +8,7 @@
 #include "MassSmartObjectFragments.h"
 #include "MassSmartObjectHandler.h"
 #include "MassStateTreeExecutionContext.h"
-#include "MassZoneGraphMovementFragments.h"
+#include "MassZoneGraphNavigationFragments.h"
 #include "SmartObjectSubsystem.h"
 #include "Engine/World.h"
 
diff --git a/Engine/Plugins/AI/MassAI/Source/MassAIBehavior/Private/MassLookAtProcessors.cpp b/Engine/Plugins/AI/MassAI/Source/MassAIBehavior/Private/MassLookAtProcessors.cpp
--- a/Engine/Plugins/AI/MassAI/Source/MassAIBehavior/Private/MassLookAtProcessors.cpp
+++ b/Engine/Plugins/AI/MassAI/Source/MassAIBehavior/Private/MassLookAtProcessors.cpp
@@ -5,12 +5,13 @@
 #include "MassAIBehaviorTypes.h"
 #include "MassCommonFragments.h"
 #include "MassLookAtFragments.h"
-#include "MassAIMovementFragments.h"
-#include "MassMovementSubsystem.h"
+#include "MassMovementFragments.h"
+#include "MassNavigationFragments.h"
+#include "MassNavigationSubsystem.h"
 #include "MassRepresentationFragments.h"
 #include "Math/UnrealMathUtility.h"
 #include "VisualLogger/VisualLogger.h"
-#include "MassZoneGraphMovementFragments.h"
+#include "MassZoneGraphNavigationFragments.h"
 #include "ZoneGraphSubsystem.h"
 #include "ZoneGraphQuery.h"
 #include "BezierUtilities.h"
@@ -97,8 +98,8 @@ void UMassLookAtProcessor::Initialize(UObject& Owner)
 {
 	Super::Initialize(Owner);
 
-	MovementSubsystem = UWorld::GetSubsystem<UMassMovementSubsystem>(Owner.GetWorld());
-	checkf(MovementSubsystem != nullptr, TEXT("UMassMovementSubsystem is mandatory when using MassLookAtProcessor processor."));
+	NavigationSubsystem = UWorld::GetSubsystem<UMassNavigationSubsystem>(Owner.GetWorld());
+	checkf(NavigationSubsystem != nullptr, TEXT("UMassNavigationSubsystem is mandatory when using MassLookAtProcessor processor."));
 
 	ZoneGraphSubsystem = UWorld::GetSubsystem<UZoneGraphSubsystem>(Owner.GetWorld());
 	checkf(ZoneGraphSubsystem != nullptr, TEXT("UMassZoneGraphSubsystem is mandatory when using MassLookAtProcessor processor."));
@@ -244,23 +245,23 @@ void UMassLookAtProcessor::FindNewGazeTarget(const UMassEntitySubsystem& EntityS
 
 	// Search for potential targets in front
 	bool bTargetFound = false;
-	if (LookAt.bRandomGazeEntities && MovementSubsystem != nullptr)
+	if (LookAt.bRandomGazeEntities && NavigationSubsystem != nullptr)
 	{
 		const float CosAngleThreshold = FMath::Cos(FMath::DegreesToRadians(AngleThresholdInDegrees));
 		const FVector Extent(QueryExtent, QueryExtent, QueryExtent);
 		const FVector QueryOrigin = Transform.TransformPosition(FVector(0.5f*QueryExtent, 0.f, 0.f));
 		const FBox QueryBox = FBox(QueryOrigin - 0.5f*Extent, QueryOrigin + 0.5f*Extent);
 
-		TArray<FAvoidanceObstacleHashGrid2D::ItemIDType> NearbyEntities;
+		TArray<FNavigationObstacleHashGrid2D::ItemIDType> NearbyEntities;
 		NearbyEntities.Reserve(16);
-		MovementSubsystem->GetGrid().Query(QueryBox, NearbyEntities);
+		NavigationSubsystem->GetObstacleGrid().Query(QueryBox, NearbyEntities);
 
 		// We'll pick the first entity that passes, this ensure that it's random one.
 		Algo::RandomShuffle(NearbyEntities);
 
 		const FVector Location = Transform.GetLocation();
 		
-		for (const FAvoidanceObstacleHashGrid2D::ItemIDType NearbyEntity : NearbyEntities)
+		for (const FNavigationObstacleHashGrid2D::ItemIDType NearbyEntity : NearbyEntities)
 		{
 			// This can happen if we remove entities in the system.
 			if (!EntitySubsystem.IsEntityValid(NearbyEntity.Entity))
diff --git a/Engine/Plugins/AI/MassAI/Source/MassAIBehavior/Private/MassStateTreeProcessors.cpp b/Engine/Plugins/AI/MassAI/Source/MassAIBehavior/Private/MassStateTreeProcessors.cpp
--- a/Engine/Plugins/AI/MassAI/Source/MassAIBehavior/Private/MassStateTreeProcessors.cpp
+++ b/Engine/Plugins/AI/MassAI/Source/MassAIBehavior/Private/MassStateTreeProcessors.cpp
@@ -4,7 +4,7 @@
 #include "StateTree.h"
 #include "MassStateTreeExecutionContext.h"
 #include "MassEntityView.h"
-#include "MassAIMovementTypes.h"
+#include "MassNavigationTypes.h"
 #include "MassSimulationLOD.h"
 #include "MassComponentHitTypes.h"
 #include "MassSmartObjectTypes.h"
diff --git a/Engine/Plugins/AI/MassAI/Source/MassAIBehavior/Private/MassZoneGraphAnnotationProcessors.cpp b/Engine/Plugins/AI/MassAI/Source/MassAIBehavior/Private/MassZoneGraphAnnotationProcessors.cpp
--- a/Engine/Plugins/AI/MassAI/Source/MassAIBehavior/Private/MassZoneGraphAnnotationProcessors.cpp
+++ b/Engine/Plugins/AI/MassAI/Source/MassAIBehavior/Private/MassZoneGraphAnnotationProcessors.cpp
@@ -2,7 +2,7 @@
 
 #include "MassZoneGraphAnnotationProcessors.h"
 #include "MassAIBehaviorTypes.h"
-#include "MassZoneGraphMovementFragments.h"
+#include "MassZoneGraphNavigationFragments.h"
 #include "MassZoneGraphAnnotationFragments.h"
 #include "MassZoneGraphAnnotationTypes.h"
 #include "ZoneGraphTypes.h"
diff --git a/Engine/Plugins/AI/MassAI/Source/MassAIBehavior/Private/Tasks/MassStateTreeSmartObjectTask.cpp b/Engine/Plugins/AI/MassAI/Source/MassAIBehavior/Private/Tasks/MassStateTreeSmartObjectTask.cpp
--- a/Engine/Plugins/AI/MassAI/Source/MassAIBehavior/Private/Tasks/MassStateTreeSmartObjectTask.cpp
+++ b/Engine/Plugins/AI/MassAI/Source/MassAIBehavior/Private/Tasks/MassStateTreeSmartObjectTask.cpp
@@ -6,9 +6,10 @@
 #include "MassSmartObjectHandler.h"
 #include "MassSmartObjectFragments.h"
 #include "MassStateTreeExecutionContext.h"
-#include "MassAIMovementFragments.h"
+#include "MassNavigationFragments.h"
 #include "SmartObjectSubsystem.h"
-#include "MassZoneGraphMovementUtils.h"
+#include "MassZoneGraphNavigationFragments.h"
+#include "MassZoneGraphNavigationUtils.h"
 #include "Engine/World.h"
 
 //----------------------------------------------------------------------//
@@ -154,7 +155,7 @@ EStateTreeRunStatus FMassUseSmartObjectTask::EnterState(FStateTreeExecutionConte
 	checkf(World != nullptr, TEXT("A valid world is expected from the execution context"));
 
 	MoveTarget.CreateNewAction(EMassMovementAction::Animate, *World);
-	const bool bSuccess = UE::MassMovement::ActivateActionAnimate(*World, Context.GetOwner(), MassContext.GetEntity(), MoveTarget);
+	const bool bSuccess = UE::MassNavigation::ActivateActionAnimate(*World, Context.GetOwner(), MassContext.GetEntity(), MoveTarget);
 
 	return bSuccess ? EStateTreeRunStatus::Running : EStateTreeRunStatus::Failed;
 }
diff --git a/Engine/Plugins/AI/MassAI/Source/MassAIBehavior/Private/Tasks/MassZoneGraphPathFollowTask.cpp b/Engine/Plugins/AI/MassAI/Source/MassAIBehavior/Private/Tasks/MassZoneGraphPathFollowTask.cpp
--- a/Engine/Plugins/AI/MassAI/Source/MassAIBehavior/Private/Tasks/MassZoneGraphPathFollowTask.cpp
+++ b/Engine/Plugins/AI/MassAI/Source/MassAIBehavior/Private/Tasks/MassZoneGraphPathFollowTask.cpp
@@ -4,22 +4,22 @@
 
 #include "MassAIBehaviorTypes.h"
 #include "MassCommonFragments.h"
-#include "MassAIMovementFragments.h"
-#include "MassMovementSettings.h"
+#include "MassNavigationFragments.h"
+#include "MassMovementFragments.h"
 #include "MassStateTreeExecutionContext.h"
-#include "MassZoneGraphMovementFragments.h"
-#include "MassZoneGraphMovementUtils.h"
+#include "MassZoneGraphNavigationFragments.h"
+#include "MassZoneGraphNavigationUtils.h"
 #include "ZoneGraphSubsystem.h"
 
 bool FMassZoneGraphPathFollowTask::Link(FStateTreeLinker& Linker)
 {
 	Linker.LinkExternalData(LocationHandle);
 	Linker.LinkExternalData(MoveTargetHandle);
-	Linker.LinkExternalData(MovementConfigHandle);
 	Linker.LinkExternalData(PathRequestHandle);
 	Linker.LinkExternalData(ShortPathHandle);
 	Linker.LinkExternalData(CachedLaneHandle);
 	Linker.LinkExternalData(AgentRadiusHandle);
+	Linker.LinkExternalData(MovementParamsHandle);
 	Linker.LinkExternalData(ZoneGraphSubsystemHandle);
 
 	Linker.LinkInstanceDataProperty(TargetLocationHandle, STATETREE_INSTANCEDATA_PROPERTY(FMassZoneGraphPathFollowTaskInstanceData, TargetLocation));
@@ -33,22 +33,12 @@ bool FMassZoneGraphPathFollowTask::RequestPath(FMassStateTreeExecutionContext& C
 {
 	const UZoneGraphSubsystem& ZoneGraphSubsystem = Context.GetExternalData(ZoneGraphSubsystemHandle);
 	const FMassZoneGraphLaneLocationFragment& LaneLocation = Context.GetExternalData(LocationHandle);
-	const FMassMovementConfigFragment& MovementConfig = Context.GetExternalData(MovementConfigHandle);
 	const FDataFragment_AgentRadius& AgentRadius = Context.GetExternalData(AgentRadiusHandle);
 	FMassZoneGraphShortPathFragment& ShortPath = Context.GetExternalData(ShortPathHandle);
 	FMassZoneGraphCachedLaneFragment& CachedLane = Context.GetExternalData(CachedLaneHandle);
 	FMassMoveTargetFragment& MoveTarget = Context.GetExternalData(MoveTargetHandle);
 	FMassZoneGraphPathRequestFragment& RequestFragment = Context.GetExternalData(PathRequestHandle);
-
-	const UMassMovementSettings* Settings = GetDefault<UMassMovementSettings>();
-	check(Settings);
-
-	const FMassMovementConfig* Config = Settings->GetMovementConfigByHandle(MovementConfig.ConfigHandle);
-	if (!Config)
-	{
-		MASSBEHAVIOR_LOG(Error, TEXT("Failed to get move config %s."), *LaneLocation.LaneHandle.ToString());
-		return false;
-	}
+	const FMassMovementParameters& MovementParams = Context.GetExternalData(MovementParamsHandle);
 
 	bool bDisplayDebug = false;
 #if WITH_MASSGAMEPLAY_DEBUG
@@ -89,12 +79,12 @@ bool FMassZoneGraphPathFollowTask::RequestPath(FMassStateTreeExecutionContext& C
 	const FMassMovementStyleRef& MovementStyle = Context.GetInstanceData(MovementStyleHandle);
 	const float SpeedScale = Context.GetInstanceData(SpeedScaleHandle);
 
-	const float DesiredSpeed = FMath::Min(Config->GenerateDesiredSpeed(MovementStyle, Context.GetEntity().Index) * SpeedScale, Config->MaximumSpeed);
+	const float DesiredSpeed = FMath::Min(MovementParams.GenerateDesiredSpeed(MovementStyle, Context.GetEntity().Index) * SpeedScale, MovementParams.MaxSpeed);
 	const UWorld* World = Context.GetWorld();
 	checkf(World != nullptr, TEXT("A valid world is expected from the execution context"));
 
 	MoveTarget.CreateNewAction(EMassMovementAction::Move, *World);
-	return UE::MassMovement::ActivateActionMove(*World, Context.GetOwner(), Context.GetEntity(), ZoneGraphSubsystem, LaneLocation, PathRequest, AgentRadius.Radius, DesiredSpeed, MoveTarget, ShortPath, CachedLane);
+	return UE::MassNavigation::ActivateActionMove(*World, Context.GetOwner(), Context.GetEntity(), ZoneGraphSubsystem, LaneLocation, PathRequest, AgentRadius.Radius, DesiredSpeed, MoveTarget, ShortPath, CachedLane);
 }
 
 EStateTreeRunStatus FMassZoneGraphPathFollowTask::EnterState(FStateTreeExecutionContext& Context, const EStateTreeStateChangeType ChangeType, const FStateTreeTransitionResult& Transition) const
diff --git a/Engine/Plugins/AI/MassAI/Source/MassAIBehavior/Private/Tasks/MassZoneGraphStandTask.cpp b/Engine/Plugins/AI/MassAI/Source/MassAIBehavior/Private/Tasks/MassZoneGraphStandTask.cpp
--- a/Engine/Plugins/AI/MassAI/Source/MassAIBehavior/Private/Tasks/MassZoneGraphStandTask.cpp
+++ b/Engine/Plugins/AI/MassAI/Source/MassAIBehavior/Private/Tasks/MassZoneGraphStandTask.cpp
@@ -3,14 +3,14 @@
 #include "Tasks/MassZoneGraphStandTask.h"
 #include "StateTreeExecutionContext.h"
 #include "ZoneGraphSubsystem.h"
-#include "MassZoneGraphMovementFragments.h"
+#include "MassZoneGraphNavigationFragments.h"
+#include "MassZoneGraphNavigationUtils.h"
 #include "MassAIBehaviorTypes.h"
-#include "MassAIMovementFragments.h"
-#include "MassMovementSettings.h"
+#include "MassNavigationFragments.h"
+#include "MassMovementFragments.h"
 #include "MassStateTreeExecutionContext.h"
 #include "MassSignalSubsystem.h"
 #include "MassSimulationLOD.h"
-#include "MassZoneGraphMovementUtils.h"
 
 bool FMassZoneGraphStandTask::Link(FStateTreeLinker& Linker)
 {
@@ -18,9 +18,9 @@ bool FMassZoneGraphStandTask::Link(FStateTreeLinker& Linker)
 	Linker.LinkExternalData(MoveTargetHandle);
 	Linker.LinkExternalData(ShortPathHandle);
 	Linker.LinkExternalData(CachedLaneHandle);
+	Linker.LinkExternalData(MovementParamsHandle);
 	Linker.LinkExternalData(ZoneGraphSubsystemHandle);
 	Linker.LinkExternalData(MassSignalSubsystemHandle);
-	Linker.LinkExternalData(MovementConfigHandle);
 
 	Linker.LinkInstanceDataProperty(DurationHandle, STATETREE_INSTANCEDATA_PROPERTY(FMassZoneGraphStandTaskInstanceData, Duration));
 	Linker.LinkInstanceDataProperty(TimeHandle, STATETREE_INSTANCEDATA_PROPERTY(FMassZoneGraphStandTaskInstanceData, Time));
@@ -40,18 +40,8 @@ EStateTreeRunStatus FMassZoneGraphStandTask::EnterState(FStateTreeExecutionConte
 
 	const FMassZoneGraphLaneLocationFragment& LaneLocation = Context.GetExternalData(LocationHandle);
 	const UZoneGraphSubsystem& ZoneGraphSubsystem = Context.GetExternalData(ZoneGraphSubsystemHandle);
-	const FMassMovementConfigFragment& MovementConfig = Context.GetExternalData(MovementConfigHandle);
+	const FMassMovementParameters& MovementParams = Context.GetExternalData(MovementParamsHandle);
 
-	const UMassMovementSettings* Settings = GetDefault<UMassMovementSettings>();
-	check(Settings);
-
-	const FMassMovementConfig* Config = Settings->GetMovementConfigByHandle(MovementConfig.ConfigHandle);
-	if (!Config)
-	{
-		MASSBEHAVIOR_LOG(Error, TEXT("Failed to get move config."));
-		return EStateTreeRunStatus::Failed;
-	}
-	
 	if (!LaneLocation.LaneHandle.IsValid())
 	{
 		MASSBEHAVIOR_LOG(Error, TEXT("Invalid lande handle"));
@@ -68,7 +58,7 @@ EStateTreeRunStatus FMassZoneGraphStandTask::EnterState(FStateTreeExecutionConte
 	checkf(World != nullptr, TEXT("A valid world is expected from the execution context"));
 
 	MoveTarget.CreateNewAction(EMassMovementAction::Stand, *World);
-	const bool bSuccess = UE::MassMovement::ActivateActionStand(*World, Context.GetOwner(), MassContext.GetEntity(), ZoneGraphSubsystem, LaneLocation, Config->DefaultDesiredSpeed, MoveTarget, ShortPath, CachedLane);
+	const bool bSuccess = UE::MassNavigation::ActivateActionStand(*World, Context.GetOwner(), MassContext.GetEntity(), ZoneGraphSubsystem, LaneLocation, MovementParams.DefaultDesiredSpeed, MoveTarget, ShortPath, CachedLane);
 	if (!bSuccess)
 	{
 		return EStateTreeRunStatus::Failed;
diff --git a/Engine/Plugins/AI/MassAI/Source/MassAIBehavior/Public/MassLookAtProcessors.h b/Engine/Plugins/AI/MassAI/Source/MassAIBehavior/Public/MassLookAtProcessors.h
--- a/Engine/Plugins/AI/MassAI/Source/MassAIBehavior/Public/MassLookAtProcessors.h
+++ b/Engine/Plugins/AI/MassAI/Source/MassAIBehavior/Public/MassLookAtProcessors.h
@@ -7,7 +7,7 @@
 #include "MassSignalProcessorBase.h"
 #include "MassLookAtProcessors.generated.h"
 
-class UMassMovementSubsystem;
+class UMassNavigationSubsystem;
 class UZoneGraphSubsystem;
 struct FMassLookAtFragment;
 struct FMassZoneGraphLaneLocationFragment;
@@ -71,7 +71,7 @@ protected:
 	FMassEntityQuery EntityQuery_Conditional;
 
 	UPROPERTY(Transient)
-	UMassMovementSubsystem* MovementSubsystem;
+	UMassNavigationSubsystem* NavigationSubsystem;
 
 	UPROPERTY(Transient)
 	UZoneGraphSubsystem* ZoneGraphSubsystem;
diff --git a/Engine/Plugins/AI/MassAI/Source/MassAIBehavior/Public/Tasks/MassZoneGraphPathFollowTask.h b/Engine/Plugins/AI/MassAI/Source/MassAIBehavior/Public/Tasks/MassZoneGraphPathFollowTask.h
--- a/Engine/Plugins/AI/MassAI/Source/MassAIBehavior/Public/Tasks/MassZoneGraphPathFollowTask.h
+++ b/Engine/Plugins/AI/MassAI/Source/MassAIBehavior/Public/Tasks/MassZoneGraphPathFollowTask.h
@@ -2,7 +2,8 @@
 
 #pragma once
 
-#include "MassAIMovementTypes.h"
+#include "MassNavigationTypes.h"
+#include "MassMovementTypes.h"
 #include "MassCommonTypes.h"
 #include "MassStateTreeTypes.h"
 #include "ZoneGraphTypes.h"
@@ -11,11 +12,11 @@
 struct FMassStateTreeExecutionContext;
 struct FMassZoneGraphLaneLocationFragment;
 struct FMassMoveTargetFragment;
-struct FMassMovementConfigFragment;
 struct FMassZoneGraphPathRequestFragment;
 struct FMassZoneGraphShortPathFragment;
 struct FMassZoneGraphCachedLaneFragment;
 struct FDataFragment_AgentRadius;
+struct FMassMovementParameters;
 class UZoneGraphSubsystem;
 
 USTRUCT()
@@ -96,11 +97,11 @@ protected:
 
 	TStateTreeExternalDataHandle<FMassZoneGraphLaneLocationFragment> LocationHandle;
 	TStateTreeExternalDataHandle<FMassMoveTargetFragment> MoveTargetHandle;
-	TStateTreeExternalDataHandle<FMassMovementConfigFragment> MovementConfigHandle;
 	TStateTreeExternalDataHandle<FMassZoneGraphPathRequestFragment> PathRequestHandle;
 	TStateTreeExternalDataHandle<FMassZoneGraphShortPathFragment> ShortPathHandle;
 	TStateTreeExternalDataHandle<FMassZoneGraphCachedLaneFragment> CachedLaneHandle;
 	TStateTreeExternalDataHandle<FDataFragment_AgentRadius> AgentRadiusHandle;
+	TStateTreeExternalDataHandle<FMassMovementParameters> MovementParamsHandle;
 	TStateTreeExternalDataHandle<UZoneGraphSubsystem> ZoneGraphSubsystemHandle;
 
 	TStateTreeInstanceDataPropertyHandle<FMassZoneGraphTargetLocation> TargetLocationHandle;
diff --git a/Engine/Plugins/AI/MassAI/Source/MassAIBehavior/Public/Tasks/MassZoneGraphStandTask.h b/Engine/Plugins/AI/MassAI/Source/MassAIBehavior/Public/Tasks/MassZoneGraphStandTask.h
--- a/Engine/Plugins/AI/MassAI/Source/MassAIBehavior/Public/Tasks/MassZoneGraphStandTask.h
+++ b/Engine/Plugins/AI/MassAI/Source/MassAIBehavior/Public/Tasks/MassZoneGraphStandTask.h
@@ -3,7 +3,6 @@
 #pragma once
 
 #include "MassStateTreeTypes.h"
-#include "MassMovementTypes.h"
 #include "MassZoneGraphStandTask.generated.h"
 
 struct FStateTreeExecutionContext;
@@ -11,9 +10,9 @@ struct FMassZoneGraphLaneLocationFragment;
 struct FMassMoveTargetFragment;
 struct FMassZoneGraphShortPathFragment;
 struct FMassZoneGraphCachedLaneFragment;
+struct FMassMovementParameters;
 class UZoneGraphSubsystem;
 class UMassSignalSubsystem;
-struct FMassMovementConfigFragment;
 
 /**
  * Stop, and stand on current ZoneGraph location
@@ -47,9 +46,9 @@ protected:
 	TStateTreeExternalDataHandle<FMassMoveTargetFragment> MoveTargetHandle;
 	TStateTreeExternalDataHandle<FMassZoneGraphShortPathFragment> ShortPathHandle;
 	TStateTreeExternalDataHandle<FMassZoneGraphCachedLaneFragment> CachedLaneHandle;
+	TStateTreeExternalDataHandle<FMassMovementParameters> MovementParamsHandle;
 	TStateTreeExternalDataHandle<UZoneGraphSubsystem> ZoneGraphSubsystemHandle;
 	TStateTreeExternalDataHandle<UMassSignalSubsystem> MassSignalSubsystemHandle;
-	TStateTreeExternalDataHandle<FMassMovementConfigFragment> MovementConfigHandle;
 
 	TStateTreeInstanceDataPropertyHandle<float> DurationHandle;
 	TStateTreeInstanceDataPropertyHandle<float> TimeHandle;
diff --git a/Engine/Plugins/AI/MassAI/Source/MassAIDebug/MassAIDebug.Build.cs b/Engine/Plugins/AI/MassAI/Source/MassAIDebug/MassAIDebug.Build.cs
--- a/Engine/Plugins/AI/MassAI/Source/MassAIDebug/MassAIDebug.Build.cs
+++ b/Engine/Plugins/AI/MassAI/Source/MassAIDebug/MassAIDebug.Build.cs
@@ -35,7 +35,8 @@ namespace UnrealBuildTool.Rules
 					"MassAIBehavior",
 					"MassCommon",
 					"MassMovement",
-					"MassAIMovement",
+					"MassNavigation",
+					"MassZoneGraphNavigation",
 					"MassAIReplication",
 					"MassSmartObjects",
 					"MassSpawner",
diff --git a/Engine/Plugins/AI/MassAI/Source/MassAIDebug/Private/GameplayDebuggerCategory_Mass.cpp b/Engine/Plugins/AI/MassAI/Source/MassAIDebug/Private/GameplayDebuggerCategory_Mass.cpp
--- a/Engine/Plugins/AI/MassAI/Source/MassAIDebug/Private/GameplayDebuggerCategory_Mass.cpp
+++ b/Engine/Plugins/AI/MassAI/Source/MassAIDebug/Private/GameplayDebuggerCategory_Mass.cpp
@@ -13,11 +13,12 @@
 #include "GameFramework/PlayerController.h"
 #include "MassAgentComponent.h"
 #include "MassEntitySubsystem.h"
-#include "MassAIMovementFragments.h"
+#include "MassNavigationFragments.h"
+#include "Steering/MassSteeringFragments.h"
 #include "MassLookAtFragments.h"
 #include "MassStateTreeFragments.h"
 #include "MassStateTreeExecutionContext.h"
-#include "MassZoneGraphMovementFragments.h"
+#include "MassZoneGraphNavigationFragments.h"
 #include "Util/ColorConstants.h"
 #include "MassSimulationLOD.h"
 #include "CanvasItem.h"
@@ -349,8 +350,10 @@ void FGameplayDebuggerCategory_Mass::CollectData(APlayerController* OwnerPC, AAc
 		EntityQuery.AddRequirement<FDataFragment_Transform>(EMassFragmentAccess::ReadOnly);
 		EntityQuery.AddRequirement<FDataFragment_AgentRadius>(EMassFragmentAccess::ReadOnly);
 		EntityQuery.AddRequirement<FMassSteeringFragment>(EMassFragmentAccess::ReadOnly);
-		EntityQuery.AddRequirement<FMassSteeringGhostFragment>(EMassFragmentAccess::ReadOnly);
+		EntityQuery.AddRequirement<FMassStandingSteeringFragment>(EMassFragmentAccess::ReadOnly);
+		EntityQuery.AddRequirement<FMassGhostLocationFragment>(EMassFragmentAccess::ReadOnly);
 		EntityQuery.AddRequirement<FMassVelocityFragment>(EMassFragmentAccess::ReadOnly);
+		EntityQuery.AddRequirement<FMassForceFragment>(EMassFragmentAccess::ReadOnly);
 		EntityQuery.AddRequirement<FMassMoveTargetFragment>(EMassFragmentAccess::ReadOnly);
 		EntityQuery.AddRequirement<FMassLookAtFragment>(EMassFragmentAccess::ReadOnly, EMassFragmentPresence::Optional);
 		EntityQuery.AddRequirement<FMassSimulationLODFragment>(EMassFragmentAccess::ReadOnly, EMassFragmentPresence::Optional);
@@ -372,8 +375,10 @@ void FGameplayDebuggerCategory_Mass::CollectData(APlayerController* OwnerPC, AAc
 				const TConstArrayView<FDataFragment_Transform> TransformList = Context.GetFragmentView<FDataFragment_Transform>();
 				const TConstArrayView<FDataFragment_AgentRadius> RadiusList = Context.GetFragmentView<FDataFragment_AgentRadius>();
 				const TConstArrayView<FMassSteeringFragment> SteeringList = Context.GetFragmentView<FMassSteeringFragment>();
-				const TConstArrayView<FMassSteeringGhostFragment> GhostList = Context.GetFragmentView<FMassSteeringGhostFragment>();
+				const TConstArrayView<FMassStandingSteeringFragment> StandingSteeringList = Context.GetFragmentView<FMassStandingSteeringFragment>();
+				const TConstArrayView<FMassGhostLocationFragment> GhostList = Context.GetFragmentView<FMassGhostLocationFragment>();
 				const TConstArrayView<FMassVelocityFragment> VelocityList = Context.GetFragmentView<FMassVelocityFragment>();
+				const TConstArrayView<FMassForceFragment> ForceList = Context.GetFragmentView<FMassForceFragment>();
 				const TConstArrayView<FMassMoveTargetFragment> MoveTargetList = Context.GetFragmentView<FMassMoveTargetFragment>();
 				const TConstArrayView<FMassLookAtFragment> LookAtList = Context.GetFragmentView<FMassLookAtFragment>();
 				const bool bHasLookAt = (LookAtList.Num() > 0);
@@ -407,8 +412,10 @@ void FGameplayDebuggerCategory_Mass::CollectData(APlayerController* OwnerPC, AAc
 
 					const FDataFragment_AgentRadius& Radius = RadiusList[EntityIndex];
 					const FMassSteeringFragment& Steering = SteeringList[EntityIndex];
-					const FMassSteeringGhostFragment& Ghost = GhostList[EntityIndex];
+					const FMassStandingSteeringFragment& StandingSteering = StandingSteeringList[EntityIndex];
+					const FMassGhostLocationFragment& Ghost = GhostList[EntityIndex];
 					const FMassVelocityFragment& Velocity = VelocityList[EntityIndex];
+					const FMassForceFragment& Force = ForceList[EntityIndex];
 					const FMassMoveTargetFragment& MoveTarget = MoveTargetList[EntityIndex];
 					const FMassSimulationLODFragment& SimLOD = bHasLOD ? SimLODList[EntityIndex] : FMassSimulationLODFragment();
 					const FMassZoneGraphShortPathFragment& ShortPath = ShortPathList[EntityIndex];
@@ -502,7 +509,7 @@ void FGameplayDebuggerCategory_Mass::CollectData(APlayerController* OwnerPC, AAc
 							GhostBasePos += FVector(0,0,5);
 							AddShape(FGameplayDebuggerShape::MakeArrow(GhostBasePos, GhostBasePos + Ghost.Velocity, 10.0f, 2.0f, FColorList::LightGrey));
 
-							const FVector GhostTargetBasePos = Ghost.SteerTarget + FVector(0.0f ,0.0f ,25.0f );
+							const FVector GhostTargetBasePos = StandingSteering.TargetLocation + FVector(0.0f ,0.0f ,25.0f );
 							AddShape(FGameplayDebuggerShape::MakeCircle(GhostTargetBasePos, FVector::UpVector, Radius.Radius * 0.75f, FColorList::Orange));
 						}
 					}
@@ -550,8 +557,8 @@ void FGameplayDebuggerCategory_Mass::CollectData(APlayerController* OwnerPC, AAc
 						}
 
 						// Movement info
-						Status += FString::Printf(TEXT("{yellow}%s/%03d {white}%.1f cm/s\n"),
-							*UEnum::GetDisplayValueAsText(MoveTarget.GetCurrentAction()).ToString(), MoveTarget.GetCurrentActionID(), Velocity.Value.Length());
+						Status += FString::Printf(TEXT("{yellow}%s/%03d {lightgrey}Speed:{white}%.1f {lightgrey}Force:{white}%.1f\n"),
+							*UEnum::GetDisplayValueAsText(MoveTarget.GetCurrentAction()).ToString(), MoveTarget.GetCurrentActionID(), Velocity.Value.Length(), Force.Value.Length());
 						Status += FString::Printf(TEXT("{pink}-> %s {white}Dist: %.1f\n"),
 							*UEnum::GetDisplayValueAsText(MoveTarget.IntentAtGoal).ToString(), MoveTarget.DistanceToGoal);
 
diff --git a/Engine/Plugins/AI/MassAI/Source/MassAIReplication/Private/MassReplicationPathHandlers.cpp b/Engine/Plugins/AI/MassAI/Source/MassAIReplication/Private/MassReplicationPathHandlers.cpp
--- a/Engine/Plugins/AI/MassAI/Source/MassAIReplication/Private/MassReplicationPathHandlers.cpp
+++ b/Engine/Plugins/AI/MassAI/Source/MassAIReplication/Private/MassReplicationPathHandlers.cpp
@@ -3,10 +3,8 @@
 #include "MassReplicationPathHandlers.h"
 #include "MassEntityQuery.h"
 #include "MassCommonFragments.h"
-#include "MassMovementTypes.h"
 #include "MassSimulationSubsystem.h"
-#include "MassZoneGraphMovementProcessors.h"
-#include "MassZoneGraphMovementUtils.h"
+#include "MassZoneGraphNavigationUtils.h"
 #include "ZoneGraphSubsystem.h"
 #include "GameFramework/GameStateBase.h"
 #include "VisualLogger/VisualLogger.h"
@@ -140,13 +138,13 @@ void FReplicatedAgentPathData::ApplyToEntity(const UWorld& InWorld, const FMassE
 	switch (Action)
 	{
 	case EMassMovementAction::Stand:
-		UE::MassMovement::ActivateActionStand(InWorld, SimulationSubsystem, Entity, *ZoneGraphSubsystem, LaneLocation, DesiredSpeed.Get(), MoveTarget, ShortPath, CachedLane);
+		UE::MassNavigation::ActivateActionStand(InWorld, SimulationSubsystem, Entity, *ZoneGraphSubsystem, LaneLocation, DesiredSpeed.Get(), MoveTarget, ShortPath, CachedLane);
 		break;
 	case EMassMovementAction::Move:
-		UE::MassMovement::ActivateActionMove(InWorld, SimulationSubsystem, Entity, *ZoneGraphSubsystem, LaneLocation, PathRequest, AgentRadius.Radius, DesiredSpeed.Get(), MoveTarget, ShortPath, CachedLane);
+		UE::MassNavigation::ActivateActionMove(InWorld, SimulationSubsystem, Entity, *ZoneGraphSubsystem, LaneLocation, PathRequest, AgentRadius.Radius, DesiredSpeed.Get(), MoveTarget, ShortPath, CachedLane);
 		break;
 	case EMassMovementAction::Animate:
-		UE::MassMovement::ActivateActionAnimate(InWorld, SimulationSubsystem, Entity, MoveTarget);
+		UE::MassNavigation::ActivateActionAnimate(InWorld, SimulationSubsystem, Entity, MoveTarget);
 		break;
 	default:
 		ensureMsgf(false, TEXT("Unhandled action type: %s"), *UEnum::GetValueAsString(Action));
diff --git a/Engine/Plugins/AI/MassAI/Source/MassAIReplication/Public/MassReplicationPathHandlers.h b/Engine/Plugins/AI/MassAI/Source/MassAIReplication/Public/MassReplicationPathHandlers.h
--- a/Engine/Plugins/AI/MassAI/Source/MassAIReplication/Public/MassReplicationPathHandlers.h
+++ b/Engine/Plugins/AI/MassAI/Source/MassAIReplication/Public/MassReplicationPathHandlers.h
@@ -5,8 +5,8 @@
 #include "MassReplicationTypes.h"
 #include "MassEntityView.h"
 #include "MassClientBubbleHandler.h"
-#include "MassAIMovementFragments.h"
-#include "MassZoneGraphMovementFragments.h"
+#include "MassNavigationFragments.h"
+#include "MassZoneGraphNavigationFragments.h"
 
 #include "MassReplicationPathHandlers.generated.h"
 
diff --git a/Engine/Plugins/AI/MassAI/Source/MassAITestSuite/MassAITestSuite.Build.cs b/Engine/Plugins/AI/MassAI/Source/MassAITestSuite/MassAITestSuite.Build.cs
--- a/Engine/Plugins/AI/MassAI/Source/MassAITestSuite/MassAITestSuite.Build.cs
+++ b/Engine/Plugins/AI/MassAI/Source/MassAITestSuite/MassAITestSuite.Build.cs
@@ -25,7 +25,7 @@ namespace UnrealBuildTool.Rules
 					"MassActors",
 					"MassAIBehavior",
 					"MassCommon",
-					"MassAIMovement",
+					"MassNavigation",
 					"MassAIReplication",
 					"MassSmartObjects",
 					"MassSpawner",
diff --git a/Engine/Plugins/AI/MassAI/Source/MassNavigation/Private/MassAvoidanceProcessors.cpp b/Engine/Plugins/AI/MassAI/Source/MassNavigation/Private/MassAvoidanceProcessors.cpp
new file mode 100644
--- /dev/null
+++ b/Engine/Plugins/AI/MassAI/Source/MassNavigation/Private/MassAvoidanceProcessors.cpp
@@ -0,0 +1,1249 @@
+// Copyright Epic Games, Inc. All Rights Reserved.
+
+#include "Avoidance/MassAvoidanceProcessors.h"
+#include "Avoidance/MassAvoidanceFragments.h"
+#include "DrawDebugHelpers.h"
+#include "MassEntityView.h"
+#include "VisualLogger/VisualLogger.h"
+#include "Math/Vector2D.h"
+#include "Logging/LogMacros.h"
+#include "MassSimulationLOD.h"
+#include "MassCommonFragments.h"
+#include "MassMovementFragments.h"
+#include "MassNavigationSubsystem.h"
+#include "MassNavigationFragments.h"
+#include "MassNavigationUtils.h"
+#include "Engine/World.h"
+
+#define UNSAFE_FOR_MT 1
+
+DEFINE_LOG_CATEGORY(LogAvoidance);
+DEFINE_LOG_CATEGORY(LogAvoidanceVelocities);
+DEFINE_LOG_CATEGORY(LogAvoidanceAgents);
+DEFINE_LOG_CATEGORY(LogAvoidanceObstacles);
+
+namespace UE::MassAvoidance
+{
+	namespace Tweakables
+	{
+		bool bEnableEnvironmentAvoidance = true;
+		bool bEnableSettingsforExtendingColliders = true;
+		bool bUseAdjacentCorridors = true;
+		bool bUseDrawDebugHelpers = false;
+	} // Tweakables
+
+	FAutoConsoleVariableRef Vars[] = 
+	{
+		FAutoConsoleVariableRef(TEXT("ai.mass.avoidance.EnableEnvironmentAvoidance"), Tweakables::bEnableEnvironmentAvoidance, TEXT("Set to false to disable avoidance forces for environment (for debug purposes)."), ECVF_Cheat),
+		FAutoConsoleVariableRef(TEXT("ai.mass.avoidance.EnableSettingsforExtendingColliders"), Tweakables::bEnableSettingsforExtendingColliders, TEXT("Set to false to disable using different settings for extending obstacles (for debug purposes)."), ECVF_Cheat),
+		FAutoConsoleVariableRef(TEXT("ai.mass.avoidance.UseAdjacentCorridors"), Tweakables::bUseAdjacentCorridors, TEXT("Set to false to disable usage of adjacent lane width."), ECVF_Cheat),
+		FAutoConsoleVariableRef(TEXT("ai.mass.avoidance.UseDrawDebugHelpers"), Tweakables::bUseDrawDebugHelpers, TEXT("Use debug draw helpers in addition to visual logs."), ECVF_Cheat)
+	};
+
+	constexpr int32 MaxExpectedAgentsPerCell = 6;
+	constexpr int32 MinTouchingCellCount = 4;
+	constexpr int32 MaxObstacleResults = MaxExpectedAgentsPerCell * MinTouchingCellCount;
+
+	static void FindCloseObstacles(const FVector& Center, const float SearchRadius, const FNavigationObstacleHashGrid2D& AvoidanceObstacleGrid,
+									TArray<FMassNavigationObstacleItem, TFixedAllocator<MaxObstacleResults>>& OutCloseEntities, const int32 MaxResults)
+	{
+		OutCloseEntities.Reset();
+		const FVector Extent(SearchRadius, SearchRadius, 0.f);
+		const FBox QueryBox = FBox(Center - Extent, Center + Extent);
+
+		struct FSortingCell
+		{
+			int32 X;
+			int32 Y;
+			int32 Level;
+			float SqDist;
+		};
+		TArray<FSortingCell, TInlineAllocator<64>> Cells;
+		const FVector QueryCenter = QueryBox.GetCenter();
+		
+		for (int32 Level = 0; Level < AvoidanceObstacleGrid.NumLevels; Level++)
+		{
+			const float CellSize = AvoidanceObstacleGrid.GetCellSize(Level);
+			const FNavigationObstacleHashGrid2D::FCellRect Rect = AvoidanceObstacleGrid.CalcQueryBounds(QueryBox, Level);
+			for (int32 Y = Rect.MinY; Y <= Rect.MaxY; Y++)
+			{
+				for (int32 X = Rect.MinX; X <= Rect.MaxX; X++)
+				{
+					const float CenterX = (X + 0.5f) * CellSize;
+					const float CenterY = (Y + 0.5f) * CellSize;
+					const float DX = CenterX - QueryCenter.X;
+					const float DY = CenterY - QueryCenter.Y;
+					const float SqDist = DX * DX + DY * DY;
+					FSortingCell SortCell;
+					SortCell.X = X;
+					SortCell.Y = Y;
+					SortCell.Level = Level;
+					SortCell.SqDist = SqDist;
+					Cells.Add(SortCell);
+				}
+			}
+		}
+
+		Cells.Sort([](const FSortingCell& A, const FSortingCell& B) { return A.SqDist < B.SqDist; });
+
+		for (const FSortingCell& SortedCell : Cells)
+		{
+			if (const FNavigationObstacleHashGrid2D::FCell* Cell = AvoidanceObstacleGrid.FindCell(SortedCell.X, SortedCell.Y, SortedCell.Level))
+			{
+				const TSparseArray<FNavigationObstacleHashGrid2D::FItem>&  Items = AvoidanceObstacleGrid.GetItems();
+				for (int32 Idx = Cell->First; Idx != INDEX_NONE; Idx = Items[Idx].Next)
+				{
+					OutCloseEntities.Add(Items[Idx].ID);
+					if (OutCloseEntities.Num() >= MaxResults)
+					{
+						return;
+					}
+				}
+			}
+		}
+	}
+
+	// Adapted from ray-capsule intersection: https://iquilezles.org/www/articles/intersectors/intersectors.htm
+	static float ComputeClosestPointOfApproach(const FVector2D Pos, const FVector2D Vel, const float Rad, const FVector2D SegStart, const FVector2D SegEnd, const float TimeHoriz)
+	{
+		const FVector2D SegDir = SegEnd - SegStart;
+		const FVector2D RelPos = Pos - SegStart;
+		const float VelSq = FVector2D::DotProduct(Vel, Vel);
+		const float SegDirSq = FVector2D::DotProduct(SegDir, SegDir);
+		const float DirVelSq = FVector2D::DotProduct(SegDir, Vel);
+		const float DirRelPosSq = FVector2D::DotProduct(SegDir, RelPos);
+		const float VelRelPosSq = FVector2D::DotProduct(Vel, RelPos);
+		const float RelPosSq = FVector2D::DotProduct(RelPos, RelPos);
+		const float A = SegDirSq * VelSq - DirVelSq * DirVelSq;
+		const float B = SegDirSq * VelRelPosSq - DirRelPosSq * DirVelSq;
+		const float C = SegDirSq * RelPosSq - DirRelPosSq * DirRelPosSq - FMath::Square(Rad) * SegDirSq;
+		const float H = FMath::Max<float>(0.f, B*B - A*C); // b^2 - ac, Using max for closest point of arrival result when no hit.
+		const float T = FMath::Abs(A) > SMALL_NUMBER ? (-B - FMath::Sqrt(H)) / A : 0.f;
+		const float Y = DirRelPosSq + T * DirVelSq;
+		
+		if (Y > 0.f && Y < SegDirSq) 
+		{
+			return FMath::Clamp(T, 0.f, TimeHoriz);
+		}
+		else 
+		{
+			// caps
+			const FVector2D CapRelPos = (Y <= 0.f) ? RelPos : Pos - SegEnd;
+			const float Cb = FVector2D::DotProduct(Vel, CapRelPos);
+			const float Cc = FVector2D::DotProduct(CapRelPos, CapRelPos) - FMath::Square(Rad);
+			const float Ch = FMath::Max<float>(0.0f, Cb * Cb - VelSq * Cc);
+			const float T1 = VelSq > SMALL_NUMBER ? (-Cb - FMath::Sqrt(Ch)) / VelSq : 0.f;
+			return FMath::Clamp(T1, 0.f, TimeHoriz);
+		}
+	}
+
+	static float ComputeClosestPointOfApproach(const FVector RelPos, const FVector RelVel, const float TotalRadius, const float TimeHoriz)
+	{
+		// Calculate time of impact based on relative agent positions and velocities.
+		const float A = FVector::DotProduct(RelVel, RelVel);
+		const float Inv2A = A > SMALL_NUMBER ? 1.f / (2.f * A) : 0.f;
+		const float B = FMath::Min(0.f, 2.f * FVector::DotProduct(RelVel, RelPos));
+		const float C = FVector::DotProduct(RelPos, RelPos) - FMath::Square(TotalRadius);
+		// Using max() here gives us CPA (closest point on arrival) when there is no hit.
+		const float Discr = FMath::Sqrt(FMath::Max(0.f, B * B - 4.f * A * C));
+		const float T = (-B - Discr) * Inv2A;
+		return FMath::Clamp(T, 0.f, TimeHoriz);
+	}
+
+	static bool UseDrawDebugHelper()
+	{
+		return Tweakables::bUseDrawDebugHelpers;
+	}
+
+#if WITH_MASSGAMEPLAY_DEBUG	
+
+	// Colors
+	static const FColor CurrentAgentColor = FColor::Emerald;
+
+	static const FColor VelocityColor = FColor::Black;
+	static const FColor PrefVelocityColor = FColor::Red;
+	static const FColor DesiredVelocityColor = FColor::Yellow;
+	static const FColor FinalSteeringForceColor = FColor::Cyan;
+	static constexpr float BigArrowThickness = 6.f;
+	static constexpr float BigArrowHeadSize = 12.f;
+
+	// Agents colors
+	static const FColor AgentsColor = FColor::Orange;
+	static const FColor AgentSeparationForceColor = FColor(255, 145, 71);	// Orange red
+	static const FColor AgentAvoidForceColor = AgentsColor;
+	
+	// Obstacles colors
+	static const FColor ObstacleColor = FColor::Blue;
+	static const FColor ObstacleContactNormalColor = FColor::Silver;
+	static const FColor ObstacleAvoidForceColor = FColor::Magenta;
+	static const FColor ObstacleSeparationForceColor = FColor(255, 66, 66);	// Bright red
+	
+	static const FVector DebugAgentHeightOffset = FVector(0.f, 0.f, 185.f);
+	static const FVector DebugLowCylinderOffset = FVector(0.f, 0.f, 20.f);
+
+	//----------------------------------------------------------------------//
+	// Begin MassDebugUtils
+	// @todo: Extract those generic debug functions to a separate location
+	//----------------------------------------------------------------------//
+	struct FDebugContext
+	{
+		FDebugContext(const UObject* InLogOwner, const FLogCategoryBase& InCategory, const UWorld* InWorld, const FMassEntityHandle InEntity)
+			: LogOwner(InLogOwner)
+			, Category(InCategory)
+			, World(InWorld)
+			, Entity(InEntity)
+		{}
+
+		const UObject* LogOwner;
+		const FLogCategoryBase& Category;
+		const UWorld* World;
+		const FMassEntityHandle Entity;
+	};
+
+	static bool DebugIsSelected(const FMassEntityHandle Entity)
+	{
+		FColor Color;
+		return UE::Mass::Debug::IsDebuggingEntity(Entity, &Color);
+	}
+
+	static void DebugDrawLine(const FDebugContext& Context, const FVector& Start, const FVector& End, const FColor& Color, const float Thickness = 0.f, const bool bPersistent = false)
+	{
+		if (!DebugIsSelected(Context.Entity))
+		{
+			return;
+		}
+
+		UE_VLOG_SEGMENT_THICK(Context.LogOwner, Context.Category, Log, Start, End, Color, (int16)Thickness, TEXT(""));
+
+		if (UseDrawDebugHelper() && Context.World)
+		{
+			DrawDebugLine(Context.World, Start, End, Color, bPersistent, /*LifeTime =*/ -1.f, /*DepthPriority =*/ 0, Thickness);
+		}
+	}
+
+	static void DebugDrawArrow(const FDebugContext& Context, const FVector& Start, const FVector& End, const FColor& Color, const float HeadSize = 8.f, const float Thickness = 1.5f)
+	{
+		if (!DebugIsSelected(Context.Entity))
+		{
+			return;
+		}
+
+		const float Pointyness = 1.8f;
+		const FVector Line = End - Start;
+		const FVector UnitV = Line.GetSafeNormal();
+		const FVector Perp = FVector::CrossProduct(UnitV, FVector::UpVector);
+		const FVector Left = Perp - (Pointyness*UnitV);
+		const FVector Right = -Perp - (Pointyness*UnitV);
+		UE_VLOG_SEGMENT_THICK(Context.LogOwner, Context.Category, Log, Start, End, Color, (int16)Thickness, TEXT(""));
+		UE_VLOG_SEGMENT_THICK(Context.LogOwner, Context.Category, Log, End, End + HeadSize * Left, Color, (int16)Thickness, TEXT(""));
+		UE_VLOG_SEGMENT_THICK(Context.LogOwner, Context.Category, Log, End, End + HeadSize * Right, Color, (int16)Thickness, TEXT(""));
+
+		if (UseDrawDebugHelper() && Context.World)
+		{
+			DrawDebugLine(Context.World, Start, End, Color, /*bPersistent=*/ false, /*LifeTime =*/ -1.f, /*DepthPriority =*/ 0, Thickness);
+			DrawDebugLine(Context.World, End, End + HeadSize * Left, Color, /*bPersistent=*/ false, /*LifeTime =*/ -1.f, /*DepthPriority =*/ 0, Thickness);
+			DrawDebugLine(Context.World, End, End + HeadSize * Right, Color, /*bPersistent=*/ false, /*LifeTime =*/ -1.f, /*DepthPriority =*/ 0, Thickness);
+		}
+	}
+
+	static void DebugDrawSphere(const FDebugContext& Context, const FVector& Center, const float Radius, const FColor& Color)
+	{
+		if (!DebugIsSelected(Context.Entity))
+		{
+			return;
+		}
+
+		UE_VLOG_LOCATION(Context.LogOwner, Context.Category, Log, Center, Radius, Color, TEXT(""));
+
+		if (UseDrawDebugHelper() && Context.World)
+		{
+			DrawDebugSphere(Context.World, Center, Radius, /*segments = */16, Color);
+		}
+	}
+
+	static void DebugDrawBox(const FDebugContext& Context, const FBox& Box, const FColor& Color)
+	{
+		if (!DebugIsSelected(Context.Entity))
+		{
+			return;
+		}
+
+		UE_VLOG_BOX(Context.LogOwner, Context.Category, Log, Box, Color, TEXT(""));
+		
+		if (UseDrawDebugHelper() && Context.World)
+		{
+			DrawDebugBox(Context.World, Box.GetCenter(), Box.GetExtent(), Color);
+		}
+	}
+	
+	static void DebugDrawCylinder(const FDebugContext& Context, const FVector& Bottom, const FVector& Top, const float Radius, const FColor& Color, const FString& Text = FString())
+	{
+		if (!DebugIsSelected(Context.Entity))
+		{
+			return;
+		}
+
+		UE_VLOG_CYLINDER(Context.LogOwner, Context.Category, Log, Bottom, Top, Radius, Color, TEXT("%s"), *Text);
+
+		if (UseDrawDebugHelper() && Context.World)
+		{
+			DrawDebugCylinder(Context.World, Bottom, Top, Radius, /*segments = */24, Color);
+		}
+	}
+	//----------------------------------------------------------------------//
+	// End MassDebugUtils
+	//----------------------------------------------------------------------//
+
+
+	// Local debug utils
+	static void DebugDrawVelocity(const FDebugContext& Context, const FVector& Start, const FVector& End, const FColor& Color)
+	{
+		// Different arrow than DebugDrawArrow()
+		if (!DebugIsSelected(Context.Entity))
+		{
+			return;
+		}
+
+		const float Thickness = 3.f;
+		const float Pointyness = 1.8f;
+		const FVector Line = End - Start;
+		const FVector UnitV = Line.GetSafeNormal();
+		const FVector Perp = FVector::CrossProduct(UnitV, FVector::UpVector);
+		const FVector Left = Perp - (Pointyness * UnitV);
+		const FVector Right = -Perp - (Pointyness * UnitV);
+		const float HeadSize = 0.08f * Line.Size();
+		UE_VLOG_SEGMENT_THICK(Context.LogOwner, Context.Category, Log, Start, End, Color, (int16)Thickness, TEXT(""));
+		UE_VLOG_SEGMENT_THICK(Context.LogOwner, Context.Category, Log, End, End + HeadSize * Left, Color, (int16)Thickness, TEXT(""));
+		UE_VLOG_SEGMENT_THICK(Context.LogOwner, Context.Category, Log, End, End + HeadSize * Right, Color, (int16)Thickness, TEXT(""));
+		UE_VLOG_SEGMENT_THICK(Context.LogOwner, Context.Category, Log, End + HeadSize * Left, End + HeadSize * Right, Color, (int16)Thickness, TEXT(""));
+
+		if (UseDrawDebugHelper() && Context.World)
+		{
+			DrawDebugLine(Context.World, Start, End, Color, /*bPersistent=*/ false, /*LifeTime =*/ -1.f, /*DepthPriority =*/ 0, Thickness);
+			DrawDebugLine(Context.World, End, End + HeadSize * Left, Color, /*bPersistent=*/ false, /*LifeTime =*/ -1.f, /*DepthPriority =*/ 0, Thickness);
+			DrawDebugLine(Context.World, End, End + HeadSize * Right, Color, /*bPersistent=*/ false, /*LifeTime =*/ -1.f, /*DepthPriority =*/ 0, Thickness);
+			DrawDebugLine(Context.World, End + HeadSize * Left, End + HeadSize * Right, Color, /*bPersistent=*/ false, /*LifeTime =*/ -1.f, /*DepthPriority =*/ 0, Thickness);
+		}
+	}
+
+	static void DebugDrawForce(const FDebugContext& Context, const FVector& Start, const FVector& End, const FColor& Color)
+	{
+		DebugDrawArrow(Context, Start, End, Color, /*HeadSize*/4.f, /*Thickness*/3.f);
+	}
+
+	static void DebugDrawSummedForce(const FDebugContext& Context, const FVector& Start, const FVector& End, const FColor& Color)
+	{
+		DebugDrawArrow(Context, Start + FVector(0.f,0.f,1.f), End + FVector(0.f, 0.f, 1.f), Color, /*HeadSize*/8.f, /*Thickness*/6.f);
+	}
+
+#endif // WITH_MASSGAMEPLAY_DEBUG
+
+} // namespace UE::MassAvoidance
+
+
+//----------------------------------------------------------------------//
+//  UMassMovingAvoidanceProcessor
+//----------------------------------------------------------------------//
+UMassMovingAvoidanceProcessor::UMassMovingAvoidanceProcessor()
+{
+	bAutoRegisterWithProcessingPhases = true;
+	ExecutionFlags = (int32)EProcessorExecutionFlags::All;
+	ExecutionOrder.ExecuteInGroup = UE::Mass::ProcessorGroupNames::Avoidance;
+	ExecutionOrder.ExecuteAfter.Add(UE::Mass::ProcessorGroupNames::LOD);
+}
+
+void UMassMovingAvoidanceProcessor::ConfigureQueries()
+{
+	EntityQuery.AddRequirement<FMassForceFragment>(EMassFragmentAccess::ReadWrite);
+	EntityQuery.AddRequirement<FMassNavigationEdgesFragment>(EMassFragmentAccess::ReadOnly);
+	EntityQuery.AddRequirement<FMassMoveTargetFragment>(EMassFragmentAccess::ReadOnly);
+	EntityQuery.AddRequirement<FDataFragment_Transform>(EMassFragmentAccess::ReadOnly);
+	EntityQuery.AddRequirement<FMassVelocityFragment>(EMassFragmentAccess::ReadOnly);
+	EntityQuery.AddRequirement<FDataFragment_AgentRadius>(EMassFragmentAccess::ReadOnly);
+	EntityQuery.AddTagRequirement<FMassMediumLODTag>(EMassFragmentPresence::None);
+	EntityQuery.AddTagRequirement<FMassLowLODTag>(EMassFragmentPresence::None);
+	EntityQuery.AddTagRequirement<FMassOffLODTag>(EMassFragmentPresence::None);
+	EntityQuery.AddConstSharedRequirement<FMassMovingAvoidanceParameters>(EMassFragmentPresence::All);
+	EntityQuery.AddConstSharedRequirement<FMassMovementParameters>(EMassFragmentPresence::All);
+}
+
+void UMassMovingAvoidanceProcessor::Initialize(UObject& Owner)
+{
+	Super::Initialize(Owner);
+
+	World = Owner.GetWorld();
+	NavigationSubsystem = UWorld::GetSubsystem<UMassNavigationSubsystem>(Owner.GetWorld());
+}
+
+void UMassMovingAvoidanceProcessor::Execute(UMassEntitySubsystem& EntitySubsystem, FMassExecutionContext& Context)
+{
+	QUICK_SCOPE_CYCLE_COUNTER(UMassMovingAvoidanceProcessor);
+
+	if (!World || !NavigationSubsystem)
+	{
+		return;
+	}
+
+	EntityQuery.ForEachEntityChunk(EntitySubsystem, Context, [this, &EntitySubsystem](FMassExecutionContext& Context)
+	{
+		const float DeltaTime = Context.GetDeltaTimeSeconds();
+		const float CurrentTime = World->GetTimeSeconds();
+		const int32 NumEntities = Context.GetNumEntities();
+		
+		const TArrayView<FMassForceFragment> ForceList = Context.GetMutableFragmentView<FMassForceFragment>();
+		const TConstArrayView<FMassNavigationEdgesFragment> NavEdgesList = Context.GetFragmentView<FMassNavigationEdgesFragment>();
+		const TConstArrayView<FDataFragment_Transform> LocationList = Context.GetFragmentView<FDataFragment_Transform>();
+		const TConstArrayView<FMassVelocityFragment> VelocityList = Context.GetFragmentView<FMassVelocityFragment>();
+		const TConstArrayView<FDataFragment_AgentRadius> RadiusList = Context.GetFragmentView<FDataFragment_AgentRadius>();
+		const TConstArrayView<FMassMoveTargetFragment> MoveTargetList = Context.GetFragmentView<FMassMoveTargetFragment>();
+		const FMassMovingAvoidanceParameters& MovingAvoidanceParams = Context.GetConstSharedFragment<FMassMovingAvoidanceParameters>();
+		const FMassMovementParameters& MovementParams = Context.GetConstSharedFragment<FMassMovementParameters>();
+
+		const float InvPredictiveAvoidanceTime = 1.0f / MovingAvoidanceParams.PredictiveAvoidanceTime;
+
+		// Arrays used to store close obstacles
+		TArray<FMassNavigationObstacleItem, TFixedAllocator<UE::MassAvoidance::MaxObstacleResults>> CloseEntities;
+
+		// Used for storing sorted list or nearest obstacles.
+		struct FSortedObstacle
+		{
+			FVector LocationCached;
+			FVector Forward;
+			FMassNavigationObstacleItem ObstacleItem;
+			float SqDist;
+		};
+		TArray<FSortedObstacle, TFixedAllocator<UE::MassAvoidance::MaxObstacleResults>> ClosestObstacles;
+
+		// Potential contact between agent and environment. 
+		struct FEnvironmentContact
+		{
+			FVector Position = FVector::ZeroVector;
+			FVector Normal = FVector::ZeroVector;
+			float Distance = 0.f;
+		};
+		TArray<FEnvironmentContact, TInlineAllocator<16>> Contacts;
+
+		// Describes collider to avoid, collected from neighbour obstacles.
+		struct FCollider
+		{
+			FVector Location = FVector::ZeroVector;
+			FVector Velocity = FVector::ZeroVector;
+			float Radius = 0.f;
+			bool bCanAvoid = true;
+			bool bIsMoving = false;
+		};
+		TArray<FCollider, TInlineAllocator<16>> Colliders;
+
+		for (int32 EntityIndex = 0; EntityIndex < NumEntities; ++EntityIndex)
+		{
+			// @todo: this check should eventually be part of the query (i.e. only handle moving agents).
+			const FMassMoveTargetFragment& MoveTarget = MoveTargetList[EntityIndex];
+			if (MoveTarget.GetCurrentAction() == EMassMovementAction::Animate || MoveTarget.GetCurrentAction() == EMassMovementAction::Stand)
+			{
+				continue;
+			}
+
+			FMassEntityHandle Entity = Context.GetEntity(EntityIndex);
+			const FMassNavigationEdgesFragment& NavEdges = NavEdgesList[EntityIndex];
+			const FDataFragment_Transform& Location = LocationList[EntityIndex];
+			const FMassVelocityFragment& Velocity = VelocityList[EntityIndex];
+			const FDataFragment_AgentRadius& Radius = RadiusList[EntityIndex];
+			FMassForceFragment& Force = ForceList[EntityIndex];
+
+			// Smaller steering max accel makes the steering more "calm" but less opportunistic, may not find solution, or gets stuck.
+			// Max contact accel should be quite a big bigger than steering so that collision response is firm. 
+			const float MaxSteerAccel = MovementParams.MaxAcceleration;
+			const float MaximumSpeed = MovementParams.MaxSpeed;
+
+			const FVector AgentLocation = Location.GetTransform().GetTranslation();
+			const FVector AgentVelocity = FVector(Velocity.Value.X, Velocity.Value.Y, 0.0f);
+			
+			const float AgentRadius = Radius.Radius;
+			const float SeparationAgentRadius = Radius.Radius * MovingAvoidanceParams.SeparationRadiusScale;
+			const float PredictiveAvoidanceAgentRadius = Radius.Radius * MovingAvoidanceParams.PredictiveAvoidanceRadiusScale;
+			
+			FVector SteeringForce = Force.Value;
+
+			// Near start and end fades are used to subdue the avoidance at the start and end of the path.
+			float NearStartFade = 1.0f;
+			float NearEndFade = 1.0f;
+
+			if (MoveTarget.GetPreviousAction() != EMassMovementAction::Move)
+			{
+				// Fade in avoidance when transitioning from other than move action.
+				// I.e. the standing behavior may move the agents so close to each,
+				// and that causes the separation to push them out quickly when avoidance is activated. 
+				NearStartFade = FMath::Min((CurrentTime - MoveTarget.GetCurrentActionStartTime()) / MovingAvoidanceParams.StartOfPathDuration, 1.0f);
+			}
+
+			if (MoveTarget.IntentAtGoal == EMassMovementAction::Stand)
+			{
+				// Estimate approach based on current desired speed.
+				const float ApproachDistance = FMath::Max(1.0f, MovingAvoidanceParams.EndOfPathDuration * MoveTarget.DesiredSpeed.Get());
+				NearEndFade = FMath::Clamp(MoveTarget.DistanceToGoal / ApproachDistance, 0.f, 1.f);
+			}
+			
+			const float NearStartScaling = FMath::Lerp(MovingAvoidanceParams.StartOfPathAvoidanceScale, 1.0f, NearStartFade);
+			const float NearEndScaling = FMath::Lerp(MovingAvoidanceParams.EndOfPathAvoidanceScale, 1.0f, NearEndFade);
+			
+#if WITH_MASSGAMEPLAY_DEBUG && UNSAFE_FOR_MT
+			const UE::MassAvoidance::FDebugContext BaseDebugContext(this, LogAvoidance, World, Entity);
+			const UE::MassAvoidance::FDebugContext VelocitiesDebugContext(this, LogAvoidanceVelocities, World, Entity);
+			const UE::MassAvoidance::FDebugContext ObstacleDebugContext(this, LogAvoidanceObstacles, World, Entity);
+			const UE::MassAvoidance::FDebugContext AgentDebugContext(this, LogAvoidanceAgents, World, Entity);
+			
+			if (UE::MassAvoidance::DebugIsSelected(Entity))
+			{
+				// Draw agent
+				const FString Text = FString::Printf(TEXT("%i"), Entity.Index);
+				DebugDrawCylinder(BaseDebugContext, AgentLocation, AgentLocation + UE::MassAvoidance::DebugAgentHeightOffset, AgentRadius+1.f, UE::MassAvoidance::CurrentAgentColor, Text);
+
+				DebugDrawSphere(BaseDebugContext, AgentLocation, 10.f, UE::MassAvoidance::CurrentAgentColor);
+
+				// Draw current velocity (black)
+				DebugDrawVelocity(VelocitiesDebugContext, AgentLocation + UE::MassAvoidance::DebugAgentHeightOffset, AgentLocation + UE::MassAvoidance::DebugAgentHeightOffset + AgentVelocity, UE::MassAvoidance::VelocityColor);
+
+				// Draw preferred velocity (red)
+//				DebugDrawVelocity(VelocitiesDebugContext, AgentLocation + UE::MassAvoidance::DebugAgentHeightOffset, AgentLocation + UE::MassAvoidance::DebugAgentHeightOffset + PrefVelocity, UE::MassAvoidance::PrefVelocityColor);
+
+				// Draw initial steering force
+				DebugDrawArrow(BaseDebugContext, AgentLocation + UE::MassAvoidance::DebugAgentHeightOffset, AgentLocation +UE::MassAvoidance:: DebugAgentHeightOffset + SteeringForce, UE::MassAvoidance::CurrentAgentColor, UE::MassAvoidance::BigArrowHeadSize, UE::MassAvoidance::BigArrowThickness);
+
+				// Draw center
+				DebugDrawSphere(BaseDebugContext, AgentLocation, /*Radius*/2.f, UE::MassAvoidance::CurrentAgentColor);
+			}
+#endif // WITH_MASSGAMEPLAY_DEBUG
+
+			FVector OldSteeringForce = FVector::ZeroVector;
+
+			//////////////////////////////////////////////////////////////////////////
+			// Environment avoidance.
+			//
+			
+			if (!MoveTarget.bOffBoundaries && UE::MassAvoidance::Tweakables::bEnableEnvironmentAvoidance)
+			{
+				const FVector DesiredAcceleration = UE::MassNavigation::ClampVector(SteeringForce, MaxSteerAccel);
+				const FVector DesiredVelocity = UE::MassNavigation::ClampVector(AgentVelocity + DesiredAcceleration * DeltaTime, MaximumSpeed);
+
+#if WITH_MASSGAMEPLAY_DEBUG && UNSAFE_FOR_MT
+				// Draw desired velocity (yellow)
+				UE::MassAvoidance::DebugDrawVelocity(VelocitiesDebugContext, AgentLocation + UE::MassAvoidance::DebugAgentHeightOffset, AgentLocation + UE::MassAvoidance::DebugAgentHeightOffset + DesiredVelocity, UE::MassAvoidance::DesiredVelocityColor);
+#endif // WITH_MASSGAMEPLAY_DEBUG
+
+				OldSteeringForce = SteeringForce;
+				Contacts.Reset();
+
+				// Collect potential contacts between agent and environment edges.
+				for (const FNavigationAvoidanceEdge& Edge : NavEdges.AvoidanceEdges)
+				{
+					const FVector EdgeDiff = Edge.End - Edge.Start;
+					FVector EdgeDir = FVector::ZeroVector;
+					float EdgeLength = 0.0f;
+					EdgeDiff.ToDirectionAndLength(EdgeDir, EdgeLength);
+
+					const FVector AgentToEdgeStart = AgentLocation - Edge.Start;
+					const float DistAlongEdge = FVector::DotProduct(EdgeDir, AgentToEdgeStart);
+					const float DistAwayFromEdge = FVector::DotProduct(Edge.LeftDir, AgentToEdgeStart);
+
+					float ConDist = 0.0f;
+					FVector ConNorm = FVector::ForwardVector;
+					FVector ConPos = FVector::ZeroVector;
+					bool bDirectlyBehindEdge = false;
+					
+					if (DistAwayFromEdge < 0.0f)
+					{
+						// Inside or behind the edge
+						if (DistAlongEdge < 0.0f)
+						{
+							ConPos = Edge.Start;
+							ConNorm = -EdgeDir;
+							ConDist = -DistAlongEdge;
+						}
+						else if (DistAlongEdge > EdgeLength)
+						{
+							ConPos = Edge.End;
+							ConNorm = EdgeDir;
+							ConDist = DistAlongEdge;
+						}
+						else
+						{
+							ConPos = Edge.Start + EdgeDir * DistAlongEdge;
+							ConNorm = Edge.LeftDir;
+							ConDist = 0.0f;
+							bDirectlyBehindEdge = true;
+						}
+					}
+					else
+					{
+						if (DistAlongEdge < 0.0f)
+						{
+							// Start Corner
+							ConPos = Edge.Start;
+							EdgeDiff.ToDirectionAndLength(ConNorm, ConDist);
+						}
+						else if (DistAlongEdge > EdgeLength)
+						{
+							// End Corner
+							ConPos = Edge.End;
+							EdgeDiff.ToDirectionAndLength(ConNorm, ConDist);
+						}
+						else
+						{
+							// Front
+							ConPos = Edge.Start + EdgeDir * DistAlongEdge;
+							ConNorm = Edge.LeftDir;
+							ConDist = DistAwayFromEdge;
+						}
+					}
+					
+					// Check to merge contacts
+					bool bAdd = true;
+					for (int ContactIndex = 0; ContactIndex < Contacts.Num(); ContactIndex++)
+					{
+						if (FVector::DotProduct(Contacts[ContactIndex].Normal, ConNorm) > 0.f && FMath::Abs(FVector::DotProduct(ConNorm, Contacts[ContactIndex].Position - ConPos)) < (10.f/*cm*/))
+						{
+							// Contacts are on same place, merge
+							if (ConDist < Contacts[ContactIndex].Distance)
+							{
+								// New is closer, override.
+								Contacts[ContactIndex].Position = ConPos;
+								Contacts[ContactIndex].Normal = ConNorm;
+								Contacts[ContactIndex].Distance = ConDist;
+							}
+							bAdd = false;
+							break;
+						}
+					}
+
+					// Not found, add new contact
+					if (bAdd)
+					{
+						FEnvironmentContact Contact;
+						Contact.Position = ConPos;
+						Contact.Normal = ConNorm;
+						Contact.Distance = ConDist;
+						Contacts.Add(Contact);
+					}
+
+					// Skip predictive avoidance when behind the edge.
+					if (!bDirectlyBehindEdge)
+					{
+						// Avoid edges
+						const float CPA = UE::MassAvoidance::ComputeClosestPointOfApproach(FVector2D(AgentLocation), FVector2D(DesiredVelocity), AgentRadius,
+							FVector2D(Edge.Start), FVector2D(Edge.End), MovingAvoidanceParams.PredictiveAvoidanceTime);
+						const FVector HitAgentPos = AgentLocation + DesiredVelocity * CPA;
+						const float EdgeT = UE::MassNavigation::ProjectPtSeg(FVector2D(HitAgentPos), FVector2D(Edge.Start), FVector2D(Edge.End));
+						const FVector HitObPos = FMath::Lerp(Edge.Start, Edge.End, EdgeT);
+
+						// Calculate penetration at CPA
+						FVector AvoidRelPos = HitAgentPos - HitObPos;
+						AvoidRelPos.Z = 0.f;	// @todo AT: ignore the z component for now until we clamp the height of obstacles
+						const float AvoidDist = AvoidRelPos.Size();
+						const FVector AvoidNormal = AvoidDist > 0.0f ? (AvoidRelPos / AvoidDist) : FVector::ForwardVector;
+
+						const float AvoidPen = (PredictiveAvoidanceAgentRadius + MovingAvoidanceParams.PredictiveAvoidanceDistance) - AvoidDist;
+						const float AvoidMag = FMath::Square(FMath::Clamp(AvoidPen / MovingAvoidanceParams.PredictiveAvoidanceDistance, 0.f, 1.f));
+						const float AvoidMagDist = 1.f + FMath::Square(1.f - CPA * InvPredictiveAvoidanceTime);
+						const FVector AvoidForce = AvoidNormal * AvoidMag * AvoidMagDist * MovingAvoidanceParams.EnvironmentPredictiveAvoidanceStiffness * NearEndScaling; // Predictive avoidance against environment is tuned down towards the end of the path
+
+						SteeringForce += AvoidForce;
+
+#if WITH_MASSGAMEPLAY_DEBUG && UNSAFE_FOR_MT
+						// Draw contact normal
+						UE::MassAvoidance::DebugDrawArrow(ObstacleDebugContext, ConPos, ConPos + 50.f * ConNorm, UE::MassAvoidance::ObstacleContactNormalColor, /*HeadSize=*/ 5.f);
+						UE::MassAvoidance::DebugDrawSphere(ObstacleDebugContext, ConPos, 2.5f, UE::MassAvoidance::ObstacleContactNormalColor);
+
+						// Draw hit pos with edge
+						UE::MassAvoidance::DebugDrawLine(ObstacleDebugContext, AgentLocation, HitAgentPos, UE::MassAvoidance::ObstacleAvoidForceColor);
+						UE::MassAvoidance::DebugDrawCylinder(ObstacleDebugContext, HitAgentPos, HitAgentPos + UE::MassAvoidance::DebugAgentHeightOffset, AgentRadius, UE::MassAvoidance::ObstacleAvoidForceColor);
+
+						// Draw avoid obstacle force
+						UE::MassAvoidance::DebugDrawForce(ObstacleDebugContext, HitObPos, HitObPos + AvoidForce, UE::MassAvoidance::ObstacleAvoidForceColor);
+#endif // WITH_MASSGAMEPLAY_DEBUG
+					}
+				} // edge loop
+
+#if WITH_MASSGAMEPLAY_DEBUG && UNSAFE_FOR_MT
+				// Draw total steering force to avoid obstacles
+				const FVector EnvironmentAvoidSteeringForce = SteeringForce - OldSteeringForce;
+				UE::MassAvoidance::DebugDrawSummedForce(ObstacleDebugContext,
+					AgentLocation + UE::MassAvoidance::DebugAgentHeightOffset,
+					AgentLocation + UE::MassAvoidance::DebugAgentHeightOffset + EnvironmentAvoidSteeringForce,
+					UE::MassAvoidance::ObstacleAvoidForceColor);
+#endif // WITH_MASSGAMEPLAY_DEBUG
+
+				// Process contacts to add edge separation force
+				const FVector SteeringForceBeforeSeparation = SteeringForce;
+				for (int ContactIndex = 0; ContactIndex < Contacts.Num(); ContactIndex++) 
+				{
+					const FVector ConNorm = Contacts[ContactIndex].Normal.GetSafeNormal();
+					const float ContactDist = Contacts[ContactIndex].Distance;
+
+					// Separation force (stay away from obstacles if possible)
+					const float SeparationPenalty = (SeparationAgentRadius + MovingAvoidanceParams.EnvironmentSeparationDistance) - ContactDist;
+					const float SeparationMag = UE::MassNavigation::Smooth(FMath::Clamp(SeparationPenalty / MovingAvoidanceParams.EnvironmentSeparationDistance, 0.f, 1.f));
+					const FVector SeparationForce = ConNorm * MovingAvoidanceParams.EnvironmentSeparationStiffness * SeparationMag;
+
+					SteeringForce += SeparationForce;
+
+#if WITH_MASSGAMEPLAY_DEBUG && UNSAFE_FOR_MT
+					// Draw individual contact forces
+					DebugDrawForce(ObstacleDebugContext, Contacts[ContactIndex].Position + UE::MassAvoidance::DebugAgentHeightOffset,
+					Contacts[ContactIndex].Position + SeparationForce + UE::MassAvoidance::DebugAgentHeightOffset, UE::MassAvoidance::ObstacleSeparationForceColor);
+#endif // WITH_MASSGAMEPLAY_DEBUG
+				}
+				
+#if WITH_MASSGAMEPLAY_DEBUG && UNSAFE_FOR_MT
+				// Draw total steering force to separate from close edges
+				const FVector TotalSeparationForce = SteeringForce - SteeringForceBeforeSeparation;
+				DebugDrawSummedForce(ObstacleDebugContext,
+					AgentLocation + UE::MassAvoidance::DebugAgentHeightOffset,
+					AgentLocation + UE::MassAvoidance::DebugAgentHeightOffset + TotalSeparationForce,
+					UE::MassAvoidance::ObstacleSeparationForceColor);
+
+				// Display close obstacle edges
+				if (UE::MassAvoidance::DebugIsSelected(Entity))
+				{
+					for (const FNavigationAvoidanceEdge& Edge : NavEdges.AvoidanceEdges)
+					{
+						DebugDrawLine(ObstacleDebugContext, UE::MassAvoidance::DebugAgentHeightOffset + Edge.Start,
+							UE::MassAvoidance::DebugAgentHeightOffset + Edge.End, UE::MassAvoidance::ObstacleColor, /*Thickness=*/2.f);
+						const FVector Middle = UE::MassAvoidance::DebugAgentHeightOffset + 0.5f * (Edge.Start + Edge.End);
+						DebugDrawArrow(ObstacleDebugContext, Middle, Middle + 10.f * FVector::CrossProduct((Edge.End - Edge.Start), FVector::UpVector).GetSafeNormal(), UE::MassAvoidance::ObstacleColor, /*HeadSize=*/2.f);
+					}
+				}
+#endif // WITH_MASSGAMEPLAY_DEBUG
+			}
+
+			//////////////////////////////////////////////////////////////////////////
+			// Avoid close agents
+
+			// Update desired velocity based on avoidance so far.
+			const FVector DesAcc = UE::MassNavigation::ClampVector(SteeringForce, MaxSteerAccel);
+			const FVector DesVel = UE::MassNavigation::ClampVector(AgentVelocity + DesAcc * DeltaTime, MaximumSpeed);
+
+			// Find close obstacles
+			const FNavigationObstacleHashGrid2D& AvoidanceObstacleGrid = NavigationSubsystem->GetObstacleGridMutable();
+			UE::MassAvoidance::FindCloseObstacles(AgentLocation, MovingAvoidanceParams.ObstacleDetectionDistance, AvoidanceObstacleGrid, CloseEntities, UE::MassAvoidance::MaxObstacleResults);
+
+			// Remove unwanted and find the closests in the CloseEntities
+			const float DistanceCutOffSqr = FMath::Square(MovingAvoidanceParams.ObstacleDetectionDistance);
+			ClosestObstacles.Reset();
+			for (const FNavigationObstacleHashGrid2D::ItemIDType OtherEntity : CloseEntities)
+			{
+				// Skip self
+				if (OtherEntity.Entity == Entity)
+				{
+					continue;
+				}
+
+				// Skip invalid entities.
+				if (!EntitySubsystem.IsEntityValid(OtherEntity.Entity))
+				{
+					UE_LOG(LogAvoidanceObstacles, VeryVerbose, TEXT("Close entity is invalid, skipped."));
+					continue;
+				}
+				
+				// Skip too far
+				const FTransform& Transform = EntitySubsystem.GetFragmentDataChecked<FDataFragment_Transform>(OtherEntity.Entity).GetTransform();
+				const FVector OtherLocation = Transform.GetLocation();
+				
+				const float SqDist = FVector::DistSquared(AgentLocation, OtherLocation);
+				if (SqDist > DistanceCutOffSqr)
+				{
+					continue;
+				}
+
+				FSortedObstacle Obstacle;
+				Obstacle.LocationCached = OtherLocation;
+				Obstacle.Forward = Transform.GetRotation().GetForwardVector();
+				Obstacle.ObstacleItem = OtherEntity;
+				Obstacle.SqDist = SqDist;
+				ClosestObstacles.Add(Obstacle);
+			}
+			ClosestObstacles.Sort([](const FSortedObstacle& A, const FSortedObstacle& B) { return A.SqDist < B.SqDist; });
+
+			// Compute forces
+			OldSteeringForce = SteeringForce;
+			FVector TotalAgentSeparationForce = FVector::ZeroVector;
+
+			// Fill collider list from close agents
+			Colliders.Reset();
+			constexpr int32 MaxColliders = 6;
+			for (int32 Index = 0; Index < ClosestObstacles.Num(); Index++)
+			{
+				if (Colliders.Num() >= MaxColliders)
+				{
+					break;
+				}
+
+				FSortedObstacle& Obstacle = ClosestObstacles[Index];
+				FMassEntityView OtherEntityView(EntitySubsystem, Obstacle.ObstacleItem.Entity);
+
+				const FMassVelocityFragment* OtherVelocityFragment = OtherEntityView.GetFragmentDataPtr<FMassVelocityFragment>();
+				const FVector OtherVelocity = OtherVelocityFragment != nullptr ? OtherVelocityFragment->Value : FVector::ZeroVector; // Get velocity from FAvoidanceComponent
+
+				// @todo: this is heavy fragment to access, see if we could handle this differently.
+				const FMassMoveTargetFragment* OtherMoveTarget = OtherEntityView.GetFragmentDataPtr<FMassMoveTargetFragment>();
+				const bool bCanAvoid = OtherMoveTarget != nullptr;
+				const bool bOtherIsMoving = OtherMoveTarget ? OtherMoveTarget->GetCurrentAction() == EMassMovementAction::Move : true; // Assume moving if other does not have move target.
+				
+				// Check for colliders data
+				if (EnumHasAnyFlags(Obstacle.ObstacleItem.ItemFlags, EMassNavigationObstacleFlags::HasColliderData))
+				{
+					if (const FMassAvoidanceColliderFragment* ColliderFragment = OtherEntityView.GetFragmentDataPtr<FMassAvoidanceColliderFragment>())
+					{
+						if (ColliderFragment->Type == EMassColliderType::Circle)
+						{
+							const FMassCircleCollider Circle = ColliderFragment->GetCircleCollider();
+							
+							FCollider& Collider = Colliders.AddDefaulted_GetRef();
+							Collider.Velocity = OtherVelocity;
+							Collider.bCanAvoid = bCanAvoid;
+							Collider.bIsMoving = bOtherIsMoving;
+							Collider.Radius = Circle.Radius;
+							Collider.Location = Obstacle.LocationCached;
+						}
+						else if (ColliderFragment->Type == EMassColliderType::Pill)
+						{
+							const FMassPillCollider Pill = ColliderFragment->GetPillCollider(); 
+
+							FCollider& Collider = Colliders.AddDefaulted_GetRef();
+							Collider.Velocity = OtherVelocity;
+							Collider.bCanAvoid = bCanAvoid;
+							Collider.bIsMoving = bOtherIsMoving;
+							Collider.Radius = Pill.Radius;
+							Collider.Location = Obstacle.LocationCached + (Pill.HalfLength * Obstacle.Forward);
+
+							if (Colliders.Num() < MaxColliders)
+							{
+								FCollider& Collider2 = Colliders.AddDefaulted_GetRef();
+								Collider2.Velocity = OtherVelocity;
+								Collider2.bCanAvoid = bCanAvoid;
+								Collider2.bIsMoving = bOtherIsMoving;
+								Collider2.Radius = Pill.Radius;
+								Collider2.Location = Obstacle.LocationCached + (-Pill.HalfLength * Obstacle.Forward);
+							}
+						}
+					}
+				}
+				else
+				{
+					FCollider& Collider = Colliders.AddDefaulted_GetRef();
+					Collider.Location = Obstacle.LocationCached;
+					Collider.Velocity = OtherVelocity;
+					Collider.Radius = OtherEntityView.GetFragmentData<FDataFragment_AgentRadius>().Radius;
+					Collider.bCanAvoid = bCanAvoid;
+					Collider.bIsMoving = bOtherIsMoving;
+				}
+			}
+
+			// Process colliders for avoidance
+			for (const FCollider& Collider : Colliders)
+			{
+				bool bHasForcedNormal = false;
+				FVector ForcedNormal = FVector::ZeroVector;
+
+				if (Collider.bCanAvoid == false)
+				{
+					// If the other obstacle cannot avoid us, try to avoid the local minima they create between the wall and their collider.
+					// If the space between edge and collider is less than MinClearance, make the agent to avoid the gap.
+					const float MinClearance = 2.f * AgentRadius * MovingAvoidanceParams. StaticObstacleClearanceScale;
+					
+					// Find the maximum distance from edges that are too close.
+					float MaxDist = -1.f;
+					FVector ClosestPoint;
+					for (const FNavigationAvoidanceEdge& Edge : NavEdges.AvoidanceEdges)
+					{
+						const FVector Point = FMath::ClosestPointOnSegment(Collider.Location, Edge.Start, Edge.End);
+						const FVector Offset = Collider.Location - Point;
+						if (FVector::DotProduct(Offset, Edge.LeftDir) < 0.f)
+						{
+							// Behind the edge, ignore.
+							continue;
+						}
+
+						const float OffsetLength = Offset.Length();
+						const bool bTooNarrow = (OffsetLength - Collider.Radius) < MinClearance; 
+						if (bTooNarrow)
+						{
+							if (OffsetLength > MaxDist)
+							{
+								MaxDist = OffsetLength;
+								ClosestPoint = Point;
+							}
+						}
+					}
+
+					if (MaxDist != -1.f)
+					{
+						// Set up forced normal to avoid the gap between collider and edge.
+						ForcedNormal = (Collider.Location - ClosestPoint).GetSafeNormal();
+						bHasForcedNormal = true;
+					}
+				}
+
+				FVector RelPos = AgentLocation - Collider.Location;
+				RelPos.Z = 0.f; // we assume we work on a flat plane for now
+				const FVector RelVel = DesVel - Collider.Velocity;
+				const float ConDist = RelPos.Size();
+				const FVector ConNorm = ConDist > 0.f ? RelPos / ConDist : FVector::ForwardVector;
+
+				FVector SeparationNormal = ConNorm;
+				if (bHasForcedNormal)
+				{
+					// The more head on the collisions is, the more we should avoid towards the forced direction.
+					const FVector RelVelNorm = RelVel.GetSafeNormal();
+					const float Blend = FMath::Max(0.0f, -FVector::DotProduct(ConNorm, RelVelNorm));
+					SeparationNormal = FMath::Lerp(ConNorm, ForcedNormal, Blend).GetSafeNormal();
+				}
+				
+				const float StandingScaling = Collider.bIsMoving ? 1.0f : MovingAvoidanceParams.StandingObstacleAvoidanceScale; // Care less about standing agents so that we can push through standing crowd.
+				
+				// Separation force (stay away from agents if possible)
+				const float PenSep = (SeparationAgentRadius + Collider.Radius + MovingAvoidanceParams.ObstacleSeparationDistance) - ConDist;
+				const float SeparationMag = FMath::Square(FMath::Clamp(PenSep / MovingAvoidanceParams.ObstacleSeparationDistance, 0.f, 1.f));
+				const FVector SepForce = SeparationNormal * MovingAvoidanceParams.ObstacleSeparationStiffness;
+				const FVector SeparationForce = SepForce * SeparationMag * StandingScaling;
+
+				SteeringForce += SeparationForce;
+				TotalAgentSeparationForce += SeparationForce;
+
+				// Calculate closest point of approach based on relative agent positions and velocities.
+				const float CPA = UE::MassAvoidance::ComputeClosestPointOfApproach(RelPos, RelVel, PredictiveAvoidanceAgentRadius + Collider.Radius, MovingAvoidanceParams.PredictiveAvoidanceTime);
+
+				// Calculate penetration at CPA
+				const FVector AvoidRelPos = RelPos + RelVel * CPA;
+				const float AvoidDist = AvoidRelPos.Size();
+				const FVector AvoidConNormal = AvoidDist > 0.0f ? (AvoidRelPos / AvoidDist) : FVector::ForwardVector;
+
+				FVector AvoidNormal = AvoidConNormal;
+				if (bHasForcedNormal)
+				{
+					// The more head on the predicted collisions is, the more we should avoid towards the forced direction.
+					const FVector RelVelNorm = RelVel.GetSafeNormal();
+					const float Blend = FMath::Max(0.0f, -FVector::DotProduct(AvoidConNormal, RelVelNorm));
+					AvoidNormal = FMath::Lerp(AvoidConNormal, ForcedNormal, Blend).GetSafeNormal();
+				}
+				
+				const float AvoidPenetration = (PredictiveAvoidanceAgentRadius + Collider.Radius + MovingAvoidanceParams.PredictiveAvoidanceDistance) - AvoidDist; // Based on future agents distance
+				const float AvoidMag = FMath::Square(FMath::Clamp(AvoidPenetration / MovingAvoidanceParams.PredictiveAvoidanceDistance, 0.f, 1.f));
+				const float AvoidMagDist = (1.f - (CPA * InvPredictiveAvoidanceTime)); // No clamp, CPA is between 0 and PredictiveAvoidanceTime
+				const FVector AvoidForce = AvoidNormal * AvoidMag * AvoidMagDist * MovingAvoidanceParams.ObstaclePredictiveAvoidanceStiffness * StandingScaling;
+
+				SteeringForce += AvoidForce;
+
+#if WITH_MASSGAMEPLAY_DEBUG && UNSAFE_FOR_MT
+				// Display close agent
+				UE::MassAvoidance::DebugDrawCylinder(AgentDebugContext, Collider.Location, Collider.Location + UE::MassAvoidance::DebugLowCylinderOffset, Collider.Radius, UE::MassAvoidance::AgentsColor);
+
+				if (bHasForcedNormal)
+				{
+					UE::MassAvoidance::DebugDrawCylinder(BaseDebugContext, Collider.Location, Collider.Location + UE::MassAvoidance::DebugAgentHeightOffset, Collider.Radius, FColor::Red);
+				}
+
+				// Draw agent contact separation force
+				UE::MassAvoidance::DebugDrawSummedForce(AgentDebugContext,
+					Collider.Location + UE::MassAvoidance::DebugAgentHeightOffset,
+					Collider.Location + UE::MassAvoidance::DebugAgentHeightOffset + SeparationForce,
+					UE::MassAvoidance::AgentSeparationForceColor); 
+				
+				if (AvoidForce.Size() > 0.f)
+				{
+					// Draw agent vs agent hit positions
+					const FVector HitPosition = AgentLocation + (DesVel * CPA);
+					const FVector LeftOffset = PredictiveAvoidanceAgentRadius * UE::MassNavigation::GetLeftDirection(DesVel.GetSafeNormal(), FVector::UpVector);
+					UE::MassAvoidance::DebugDrawLine(AgentDebugContext, AgentLocation + UE::MassAvoidance::DebugAgentHeightOffset + LeftOffset, HitPosition + UE::MassAvoidance::DebugAgentHeightOffset + LeftOffset, UE::MassAvoidance::CurrentAgentColor, 1.5f);
+					UE::MassAvoidance::DebugDrawLine(AgentDebugContext, AgentLocation + UE::MassAvoidance::DebugAgentHeightOffset - LeftOffset, HitPosition + UE::MassAvoidance::DebugAgentHeightOffset - LeftOffset, UE::MassAvoidance::CurrentAgentColor, 1.5f);
+					UE::MassAvoidance::DebugDrawCylinder(AgentDebugContext, HitPosition, HitPosition + UE::MassAvoidance::DebugAgentHeightOffset, PredictiveAvoidanceAgentRadius, UE::MassAvoidance::CurrentAgentColor);
+
+					const FVector OtherHitPosition = Collider.Location + (Collider.Velocity * CPA);
+					const FVector OtherLeftOffset = Collider.Radius * UE::MassNavigation::GetLeftDirection(Collider.Velocity.GetSafeNormal(), FVector::UpVector);
+					const FVector Left = UE::MassAvoidance::DebugAgentHeightOffset + OtherLeftOffset;
+					const FVector Right = UE::MassAvoidance::DebugAgentHeightOffset - OtherLeftOffset;
+					UE::MassAvoidance::DebugDrawLine(AgentDebugContext, Collider.Location + Left, OtherHitPosition + Left, UE::MassAvoidance::AgentsColor, 1.5f);
+					UE::MassAvoidance::DebugDrawLine(AgentDebugContext, Collider.Location + Right, OtherHitPosition + Right, UE::MassAvoidance::AgentsColor, 1.5f);
+					UE::MassAvoidance::DebugDrawCylinder(AgentDebugContext, Collider.Location, Collider.Location + UE::MassAvoidance::DebugAgentHeightOffset, AgentRadius, UE::MassAvoidance::AgentsColor);
+					UE::MassAvoidance::DebugDrawCylinder(AgentDebugContext, OtherHitPosition, OtherHitPosition + UE::MassAvoidance::DebugAgentHeightOffset, AgentRadius, UE::MassAvoidance::AgentsColor);
+
+					// Draw agent avoid force
+					UE::MassAvoidance::DebugDrawForce(AgentDebugContext,
+						OtherHitPosition + UE::MassAvoidance::DebugAgentHeightOffset,
+						OtherHitPosition + UE::MassAvoidance::DebugAgentHeightOffset + AvoidForce,
+						UE::MassAvoidance::AgentAvoidForceColor);
+				}
+#endif // WITH_MASSGAMEPLAY_DEBUG
+			} // close entities loop
+
+			SteeringForce *= NearStartScaling * NearEndScaling;
+			
+			Force.Value = UE::MassNavigation::ClampVector(SteeringForce, MaxSteerAccel); // Assume unit mass
+
+#if WITH_MASSGAMEPLAY_DEBUG && UNSAFE_FOR_MT
+			const FVector AgentAvoidSteeringForce = SteeringForce - OldSteeringForce;
+
+			// Draw total steering force to separate agents
+			UE::MassAvoidance::DebugDrawSummedForce(AgentDebugContext,
+				AgentLocation + UE::MassAvoidance::DebugAgentHeightOffset,
+				AgentLocation + UE::MassAvoidance::DebugAgentHeightOffset + TotalAgentSeparationForce,
+				UE::MassAvoidance::AgentSeparationForceColor);
+
+			// Draw total steering force to avoid agents
+			UE::MassAvoidance::DebugDrawSummedForce(AgentDebugContext,
+				AgentLocation + UE::MassAvoidance::DebugAgentHeightOffset,
+				AgentLocation + UE::MassAvoidance::DebugAgentHeightOffset + AgentAvoidSteeringForce,
+				UE::MassAvoidance::AgentAvoidForceColor);
+
+			// Draw final steering force adding to the agent velocity
+			UE::MassAvoidance::DebugDrawArrow(BaseDebugContext, 
+				AgentLocation + AgentVelocity + UE::MassAvoidance::DebugAgentHeightOffset,
+				AgentLocation + AgentVelocity + UE::MassAvoidance::DebugAgentHeightOffset + Force.Value,
+				UE::MassAvoidance::FinalSteeringForceColor, UE::MassAvoidance::BigArrowHeadSize, UE::MassAvoidance::BigArrowThickness);
+#endif // WITH_MASSGAMEPLAY_DEBUG
+		}
+	});
+}
+
+//----------------------------------------------------------------------//
+//  UMassStandingAvoidanceProcessor
+//----------------------------------------------------------------------//
+UMassStandingAvoidanceProcessor::UMassStandingAvoidanceProcessor()
+{
+	bAutoRegisterWithProcessingPhases = true;
+	ExecutionFlags = (int32)EProcessorExecutionFlags::All;
+	ExecutionOrder.ExecuteInGroup = UE::Mass::ProcessorGroupNames::Avoidance;
+	ExecutionOrder.ExecuteAfter.Add(TEXT("MassMovingAvoidanceProcessor"));
+}
+
+void UMassStandingAvoidanceProcessor::ConfigureQueries()
+{
+	EntityQuery.AddRequirement<FMassGhostLocationFragment>(EMassFragmentAccess::ReadWrite);
+	EntityQuery.AddRequirement<FMassNavigationEdgesFragment>(EMassFragmentAccess::ReadOnly);
+	EntityQuery.AddRequirement<FMassMoveTargetFragment>(EMassFragmentAccess::ReadOnly);
+	EntityQuery.AddRequirement<FDataFragment_Transform>(EMassFragmentAccess::ReadOnly);
+	EntityQuery.AddRequirement<FDataFragment_AgentRadius>(EMassFragmentAccess::ReadOnly);
+	EntityQuery.AddTagRequirement<FMassMediumLODTag>(EMassFragmentPresence::None);
+	EntityQuery.AddTagRequirement<FMassLowLODTag>(EMassFragmentPresence::None);
+	EntityQuery.AddTagRequirement<FMassOffLODTag>(EMassFragmentPresence::None);
+	EntityQuery.AddConstSharedRequirement<FMassStandingAvoidanceParameters>(EMassFragmentPresence::All);
+}
+
+void UMassStandingAvoidanceProcessor::Initialize(UObject& Owner)
+{
+	Super::Initialize(Owner);
+
+	World = Owner.GetWorld();
+	NavigationSubsystem = UWorld::GetSubsystem<UMassNavigationSubsystem>(Owner.GetWorld());
+}
+
+void UMassStandingAvoidanceProcessor::Execute(UMassEntitySubsystem& EntitySubsystem, FMassExecutionContext& Context)
+{
+	QUICK_SCOPE_CYCLE_COUNTER(UMassStandingAvoidanceProcessor);
+
+	if (!World || !NavigationSubsystem)
+	{
+		return;
+	}
+
+	// Avoidance while standing
+	EntityQuery.ForEachEntityChunk(EntitySubsystem, Context, [this, &EntitySubsystem](FMassExecutionContext& Context)
+	{
+		const int32 NumEntities = Context.GetNumEntities();
+		const float DeltaTime = Context.GetDeltaTimeSeconds();
+
+		const TArrayView<FMassGhostLocationFragment> GhostList = Context.GetMutableFragmentView<FMassGhostLocationFragment>();
+		const TConstArrayView<FDataFragment_Transform> LocationList = Context.GetFragmentView<FDataFragment_Transform>();
+		const TConstArrayView<FDataFragment_AgentRadius> RadiusList = Context.GetFragmentView<FDataFragment_AgentRadius>();
+		const TConstArrayView<FMassMoveTargetFragment> MoveTargetList = Context.GetFragmentView<FMassMoveTargetFragment>();
+		const FMassStandingAvoidanceParameters& StandingParams = Context.GetConstSharedFragment<FMassStandingAvoidanceParameters>();
+
+		const float GhostSeparationDistance = StandingParams.GhostSeparationDistance;
+		const float GhostSeparationStiffness = StandingParams.GhostSeparationStiffness;
+
+		const float MovingSeparationDistance = StandingParams.GhostSeparationDistance * StandingParams.MovingObstacleAvoidanceScale;
+		const float MovingSeparationStiffness = StandingParams.GhostSeparationStiffness * StandingParams.MovingObstacleAvoidanceScale;
+
+		// Arrays used to store close agents
+		TArray<FMassNavigationObstacleItem, TFixedAllocator<UE::MassAvoidance::MaxObstacleResults>> CloseEntities;
+
+		struct FSortedObstacle
+		{
+			FSortedObstacle() = default;
+			FSortedObstacle(const FMassEntityHandle InEntity, const FVector InLocation, const FVector InForward, const float InDistSq) : Entity(InEntity), Location(InLocation), Forward(InForward), DistSq(InDistSq) {}
+			
+			FMassEntityHandle Entity;
+			FVector Location = FVector::ZeroVector;
+			FVector Forward = FVector::ForwardVector;
+			float DistSq = 0.0f;
+		};
+		TArray<FSortedObstacle, TFixedAllocator<UE::MassAvoidance::MaxObstacleResults>> ClosestObstacles;
+
+		for (int32 EntityIndex = 0; EntityIndex < NumEntities; ++EntityIndex)
+		{
+			// @todo: this check should eventually be part of the query.
+			const FMassMoveTargetFragment& MoveTarget = MoveTargetList[EntityIndex];
+			if (MoveTarget.GetCurrentAction() != EMassMovementAction::Stand)
+			{
+				continue;
+			}
+			
+			FMassGhostLocationFragment& Ghost = GhostList[EntityIndex];
+			// Skip if the ghost is not valid for this movement action yet.
+			if (Ghost.IsValid(MoveTarget.GetCurrentActionID()) == false)
+			{
+				continue;
+			}
+
+			const FDataFragment_Transform& Location = LocationList[EntityIndex];
+			const FDataFragment_AgentRadius& Radius = RadiusList[EntityIndex];
+
+			FMassEntityHandle Entity = Context.GetEntity(EntityIndex);
+			const FVector AgentLocation = Location.GetTransform().GetTranslation();
+			const float AgentRadius = Radius.Radius;
+
+			// Steer ghost to move target.
+			const float SteerK = 1.f / StandingParams.GhostSteeringReactionTime;
+			constexpr float SteeringMinDistance = 1.0f; // Do not bother to steer if the distance is less than this.
+
+			FVector SteerDirection = FVector::ZeroVector;
+			FVector Delta = MoveTarget.Center - Ghost.Location;
+			Delta.Z = 0.0f;
+			const float Distance = Delta.Size();
+			float SpeedFade = 0.0;
+			if (Distance > SteeringMinDistance)
+			{
+				SteerDirection = Delta / Distance;
+				SpeedFade = FMath::Clamp(Distance / FMath::Max(KINDA_SMALL_NUMBER, StandingParams.GhostStandSlowdownRadius), 0.0f, 1.0f);
+			}
+
+			const FVector GhostDesiredVelocity = SteerDirection * StandingParams.GhostMaxSpeed * SpeedFade;
+			FVector GhostSteeringForce = SteerK * (GhostDesiredVelocity - Ghost.Velocity); // Goal force
+			
+			// Find close obstacles
+			// @todo: optimize FindCloseObstacles() and cache results. We're intentionally using agent location here, to allow to share the results with moving avoidance.
+			const FNavigationObstacleHashGrid2D& ObstacleGrid = NavigationSubsystem->GetObstacleGridMutable();
+			UE::MassAvoidance::FindCloseObstacles(AgentLocation, StandingParams.GhostObstacleDetectionDistance, ObstacleGrid, CloseEntities, UE::MassAvoidance::MaxObstacleResults);
+
+			// Remove unwanted and find the closest in the CloseEntities
+			const float DistanceCutOffSqr = FMath::Square(StandingParams.GhostObstacleDetectionDistance);
+			ClosestObstacles.Reset();
+			for (const FNavigationObstacleHashGrid2D::ItemIDType OtherEntity : CloseEntities)
+			{
+				// Skip self
+				if (OtherEntity.Entity == Entity)
+				{
+					continue;
+				}
+
+				// Skip invalid entities.
+				if (!EntitySubsystem.IsEntityValid(OtherEntity.Entity))
+				{
+					UE_LOG(LogAvoidanceObstacles, VeryVerbose, TEXT("Close entity is invalid, skipped."));
+					continue;
+				}
+
+				// Skip too far
+				const FDataFragment_Transform& OtherTransform = EntitySubsystem.GetFragmentDataChecked<FDataFragment_Transform>(OtherEntity.Entity);
+				const FVector OtherLocation = OtherTransform.GetTransform().GetLocation();
+				const float DistSq = FVector::DistSquared(AgentLocation, OtherLocation);
+				if (DistSq > DistanceCutOffSqr)
+				{
+					continue;
+				}
+
+				ClosestObstacles.Emplace(OtherEntity.Entity, OtherLocation, OtherTransform.GetTransform().GetRotation().GetForwardVector(), DistSq);
+			}
+			ClosestObstacles.Sort([](const FSortedObstacle& A, const FSortedObstacle& B) { return A.DistSq < B.DistSq; });
+
+			const float GhostRadius = AgentRadius * StandingParams.GhostSeparationRadiusScale;
+			
+			// Compute forces
+			constexpr int32 MaxCloseObstacleTreated = 6;
+			const int32 NumCloseObstacles = FMath::Min(ClosestObstacles.Num(), MaxCloseObstacleTreated);
+			for (int32 Index = 0; Index < NumCloseObstacles; Index++)
+			{
+				FSortedObstacle& OtherAgent = ClosestObstacles[Index];
+				FMassEntityView OtherEntityView(EntitySubsystem, OtherAgent.Entity);
+
+				const float OtherRadius = OtherEntityView.GetFragmentData<FDataFragment_AgentRadius>().Radius;
+				const float TotalRadius = GhostRadius + OtherRadius;
+
+				// @todo: this is heavy fragment to access, see if we could handle this differently.
+				const FMassMoveTargetFragment* OtherMoveTarget = OtherEntityView.GetFragmentDataPtr<FMassMoveTargetFragment>();
+				const FMassGhostLocationFragment* OtherGhost = OtherEntityView.GetFragmentDataPtr<FMassGhostLocationFragment>();
+
+				const bool bOtherHasGhost = OtherMoveTarget != nullptr && OtherGhost != nullptr
+											&& OtherMoveTarget->GetCurrentAction() == EMassMovementAction::Stand
+											&& OtherGhost->IsValid(OtherMoveTarget->GetCurrentActionID());
+
+				// If other has ghost active, avoid that, else avoid the actual agent.
+				if (bOtherHasGhost)
+				{
+					// Avoid the other agent more, when it is further away from it's goal location.
+					const float OtherDistanceToGoal = FVector::Distance(OtherGhost->Location, OtherMoveTarget->Center);
+					const float OtherSteerFade = FMath::Clamp(OtherDistanceToGoal / StandingParams.GhostToTargetMaxDeviation, 0.0f, 1.0f);
+					const float SeparationStiffness = FMath::Lerp(GhostSeparationStiffness, MovingSeparationStiffness, OtherSteerFade);
+
+					// Ghost separation
+					FVector RelPos = Ghost.Location - OtherGhost->Location;
+					RelPos.Z = 0.f; // we assume we work on a flat plane for now
+					const float ConDist = RelPos.Size();
+					const FVector ConNorm = ConDist > 0.f ? RelPos / ConDist : FVector::ForwardVector;
+
+					// Separation force (stay away from obstacles if possible)
+					const float PenSep = (TotalRadius + GhostSeparationDistance) - ConDist;
+					const float SeparationMag = UE::MassNavigation::Smooth(FMath::Clamp(PenSep / GhostSeparationDistance, 0.f, 1.f));
+					const FVector SeparationForce = ConNorm * SeparationStiffness * SeparationMag;
+
+					GhostSteeringForce += SeparationForce;
+				}
+				else
+				{
+					// Avoid more when the avoidance other is in front,
+					const FVector DirToOther = (OtherAgent.Location - Ghost.Location).GetSafeNormal();
+					const float DirectionalFade = FMath::Square(FMath::Max(0.0f, FVector::DotProduct(MoveTarget.Forward, DirToOther)));
+					const float DirectionScale = FMath::Lerp(StandingParams.MovingObstacleDirectionalScale, 1.0f, DirectionalFade);
+
+					// Treat the other agent as a 2D capsule protruding towards forward.
+ 					const FVector OtherBasePosition = OtherAgent.Location;
+					const FVector OtherPersonalSpacePosition = OtherAgent.Location + OtherAgent.Forward * OtherRadius * StandingParams.MovingObstaclePersonalSpaceScale * DirectionScale;
+					const FVector OtherLocation = FMath::ClosestPointOnSegment(Ghost.Location, OtherBasePosition, OtherPersonalSpacePosition);
+
+					FVector RelPos = Ghost.Location - OtherLocation;
+					RelPos.Z = 0.f;
+					const float ConDist = RelPos.Size();
+					const FVector ConNorm = ConDist > 0.f ? RelPos / ConDist : FVector::ForwardVector;
+
+					// Separation force (stay away from obstacles if possible)
+					const float PenSep = (TotalRadius + MovingSeparationDistance) - ConDist;
+					const float SeparationMag = UE::MassNavigation::Smooth(FMath::Clamp(PenSep / MovingSeparationDistance, 0.f, 1.f));
+					const FVector SeparationForce = ConNorm * MovingSeparationStiffness * SeparationMag;
+
+					GhostSteeringForce += SeparationForce;
+				}
+			}
+
+			GhostSteeringForce.Z = 0.0f;
+			GhostSteeringForce = UE::MassNavigation::ClampVector(GhostSteeringForce, StandingParams.GhostMaxAcceleration); // Assume unit mass
+			Ghost.Velocity += GhostSteeringForce * DeltaTime;
+			Ghost.Velocity.Z = 0.0f;
+			
+			// Damping
+			FMath::ExponentialSmoothingApprox(Ghost.Velocity, FVector::ZeroVector, DeltaTime, StandingParams.GhostVelocityDampingTime);
+			
+			Ghost.Location += Ghost.Velocity * DeltaTime;
+
+			// Dont let the ghost location too far from move target center.
+			const FVector DirToCenter = Ghost.Location - MoveTarget.Center;
+			const float DistToCenter = DirToCenter.Length();
+			if (DistToCenter > StandingParams.GhostToTargetMaxDeviation)
+			{
+				Ghost.Location = MoveTarget.Center + DirToCenter * (StandingParams.GhostToTargetMaxDeviation / DistToCenter);
+			}
+
+		}
+	});
+	
+}
+
+#undef UNSAFE_FOR_MT
diff --git a/Engine/Plugins/AI/MassAI/Source/MassNavigation/Private/MassAvoidanceTrait.cpp b/Engine/Plugins/AI/MassAI/Source/MassNavigation/Private/MassAvoidanceTrait.cpp
new file mode 100644
--- /dev/null
+++ b/Engine/Plugins/AI/MassAI/Source/MassNavigation/Private/MassAvoidanceTrait.cpp
@@ -0,0 +1,32 @@
+﻿// Copyright Epic Games, Inc. All Rights Reserved.
+
+#include "Avoidance/MassAvoidanceTrait.h"
+#include "Avoidance/MassAvoidanceFragments.h"
+#include "MassEntityTemplateRegistry.h"
+#include "MassMovementFragments.h"
+#include "MassCommonFragments.h"
+#include "MassNavigationFragments.h"
+#include "Engine/World.h"
+
+void UMassObstacleAvoidanceTrait::BuildTemplate(FMassEntityTemplateBuildContext& BuildContext, UWorld& World) const
+{
+	UMassEntitySubsystem* EntitySubsystem = UWorld::GetSubsystem<UMassEntitySubsystem>(&World);
+	check(EntitySubsystem);
+
+	BuildContext.AddFragment<FDataFragment_AgentRadius>();
+	BuildContext.AddFragment<FMassNavigationEdgesFragment>();
+	BuildContext.AddFragment<FDataFragment_Transform>();
+	BuildContext.AddFragment<FMassVelocityFragment>();
+	BuildContext.AddFragment<FMassForceFragment>();
+	BuildContext.AddFragment<FMassMoveTargetFragment>();
+
+	const FMassMovingAvoidanceParameters MovingValidated = MovingParameters.GetValidated();
+	const uint32 MovingHash = UE::StructUtils::GetStructCrc32(FConstStructView::Make(MovingValidated));
+	const FConstSharedStruct MovingFragment = EntitySubsystem->GetOrCreateConstSharedFragment(MovingHash, MovingValidated);
+	BuildContext.AddConstSharedFragment(MovingFragment);
+
+	const FMassStandingAvoidanceParameters StandingValidated = StandingParameters.GetValidated();
+	const uint32 StandingHash = UE::StructUtils::GetStructCrc32(FConstStructView::Make(StandingValidated));
+	const FConstSharedStruct StandingFragment = EntitySubsystem->GetOrCreateConstSharedFragment(StandingHash, StandingValidated);
+	BuildContext.AddConstSharedFragment(StandingFragment);
+}
diff --git a/Engine/Plugins/AI/MassAI/Source/MassNavigation/Private/MassNavigationObstacleTrait.cpp b/Engine/Plugins/AI/MassAI/Source/MassNavigation/Private/MassNavigationObstacleTrait.cpp
new file mode 100644
--- /dev/null
+++ b/Engine/Plugins/AI/MassAI/Source/MassNavigation/Private/MassNavigationObstacleTrait.cpp
@@ -0,0 +1,12 @@
+﻿// Copyright Epic Games, Inc. All Rights Reserved.
+#include "Avoidance/MassNavigationObstacleTrait.h"
+#include "MassEntityTemplateRegistry.h"
+#include "MassNavigationFragments.h"
+#include "MassCommonFragments.h"
+
+void UMassNavigationObstacleTrait::BuildTemplate(FMassEntityTemplateBuildContext& BuildContext, UWorld& World) const
+{
+	BuildContext.AddFragment<FDataFragment_AgentRadius>();
+
+	BuildContext.AddFragment<FMassNavigationObstacleGridCellLocationFragment>();
+}
