Mass Movement and Navigation cleanup

diff --git a/Engine/Plugins/AI/MassAI/Source/MassNavigation/Private/MassNavigationProcessors.cpp b/Engine/Plugins/AI/MassAI/Source/MassNavigation/Private/MassNavigationProcessors.cpp
new file mode 100644
--- /dev/null
+++ b/Engine/Plugins/AI/MassAI/Source/MassNavigation/Private/MassNavigationProcessors.cpp
@@ -0,0 +1,322 @@
+// Copyright Epic Games, Inc. All Rights Reserved.
+
+#include "MassNavigationProcessors.h"
+#include "MassCommonUtils.h"
+#include "MassCommandBuffer.h"
+#include "MassCommonFragments.h"
+#include "MassNavigationFragments.h"
+#include "MassNavigationSubsystem.h"
+#include "MassSimulationLOD.h"
+#include "MassMovementTypes.h"
+#include "MassMovementFragments.h"
+#include "MassEntityView.h"
+#include "Engine/World.h"
+
+
+#define UNSAFE_FOR_MT 0
+#define MOVEMENT_DEBUGDRAW 0	// Set to 1 to see heading debugdraw
+
+//----------------------------------------------------------------------//
+//  UMassOffLODNavigationProcessor
+//----------------------------------------------------------------------//
+
+UMassOffLODNavigationProcessor::UMassOffLODNavigationProcessor()
+{
+	ExecutionFlags = (int32)EProcessorExecutionFlags::All;
+	ExecutionOrder.ExecuteInGroup = UE::Mass::ProcessorGroupNames::Movement;
+	ExecutionOrder.ExecuteAfter.Add(UE::Mass::ProcessorGroupNames::Avoidance); // @todo: remove this direct dependency
+}
+
+void UMassOffLODNavigationProcessor::ConfigureQueries()
+{
+	EntityQuery_Conditional.AddRequirement<FDataFragment_Transform>(EMassFragmentAccess::ReadWrite);
+	EntityQuery_Conditional.AddRequirement<FMassMoveTargetFragment>(EMassFragmentAccess::ReadOnly);
+	EntityQuery_Conditional.AddTagRequirement<FMassOffLODTag>(EMassFragmentPresence::All);
+	EntityQuery_Conditional.AddChunkRequirement<FMassSimulationVariableTickChunkFragment>(EMassFragmentAccess::ReadOnly, EMassFragmentPresence::Optional);
+	EntityQuery_Conditional.SetChunkFilter(&FMassSimulationVariableTickChunkFragment::ShouldTickChunkThisFrame);
+}
+
+void UMassOffLODNavigationProcessor::Execute(UMassEntitySubsystem& EntitySubsystem,
+													FMassExecutionContext& Context)
+{
+	EntityQuery_Conditional.ForEachEntityChunk(EntitySubsystem, Context, [this](FMassExecutionContext& Context)
+	{
+#if WITH_MASSGAMEPLAY_DEBUG
+		if (UE::MassMovement::bFreezeMovement)
+		{
+			return;
+		}
+#endif // WITH_MASSGAMEPLAY_DEBUG
+		const int32 NumEntities = Context.GetNumEntities();
+
+		const TArrayView<FDataFragment_Transform> LocationList = Context.GetMutableFragmentView<FDataFragment_Transform>();
+		const TConstArrayView<FMassMoveTargetFragment> MoveTargetList = Context.GetFragmentView<FMassMoveTargetFragment>();
+
+		for (int32 EntityIndex = 0; EntityIndex < NumEntities; ++EntityIndex)
+		{
+			FTransform& CurrentTransform = LocationList[EntityIndex].GetMutableTransform();
+			const FMassMoveTargetFragment& MoveTarget = MoveTargetList[EntityIndex];
+
+			// Snap position to move target directly
+			CurrentTransform.SetLocation(MoveTarget.Center);
+		}
+	});
+}
+
+
+//----------------------------------------------------------------------//
+//  UMassNavigationSmoothHeightProcessor
+//----------------------------------------------------------------------//
+
+UMassNavigationSmoothHeightProcessor::UMassNavigationSmoothHeightProcessor()
+{
+	ExecutionFlags = (int32)EProcessorExecutionFlags::All;
+	ExecutionOrder.ExecuteAfter.Add(UE::Mass::ProcessorGroupNames::Movement);
+}
+
+void UMassNavigationSmoothHeightProcessor::ConfigureQueries()
+{
+	EntityQuery.AddRequirement<FDataFragment_Transform>(EMassFragmentAccess::ReadWrite);
+	EntityQuery.AddRequirement<FMassMoveTargetFragment>(EMassFragmentAccess::ReadOnly);
+	EntityQuery.AddTagRequirement<FMassOffLODTag>(EMassFragmentPresence::None);
+	EntityQuery.AddConstSharedRequirement<FMassMovementParameters>(EMassFragmentPresence::All);
+}
+
+void UMassNavigationSmoothHeightProcessor::Execute(UMassEntitySubsystem& EntitySubsystem,
+													FMassExecutionContext& Context)
+{
+	EntityQuery.ForEachEntityChunk(EntitySubsystem, Context, [this](FMassExecutionContext& Context)
+	{
+#if WITH_MASSGAMEPLAY_DEBUG
+		if (UE::MassMovement::bFreezeMovement)
+		{
+			return;
+		}
+#endif // WITH_MASSGAMEPLAY_DEBUG
+		const int32 NumEntities = Context.GetNumEntities();
+		const float DeltaTime = Context.GetDeltaTimeSeconds();
+
+		const FMassMovementParameters& MovementParams = Context.GetConstSharedFragment<FMassMovementParameters>();
+		const TArrayView<FDataFragment_Transform> LocationList = Context.GetMutableFragmentView<FDataFragment_Transform>();
+		const TConstArrayView<FMassMoveTargetFragment> MoveTargetList = Context.GetFragmentView<FMassMoveTargetFragment>();
+
+		for (int32 EntityIndex = 0; EntityIndex < NumEntities; ++EntityIndex)
+		{
+			FTransform& CurrentTransform = LocationList[EntityIndex].GetMutableTransform();
+			const FMassMoveTargetFragment& MoveTarget = MoveTargetList[EntityIndex];
+
+			if (MoveTarget.GetCurrentAction() == EMassMovementAction::Move || MoveTarget.GetCurrentAction() == EMassMovementAction::Stand)
+			{
+				// Set height smoothly to follow current move targets height.
+				FVector CurrentLocation = CurrentTransform.GetLocation();
+				FMath::ExponentialSmoothingApprox(CurrentLocation.Z, MoveTarget.Center.Z, DeltaTime, MovementParams.HeightSmoothingTime);
+				CurrentTransform.SetLocation(CurrentLocation);
+			}
+		}
+	});
+}
+
+//----------------------------------------------------------------------//
+//  UMassMoveTargetFragmentInitializer
+//----------------------------------------------------------------------//
+
+UMassMoveTargetFragmentInitializer::UMassMoveTargetFragmentInitializer()
+{
+	FragmentType = FMassMoveTargetFragment::StaticStruct();
+}
+
+void UMassMoveTargetFragmentInitializer::ConfigureQueries()
+{
+	InitializerQuery.AddRequirement<FMassMoveTargetFragment>(EMassFragmentAccess::ReadWrite);
+	InitializerQuery.AddRequirement<FDataFragment_Transform>(EMassFragmentAccess::ReadOnly);
+}
+
+void UMassMoveTargetFragmentInitializer::Execute(UMassEntitySubsystem& EntitySubsystem, FMassExecutionContext& Context)
+{
+	InitializerQuery.ForEachEntityChunk(EntitySubsystem, Context, [](FMassExecutionContext& Context)
+	{
+		const int32 NumEntities = Context.GetNumEntities();
+		const TArrayView<FMassMoveTargetFragment> MoveTargetList = Context.GetMutableFragmentView<FMassMoveTargetFragment>();
+		const TConstArrayView<FDataFragment_Transform> LocationList = Context.GetFragmentView<FDataFragment_Transform>();
+
+		for (int32 EntityIndex = 0; EntityIndex < NumEntities; ++EntityIndex)
+		{
+			FMassMoveTargetFragment& MoveTarget = MoveTargetList[EntityIndex];
+			const FDataFragment_Transform& Location = LocationList[EntityIndex];
+
+			MoveTarget.Center = Location.GetTransform().GetLocation();
+			MoveTarget.Forward = Location.GetTransform().GetRotation().Vector();
+			MoveTarget.DistanceToGoal = 0.0f;
+			MoveTarget.SlackRadius = 0.0f;
+		}
+	});
+}
+
+
+//----------------------------------------------------------------------//
+//  UMassNavigationObstacleGridProcessor
+//----------------------------------------------------------------------//
+UMassNavigationObstacleGridProcessor::UMassNavigationObstacleGridProcessor()
+{
+	ExecutionFlags = (int32)EProcessorExecutionFlags::All;
+	ExecutionOrder.ExecuteAfter.Add(UE::Mass::ProcessorGroupNames::Movement);
+}
+
+void UMassNavigationObstacleGridProcessor::ConfigureQueries()
+{
+	AddToGridEntityQuery.AddRequirement<FDataFragment_Transform>(EMassFragmentAccess::ReadOnly);
+	AddToGridEntityQuery.AddRequirement<FDataFragment_AgentRadius>(EMassFragmentAccess::ReadOnly);
+	AddToGridEntityQuery.AddRequirement<FMassNavigationObstacleGridCellLocationFragment>(EMassFragmentAccess::ReadWrite);
+	UpdateGridEntityQuery = AddToGridEntityQuery;
+	RemoveFromGridEntityQuery = AddToGridEntityQuery;
+
+	AddToGridEntityQuery.AddTagRequirement<FMassOffLODTag>(EMassFragmentPresence::None);
+	AddToGridEntityQuery.AddTagRequirement<FMassInNavigationObstacleGridTag>(EMassFragmentPresence::None);
+
+	UpdateGridEntityQuery.AddTagRequirement<FMassOffLODTag>(EMassFragmentPresence::None);
+	UpdateGridEntityQuery.AddTagRequirement<FMassInNavigationObstacleGridTag>(EMassFragmentPresence::All);
+
+	RemoveFromGridEntityQuery.AddTagRequirement<FMassOffLODTag>(EMassFragmentPresence::All);
+	RemoveFromGridEntityQuery.AddTagRequirement<FMassInNavigationObstacleGridTag>(EMassFragmentPresence::All);
+}
+
+void UMassNavigationObstacleGridProcessor::Initialize(UObject& Owner)
+{
+	Super::Initialize(Owner);
+
+	NavigationSubsystem = UWorld::GetSubsystem<UMassNavigationSubsystem>(Owner.GetWorld());
+}
+
+void UMassNavigationObstacleGridProcessor::Execute(UMassEntitySubsystem& EntitySubsystem, FMassExecutionContext& Context)
+{
+	if (!NavigationSubsystem)
+	{
+		return;
+	}
+
+	// can't be ParallelFor due to MovementSubsystem->GetGridMutable().Move not being thread-safe
+	AddToGridEntityQuery.ForEachEntityChunk(EntitySubsystem, Context, [this, &EntitySubsystem](FMassExecutionContext& Context)
+	{
+		const int32 NumEntities = Context.GetNumEntities();
+
+		TConstArrayView<FDataFragment_Transform> LocationList = Context.GetFragmentView<FDataFragment_Transform>();
+		TConstArrayView<FDataFragment_AgentRadius> RadiiList = Context.GetFragmentView<FDataFragment_AgentRadius>();
+		TArrayView<FMassNavigationObstacleGridCellLocationFragment> NavigationObstacleCellLocationList = Context.GetMutableFragmentView<FMassNavigationObstacleGridCellLocationFragment>();
+
+		for (int32 EntityIndex = 0; EntityIndex < NumEntities; ++EntityIndex)
+		{
+			// Add to the grid
+			const FVector NewPos = LocationList[EntityIndex].GetTransform().GetLocation();
+			const float Radius = RadiiList[EntityIndex].Radius;
+
+			FMassNavigationObstacleItem ObstacleItem;
+			ObstacleItem.Entity = Context.GetEntity(EntityIndex);
+			FMassEntityView EntityView(EntitySubsystem, ObstacleItem.Entity);
+			const FMassAvoidanceColliderFragment* Collider = EntityView.GetFragmentDataPtr<FMassAvoidanceColliderFragment>();
+			if (Collider)
+			{
+				ObstacleItem.ItemFlags |= EMassNavigationObstacleFlags::HasColliderData;
+			}
+			
+			const FBox NewBounds(NewPos - FVector(Radius, Radius, 0.f), NewPos + FVector(Radius, Radius, 0.f));
+			NavigationObstacleCellLocationList[EntityIndex].CellLoc = NavigationSubsystem->GetObstacleGridMutable().Add(ObstacleItem, NewBounds);
+
+			Context.Defer().AddTag<FMassInNavigationObstacleGridTag>(ObstacleItem.Entity);
+		}
+	});
+
+	UpdateGridEntityQuery.ForEachEntityChunk(EntitySubsystem, Context, [this, &EntitySubsystem](FMassExecutionContext& Context)
+	{
+		const int32 NumEntities = Context.GetNumEntities();
+
+		TConstArrayView<FDataFragment_Transform> LocationList = Context.GetFragmentView<FDataFragment_Transform>();
+		TConstArrayView<FDataFragment_AgentRadius> RadiiList = Context.GetFragmentView<FDataFragment_AgentRadius>();
+		TArrayView<FMassNavigationObstacleGridCellLocationFragment> NavigationObstacleCellLocationList = Context.GetMutableFragmentView<FMassNavigationObstacleGridCellLocationFragment>();
+
+		for (int32 EntityIndex = 0; EntityIndex < NumEntities; ++EntityIndex)
+		{
+			// Update position in grid
+			const FVector NewPos = LocationList[EntityIndex].GetTransform().GetLocation();
+			const float Radius = RadiiList[EntityIndex].Radius;
+			FMassNavigationObstacleItem ObstacleItem;
+			ObstacleItem.Entity = Context.GetEntity(EntityIndex);
+			FMassEntityView EntityView(EntitySubsystem, ObstacleItem.Entity);
+			const FMassAvoidanceColliderFragment* Collider = EntityView.GetFragmentDataPtr<FMassAvoidanceColliderFragment>();
+			if (Collider)
+			{
+				ObstacleItem.ItemFlags |= EMassNavigationObstacleFlags::HasColliderData;
+			}
+			const FBox NewBounds(NewPos - FVector(Radius, Radius, 0.f), NewPos + FVector(Radius, Radius, 0.f));
+			NavigationObstacleCellLocationList[EntityIndex].CellLoc = NavigationSubsystem->GetObstacleGridMutable().Move(ObstacleItem, NavigationObstacleCellLocationList[EntityIndex].CellLoc, NewBounds);
+
+#if WITH_MASSGAMEPLAY_DEBUG && 0
+			const FDebugContext BaseDebugContext(this, LogAvoidance, nullptr, ObstacleItem.Entity);
+			if (DebugIsSelected(ObstacleItem.Entity))
+			{
+				FBox Box = MovementSubsystem->GetGridMutable().CalcCellBounds(AvoidanceObstacleCellLocationList[EntityIndex].CellLoc);
+				Box.Max.Z += 200.f;					
+				DebugDrawBox(BaseDebugContext, Box, FColor::Yellow);
+			}
+#endif // WITH_MASSGAMEPLAY_DEBUG
+		}
+	});
+
+	RemoveFromGridEntityQuery.ForEachEntityChunk(EntitySubsystem, Context, [this](FMassExecutionContext& Context)
+	{
+		const int32 NumEntities = Context.GetNumEntities();
+
+		TArrayView<FMassNavigationObstacleGridCellLocationFragment> AvoidanceObstacleCellLocationList = Context.GetMutableFragmentView<FMassNavigationObstacleGridCellLocationFragment>();
+		for (int32 EntityIndex = 0; EntityIndex < NumEntities; ++EntityIndex)
+		{
+			FMassNavigationObstacleItem ObstacleItem;
+			ObstacleItem.Entity = Context.GetEntity(EntityIndex);
+			NavigationSubsystem->GetObstacleGridMutable().Remove(ObstacleItem, AvoidanceObstacleCellLocationList[EntityIndex].CellLoc);
+			AvoidanceObstacleCellLocationList[EntityIndex].CellLoc = FNavigationObstacleHashGrid2D::FCellLocation();
+
+			Context.Defer().RemoveTag<FMassInNavigationObstacleGridTag>(ObstacleItem.Entity);
+		}
+	});
+}
+
+//----------------------------------------------------------------------//
+//  UMassNavigationObstacleRemoverProcessor
+//----------------------------------------------------------------------//
+UMassNavigationObstacleRemoverProcessor::UMassNavigationObstacleRemoverProcessor()
+{
+	FragmentType = FMassNavigationObstacleGridCellLocationFragment::StaticStruct();
+	ExecutionFlags = (int32)(EProcessorExecutionFlags::All);
+}
+
+void UMassNavigationObstacleRemoverProcessor::ConfigureQueries()
+{
+	EntityQuery.AddRequirement<FMassNavigationObstacleGridCellLocationFragment>(EMassFragmentAccess::ReadWrite);
+}
+
+void UMassNavigationObstacleRemoverProcessor::Initialize(UObject& Owner)
+{
+	Super::Initialize(Owner);
+	NavigationSubsystem = UWorld::GetSubsystem<UMassNavigationSubsystem>(Owner.GetWorld());
+}
+
+void UMassNavigationObstacleRemoverProcessor::Execute(UMassEntitySubsystem& EntitySubsystem, FMassExecutionContext& Context)
+{
+	if (!NavigationSubsystem)
+	{
+		return;
+	}
+
+	EntityQuery.ForEachEntityChunk(EntitySubsystem, Context, [this](FMassExecutionContext& Context)
+	{
+		const int32 NumEntities = Context.GetNumEntities();
+		const TArrayView<FMassNavigationObstacleGridCellLocationFragment> AvoidanceObstacleCellLocationList = Context.GetMutableFragmentView<FMassNavigationObstacleGridCellLocationFragment>();
+
+		for (int32 i = 0; i < NumEntities; ++i)
+		{
+			FMassNavigationObstacleItem ObstacleItem;
+			ObstacleItem.Entity = Context.GetEntity(i);
+			NavigationSubsystem->GetObstacleGridMutable().Remove(ObstacleItem, AvoidanceObstacleCellLocationList[i].CellLoc);
+		}
+	});
+}
+
+#undef UNSAFE_FOR_MT
diff --git a/Engine/Plugins/AI/MassAI/Source/MassNavigation/Private/MassSmoothOrientationProcessors.cpp b/Engine/Plugins/AI/MassAI/Source/MassNavigation/Private/MassSmoothOrientationProcessors.cpp
new file mode 100644
--- /dev/null
+++ b/Engine/Plugins/AI/MassAI/Source/MassNavigation/Private/MassSmoothOrientationProcessors.cpp
@@ -0,0 +1,137 @@
+﻿// Copyright Epic Games, Inc. All Rights Reserved.
+
+#include "SmoothOrientation/MassSmoothOrientationProcessors.h"
+#include "SmoothOrientation/MassSmoothOrientationFragments.h"
+#include "MassCommandBuffer.h"
+#include "MassCommonFragments.h"
+#include "MassNavigationFragments.h"
+#include "MassMovementFragments.h"
+#include "Math/UnrealMathUtility.h"
+#include "MassSimulationLOD.h"
+#include "MassNavigationUtils.h"
+
+#define UNSAFE_FOR_MT 0
+#define MOVEMENT_DEBUGDRAW 0	// Set to 1 to see heading debugdraw
+
+//----------------------------------------------------------------------//
+//  UMassSmoothOrientationProcessor
+//----------------------------------------------------------------------//
+UMassSmoothOrientationProcessor::UMassSmoothOrientationProcessor()
+{
+	ExecutionFlags = (int32)EProcessorExecutionFlags::All;
+	ExecutionOrder.ExecuteInGroup = UE::Mass::ProcessorGroupNames::Movement;
+}
+
+void UMassSmoothOrientationProcessor::ConfigureQueries()
+{
+	HighResEntityQuery.AddRequirement<FMassMoveTargetFragment>(EMassFragmentAccess::ReadOnly);
+	HighResEntityQuery.AddRequirement<FMassVelocityFragment>(EMassFragmentAccess::ReadWrite);
+	HighResEntityQuery.AddRequirement<FDataFragment_Transform>(EMassFragmentAccess::ReadWrite);
+	HighResEntityQuery.AddTagRequirement<FMassOffLODTag>(EMassFragmentPresence::None);
+	HighResEntityQuery.AddConstSharedRequirement<FMassSmoothOrientationParameters>(EMassFragmentPresence::All);
+
+	LowResEntityQuery_Conditional.AddRequirement<FDataFragment_Transform>(EMassFragmentAccess::ReadWrite);
+	LowResEntityQuery_Conditional.AddRequirement<FMassMoveTargetFragment>(EMassFragmentAccess::ReadOnly);
+	LowResEntityQuery_Conditional.AddTagRequirement<FMassOffLODTag>(EMassFragmentPresence::All);
+	LowResEntityQuery_Conditional.AddChunkRequirement<FMassSimulationVariableTickChunkFragment>(EMassFragmentAccess::ReadOnly, EMassFragmentPresence::Optional);
+	LowResEntityQuery_Conditional.SetChunkFilter(&FMassSimulationVariableTickChunkFragment::ShouldTickChunkThisFrame);
+}
+
+void UMassSmoothOrientationProcessor::Execute(UMassEntitySubsystem& EntitySubsystem,
+													FMassExecutionContext& Context)
+{
+	// Clamp max delta time to avoid force explosion on large time steps (i.e. during initialization).
+	const float DeltaTime = FMath::Min(0.1f, Context.GetDeltaTimeSeconds());
+
+	{
+		QUICK_SCOPE_CYCLE_COUNTER(HighRes);
+
+		HighResEntityQuery.ForEachEntityChunk(EntitySubsystem, Context, [this, DeltaTime](FMassExecutionContext& Context)
+		{
+			const int32 NumEntities = Context.GetNumEntities();
+
+			const FMassSmoothOrientationParameters& OrientationParams = Context.GetConstSharedFragment<FMassSmoothOrientationParameters>();
+
+			const TConstArrayView<FMassMoveTargetFragment> MoveTargetList = Context.GetFragmentView<FMassMoveTargetFragment>();
+			const TArrayView<FDataFragment_Transform> LocationList = Context.GetMutableFragmentView<FDataFragment_Transform>();
+			const TArrayView<FMassVelocityFragment> VelocityList = Context.GetMutableFragmentView<FMassVelocityFragment>();
+
+			for (int32 EntityIndex = 0; EntityIndex < NumEntities; ++EntityIndex)
+			{
+				const FMassMoveTargetFragment& MoveTarget = MoveTargetList[EntityIndex];
+
+				// Do not touch transform at all when animating
+				if (MoveTarget.GetCurrentAction() == EMassMovementAction::Animate)
+				{
+					continue;
+				}
+
+				const FMassVelocityFragment& CurrentVelocity = VelocityList[EntityIndex];
+				FTransform& CurrentTransform = LocationList[EntityIndex].GetMutableTransform();
+				const FVector CurrentForward = CurrentTransform.GetRotation().GetForwardVector();
+				const float CurrentHeading = UE::MassNavigation::GetYawFromDirection(CurrentForward);
+
+				const float EndOfPathAnticipationDistance = OrientationParams.EndOfPathDuration * MoveTarget.DesiredSpeed.Get();
+				
+				float MoveTargetWeight = 0.5f;
+				float VelocityWeight = 0.5f;
+				
+				if (MoveTarget.GetCurrentAction() == EMassMovementAction::Move)
+				{
+					if (MoveTarget.IntentAtGoal == EMassMovementAction::Stand && MoveTarget.DistanceToGoal < EndOfPathAnticipationDistance)
+					{
+						// Fade towards the movement target direction at the end of the path.
+						const float Fade = FMath::Square(FMath::Clamp(MoveTarget.DistanceToGoal / EndOfPathAnticipationDistance, 0.0f, 1.0f)); // zero at end of the path
+
+						MoveTargetWeight = FMath::Lerp(OrientationParams.Standing.MoveTargetWeight, OrientationParams.Moving.MoveTargetWeight, Fade);
+						VelocityWeight = FMath::Lerp(OrientationParams.Standing.VelocityWeight, OrientationParams.Moving.VelocityWeight, Fade);
+					}
+					else
+					{
+						MoveTargetWeight = OrientationParams.Moving.MoveTargetWeight;
+						VelocityWeight = OrientationParams.Moving.VelocityWeight;
+					}
+				}
+				else // Stand
+				{
+					MoveTargetWeight = OrientationParams.Standing.MoveTargetWeight;
+					VelocityWeight = OrientationParams.Standing.VelocityWeight;
+				}
+				
+				const float VelocityHeading = UE::MassNavigation::GetYawFromDirection(CurrentVelocity.Value);
+				const float MovementHeading = UE::MassNavigation::GetYawFromDirection(MoveTarget.Forward);
+
+				const float Ratio = MoveTargetWeight / (MoveTargetWeight + VelocityWeight);
+				const float DesiredHeading = UE::MassNavigation::LerpAngle(VelocityHeading, MovementHeading,Ratio);
+				
+				const float NewHeading = UE::MassNavigation::ExponentialSmoothingAngle(CurrentHeading, DesiredHeading, DeltaTime, OrientationParams.OrientationSmoothingTime);
+
+				FQuat Rotation(FVector::UpVector, NewHeading);
+				CurrentTransform.SetRotation(Rotation);
+			}
+		});
+	}
+
+	{
+		QUICK_SCOPE_CYCLE_COUNTER(LowRes);
+
+		LowResEntityQuery_Conditional.ForEachEntityChunk(EntitySubsystem, Context, [this](FMassExecutionContext& Context)
+		{
+			const int32 NumEntities = Context.GetNumEntities();
+
+			const TArrayView<FDataFragment_Transform> LocationList = Context.GetMutableFragmentView<FDataFragment_Transform>();
+			const TConstArrayView<FMassMoveTargetFragment> MoveTargetList = Context.GetFragmentView<FMassMoveTargetFragment>();
+
+			for (int32 EntityIndex = 0; EntityIndex < NumEntities; ++EntityIndex)
+			{
+				FTransform& CurrentTransform = LocationList[EntityIndex].GetMutableTransform();
+				const FMassMoveTargetFragment& MoveTarget = MoveTargetList[EntityIndex];
+
+				// Snap position to move target directly
+				CurrentTransform.SetRotation(FQuat::FindBetweenNormals(FVector::ForwardVector, MoveTarget.Forward));
+			}
+		});
+	}
+}
+
+#undef UNSAFE_FOR_MT
diff --git a/Engine/Plugins/AI/MassAI/Source/MassNavigation/Private/MassSmoothOrientationTrait.cpp b/Engine/Plugins/AI/MassAI/Source/MassNavigation/Private/MassSmoothOrientationTrait.cpp
new file mode 100644
--- /dev/null
+++ b/Engine/Plugins/AI/MassAI/Source/MassNavigation/Private/MassSmoothOrientationTrait.cpp
@@ -0,0 +1,20 @@
+﻿// Copyright Epic Games, Inc. All Rights Reserved.
+#include "SmoothOrientation/MassSmoothOrientationTrait.h"
+#include "MassEntityTemplateRegistry.h"
+#include "MassMovementFragments.h"
+#include "MassCommonFragments.h"
+#include "MassNavigationFragments.h"
+#include "Engine/World.h"
+
+void UMassSmoothOrientationTrait::BuildTemplate(FMassEntityTemplateBuildContext& BuildContext, UWorld& World) const
+{
+	UMassEntitySubsystem* EntitySubsystem = UWorld::GetSubsystem<UMassEntitySubsystem>(&World);
+	check(EntitySubsystem);
+
+	BuildContext.AddFragment<FMassMoveTargetFragment>();
+	BuildContext.AddFragment<FMassVelocityFragment>();
+	BuildContext.AddFragment<FDataFragment_Transform>();
+
+	const FConstSharedStruct OrientationFragment = EntitySubsystem->GetOrCreateConstSharedFragment(UE::StructUtils::GetStructCrc32(FConstStructView::Make(Orientation)), Orientation);
+	BuildContext.AddConstSharedFragment(OrientationFragment);
+}
diff --git a/Engine/Plugins/AI/MassAI/Source/MassNavigation/Private/MassSteeringProcessors.cpp b/Engine/Plugins/AI/MassAI/Source/MassNavigation/Private/MassSteeringProcessors.cpp
new file mode 100644
--- /dev/null
+++ b/Engine/Plugins/AI/MassAI/Source/MassNavigation/Private/MassSteeringProcessors.cpp
@@ -0,0 +1,324 @@
+﻿// Copyright Epic Games, Inc. All Rights Reserved.
+
+#include "Steering/MassSteeringProcessors.h"
+#include "MassCommonUtils.h"
+#include "MassCommandBuffer.h"
+#include "MassCommonFragments.h"
+#include "MassMovementFragments.h"
+#include "MassNavigationFragments.h"
+#include "Steering/MassSteeringFragments.h"
+#include "Math/UnrealMathUtility.h"
+#include "MassSimulationLOD.h"
+
+#define UNSAFE_FOR_MT 0
+#define MOVEMENT_DEBUGDRAW 0	// Set to 1 to see heading debugdraw
+
+namespace UE::MassNavigation
+{
+	/*
+	* Calculates speed scale based on agent's forward direction and desired steering direction.
+	*/
+	static float CalcDirectionalSpeedScale(const FVector ForwardDirection, const FVector SteerDirection)
+	{
+		// @todo: make these configurable
+		constexpr float ForwardSpeedScale = 1.0f;
+		constexpr float BackwardSpeedScale = 0.25f;
+		constexpr float SideSpeedScale = 0.5f;
+
+		const FVector LeftDirection = FVector::CrossProduct(ForwardDirection, FVector::UpVector);
+		const float DirX = FVector::DotProduct(LeftDirection, SteerDirection);
+		const float DirY = FVector::DotProduct(ForwardDirection, SteerDirection);
+
+		// Calculate intersection between a direction vector and ellipse, where A & B are the size of the ellipse.
+		// The direction vector is starting from the center of the ellipse.
+		constexpr float SideA = SideSpeedScale;
+		const float SideB = DirY > 0.0f ? ForwardSpeedScale : BackwardSpeedScale;
+		const float Disc = FMath::Square(SideA) * FMath::Square(DirY) + FMath::Square(SideB) * FMath::Square(DirX);
+		const float Speed = (Disc > SMALL_NUMBER) ? (SideA * SideB / FMath::Sqrt(Disc)) : 0.0f;;
+
+		return Speed;
+	}
+
+	/** Speed envelope when approaching a point. NormalizedDistance in range [0..1] */
+	static float ArrivalSpeedEnvelope(const float NormalizedDistance)
+	{
+		return FMath::Sqrt(NormalizedDistance);
+	}
+
+} // UE::MassNavigation
+
+//----------------------------------------------------------------------//
+//  UMassSteerToMoveTargetProcessor
+//----------------------------------------------------------------------//
+UMassSteerToMoveTargetProcessor::UMassSteerToMoveTargetProcessor()
+{
+	ExecutionFlags = int32(EProcessorExecutionFlags::All);
+	ExecutionOrder.ExecuteAfter.Add(UE::Mass::ProcessorGroupNames::Tasks);
+	ExecutionOrder.ExecuteBefore.Add(UE::Mass::ProcessorGroupNames::Avoidance);
+}
+
+void UMassSteerToMoveTargetProcessor::ConfigureQueries()
+{
+	EntityQuery.AddRequirement<FDataFragment_Transform>(EMassFragmentAccess::ReadOnly);
+	EntityQuery.AddRequirement<FMassMoveTargetFragment>(EMassFragmentAccess::ReadWrite);
+	EntityQuery.AddRequirement<FMassSteeringFragment>(EMassFragmentAccess::ReadWrite);
+	EntityQuery.AddRequirement<FMassStandingSteeringFragment>(EMassFragmentAccess::ReadWrite);
+	EntityQuery.AddRequirement<FMassGhostLocationFragment>(EMassFragmentAccess::ReadWrite);
+	EntityQuery.AddRequirement<FMassForceFragment>(EMassFragmentAccess::ReadWrite);
+	EntityQuery.AddRequirement<FMassVelocityFragment>(EMassFragmentAccess::ReadWrite);
+	EntityQuery.AddConstSharedRequirement<FMassMovementParameters>(EMassFragmentPresence::All);
+	EntityQuery.AddConstSharedRequirement<FMassMovingSteeringParameters>(EMassFragmentPresence::All);
+	EntityQuery.AddConstSharedRequirement<FMassStandingSteeringParameters>(EMassFragmentPresence::All);
+
+	// No need for Off LOD to do steering, applying move target directly
+	EntityQuery.AddTagRequirement<FMassOffLODTag>(EMassFragmentPresence::None);
+}
+
+void UMassSteerToMoveTargetProcessor::Execute(UMassEntitySubsystem& EntitySubsystem, FMassExecutionContext& Context)
+{
+	EntityQuery.ForEachEntityChunk(EntitySubsystem, Context, [this](FMassExecutionContext& Context)
+	{
+		const int32 NumEntities = Context.GetNumEntities();
+		const TConstArrayView<FDataFragment_Transform> TransformList = Context.GetFragmentView<FDataFragment_Transform>();
+		const TArrayView<FMassMoveTargetFragment> MoveTargetList = Context.GetMutableFragmentView<FMassMoveTargetFragment>();
+		const TArrayView<FMassVelocityFragment> VelocityList = Context.GetMutableFragmentView<FMassVelocityFragment>();
+		const TArrayView<FMassForceFragment> ForceList = Context.GetMutableFragmentView<FMassForceFragment>();
+		const TArrayView<FMassSteeringFragment> SteeringList = Context.GetMutableFragmentView<FMassSteeringFragment>();
+		const TArrayView<FMassStandingSteeringFragment> StandingSteeringList = Context.GetMutableFragmentView<FMassStandingSteeringFragment>();
+		const TArrayView<FMassGhostLocationFragment> GhostList = Context.GetMutableFragmentView<FMassGhostLocationFragment>();
+		const FMassMovementParameters& MovementParams = Context.GetConstSharedFragment<FMassMovementParameters>();
+		const FMassMovingSteeringParameters& MovingSteeringParams = Context.GetConstSharedFragment<FMassMovingSteeringParameters>();
+		const FMassStandingSteeringParameters& StandingSteeringParams = Context.GetConstSharedFragment<FMassStandingSteeringParameters>();
+
+		const float SteerK = 1.f / MovingSteeringParams.ReactionTime;
+		const float DeltaTime = Context.GetDeltaTimeSeconds();
+
+		for (int32 EntityIndex = 0; EntityIndex < NumEntities; ++EntityIndex)
+		{
+			const FDataFragment_Transform& TransformFragment = TransformList[EntityIndex];
+			FMassSteeringFragment& Steering = SteeringList[EntityIndex];
+			FMassStandingSteeringFragment& StandingSteering = StandingSteeringList[EntityIndex];
+			FMassGhostLocationFragment& Ghost = GhostList[EntityIndex];
+			FMassMoveTargetFragment& MoveTarget = MoveTargetList[EntityIndex];
+			FMassForceFragment& Force = ForceList[EntityIndex];
+			FMassVelocityFragment& Velocity = VelocityList[EntityIndex];
+			const FMassEntityHandle Entity = Context.GetEntity(EntityIndex);
+
+			const FTransform& Transform = TransformFragment.GetTransform();;
+
+			// Calculate velocity for steering.
+			const FVector CurrentLocation = Transform.GetLocation();
+			const FVector CurrentForward = Transform.GetRotation().GetForwardVector();
+
+			const float LookAheadDistance = FMath::Max(1.0f, MovingSteeringParams.LookAheadTime * MoveTarget.DesiredSpeed.Get());
+
+			if (MoveTarget.GetCurrentAction() == EMassMovementAction::Move)
+			{
+				// Tune down avoidance and speed when arriving at goal.
+				float ArrivalFade = 1.0f;
+				if (MoveTarget.IntentAtGoal == EMassMovementAction::Stand)
+				{
+					ArrivalFade = FMath::Clamp(MoveTarget.DistanceToGoal / LookAheadDistance, 0.0f, 1.0f);
+				}
+				const float SteeringPredictionDistance = LookAheadDistance * ArrivalFade;
+
+				// Steer towards and along the move target.
+				const FVector TargetSide = FVector::CrossProduct(MoveTarget.Forward, FVector::UpVector);
+				const FVector Delta = CurrentLocation - MoveTarget.Center;
+
+				const float ForwardOffset = FVector::DotProduct(MoveTarget.Forward, Delta);
+
+				// Calculate steering direction. When far away from the line defined by TargetPosition and TargetTangent,
+				// the steering direction is towards the line, the close we get, the more it aligns with the line.
+				const float SidewaysOffset = FVector::DotProduct(TargetSide, Delta);
+				const float SteerForward = FMath::Sqrt(FMath::Max(0.0f, FMath::Square(SteeringPredictionDistance) - FMath::Square(SidewaysOffset)));
+
+				// The Max() here makes the steering directions behind the TargetPosition to steer towards it directly.
+				FVector SteerTarget = MoveTarget.Center + MoveTarget.Forward * FMath::Clamp(ForwardOffset + SteerForward, 0.0f, SteeringPredictionDistance);
+
+				FVector SteerDirection = SteerTarget - CurrentLocation;
+				SteerDirection.Z = 0.0f;
+				const float DistanceToSteerTarget = SteerDirection.Length();
+				if (DistanceToSteerTarget > KINDA_SMALL_NUMBER)
+				{
+					SteerDirection *= 1.0f / DistanceToSteerTarget;
+				}
+				
+				const float DirSpeedScale = UE::MassNavigation::CalcDirectionalSpeedScale(CurrentForward, SteerDirection);
+				float DesiredSpeed = MoveTarget.DesiredSpeed.Get() * DirSpeedScale;
+
+				// Control speed based relation to the forward axis of the move target.
+				float CatchupDesiredSpeed = DesiredSpeed;
+				if (ForwardOffset < 0.0f)
+				{
+					// Falling behind, catch up
+					const float T = FMath::Min(-ForwardOffset / LookAheadDistance, 1.0f);
+					CatchupDesiredSpeed = FMath::Lerp(DesiredSpeed, MovementParams.MaxSpeed, T);
+				}
+				else if (ForwardOffset > 0.0f)
+				{
+					// Ahead, slow down.
+					const float T = FMath::Min(ForwardOffset / LookAheadDistance, 1.0f);
+					CatchupDesiredSpeed = FMath::Lerp(DesiredSpeed, DesiredSpeed * 0.0f, 1.0f - FMath::Square(1.0f - T));
+				}
+
+				// Control speed based on distance to move target. This allows to catch up even if speed above reaches zero.
+				const float DeviantSpeed = FMath::Min(FMath::Abs(SidewaysOffset) / LookAheadDistance, 1.0f) * DesiredSpeed;
+
+				DesiredSpeed = FMath::Max(CatchupDesiredSpeed, DeviantSpeed);
+
+				// Slow down towards the end of path.
+				if (MoveTarget.IntentAtGoal == EMassMovementAction::Stand)
+				{
+					const float NormalizedDistanceToSteerTarget = FMath::Clamp(DistanceToSteerTarget / LookAheadDistance, 0.0f, 1.0f);
+					DesiredSpeed *= UE::MassNavigation::ArrivalSpeedEnvelope(FMath::Max(ArrivalFade, NormalizedDistanceToSteerTarget));
+				}
+
+				MoveTarget.bSteeringFallingBehind = ForwardOffset < -LookAheadDistance * 0.8f;
+
+				// @todo: This current completely overrides steering, we probably should have one processor that resets the steering at the beginning of the frame.
+				Steering.DesiredVelocity = SteerDirection * DesiredSpeed;
+				Force.Value = SteerK * (Steering.DesiredVelocity - Velocity.Value); // Goal force
+			}
+			else if (MoveTarget.GetCurrentAction() == EMassMovementAction::Stand)
+			{
+				// Calculate unique target move threshold so that different agents react a bit differently.
+				const float PerEntityScale = UE::RandomSequence::FRand(Entity.Index);
+				const float TargetMoveThreshold = StandingSteeringParams.TargetMoveThreshold * (1.0f - StandingSteeringParams.TargetMoveThresholdVariance + PerEntityScale * StandingSteeringParams.TargetMoveThresholdVariance * 2.0f);
+				
+				if (Ghost.LastSeenActionID != MoveTarget.GetCurrentActionID())
+				{
+					// Reset when action changes. @todo: should reset only when move->stand?
+					Ghost.Location = MoveTarget.Center;
+					Ghost.Velocity = FVector::ZeroVector;
+					Ghost.LastSeenActionID = MoveTarget.GetCurrentActionID();
+
+					StandingSteering.TargetLocation = MoveTarget.Center;
+					StandingSteering.TrackedTargetSpeed = 0.0f;
+					StandingSteering.bIsUpdatingTarget = false;
+					StandingSteering.TargetSelectionCooldown = StandingSteeringParams.TargetSelectionCooldown * FMath::RandRange(1.0f - StandingSteeringParams.TargetSelectionCooldownVariance, 1.0f + StandingSteeringParams.TargetSelectionCooldownVariance);
+					StandingSteering.bEnteredFromMoveAction = MoveTarget.GetPreviousAction() == EMassMovementAction::Move;
+				}
+
+				StandingSteering.TargetSelectionCooldown = FMath::Max(0.0f, StandingSteering.TargetSelectionCooldown - DeltaTime);
+
+				if (!StandingSteering.bIsUpdatingTarget)
+				{
+					// Update the move target if enough time has passed and the target has moved. 
+					if (StandingSteering.TargetSelectionCooldown <= 0.0f
+						&& FVector::DistSquared(StandingSteering.TargetLocation, Ghost.Location) > FMath::Square(TargetMoveThreshold))
+					{
+						StandingSteering.TargetLocation = Ghost.Location;
+						StandingSteering.TrackedTargetSpeed = 0.0f;
+						StandingSteering.bIsUpdatingTarget = true;
+						StandingSteering.bEnteredFromMoveAction = false;
+					}
+				}
+				else
+				{
+					// Updating target
+					StandingSteering.TargetLocation = Ghost.Location;
+					const float GhostSpeed = Ghost.Velocity.Length();
+					if (GhostSpeed > (StandingSteering.TrackedTargetSpeed * StandingSteeringParams.TargetSpeedHysteresisScale))
+					{
+						StandingSteering.TrackedTargetSpeed = FMath::Max(StandingSteering.TrackedTargetSpeed, GhostSpeed);
+					}
+					else
+					{
+						// Speed is dropping, we have found the peak change, stop updating the target and start cooldown.
+						StandingSteering.TargetSelectionCooldown = StandingSteeringParams.TargetSelectionCooldown * FMath::RandRange(1.0f - StandingSteeringParams.TargetSelectionCooldownVariance, 1.0f + StandingSteeringParams.TargetSelectionCooldownVariance);
+						StandingSteering.bIsUpdatingTarget = false;
+					}
+				}
+				
+				// Move directly towards the move target when standing.
+				FVector SteerDirection = FVector::ZeroVector;
+				float DesiredSpeed = 0.0f;
+
+				FVector Delta = StandingSteering.TargetLocation - CurrentLocation;
+				Delta.Z = 0.0f;
+				const float Distance = Delta.Size();
+				if (Distance > StandingSteeringParams.DeadZoneRadius)
+				{
+					SteerDirection = Delta / Distance;
+					if (StandingSteering.bEnteredFromMoveAction)
+					{
+						// If the current steering target is from approaching a move target, use the same speed logic as movement to ensure smooth transition.
+						const float SpeedFade = FMath::Clamp(Distance / LookAheadDistance, 0.0f, 1.0f);
+						DesiredSpeed = MoveTarget.DesiredSpeed.Get() * UE::MassNavigation::CalcDirectionalSpeedScale(CurrentForward, SteerDirection) * UE::MassNavigation::ArrivalSpeedEnvelope(SpeedFade);
+					}
+					else
+					{
+						// More aggressive movement when doing adjustments.
+						const float AdjustmentLookAheadDistance = LookAheadDistance * 0.5f;
+						// Intentionally not taking dead zone into account here, so that the speed does not drop to zero.
+						const float SpeedFade = FMath::Clamp(Distance / AdjustmentLookAheadDistance, 0.0f, 1.0f);
+						DesiredSpeed = MoveTarget.DesiredSpeed.Get() * UE::MassNavigation::ArrivalSpeedEnvelope(SpeedFade);
+					}
+					
+					// @todo: This current completely overrides steering, we probably should have one processor that resets the steering at the beginning of the frame.
+					Steering.DesiredVelocity = SteerDirection * DesiredSpeed;
+					Force.Value = SteerK * (Steering.DesiredVelocity - Velocity.Value); // Goal force
+
+				}
+				else
+				{
+					// When reached destination, clamp small velocities to zero to avoid tiny drifting.
+					if (Velocity.Value.SquaredLength() < FMath::Square(StandingSteeringParams.LowSpeedThreshold))
+					{
+						Velocity.Value = FVector::ZeroVector;
+						Force.Value = FVector::ZeroVector;
+					}
+				}
+
+				MoveTarget.bSteeringFallingBehind = false;
+			}
+			else if (MoveTarget.GetCurrentAction() == EMassMovementAction::Animate)
+			{
+				// Stop all movement when animating.
+				Steering.Reset();
+				MoveTarget.bSteeringFallingBehind = false;
+				Force.Value = FVector::ZeroVector;
+				Velocity.Value = FVector::ZeroVector;
+			}
+
+#if WITH_MASSGAMEPLAY_DEBUG && UNSAFE_FOR_MT
+			FColor EntityColor = FColor::White;
+			const bool bDisplayDebug = UE::Mass::Debug::IsDebuggingEntity(Entity, &EntityColor);
+			if (bDisplayDebug)
+			{
+				const FVector ZOffset(0,0,25);
+
+				const FColor DarkEntityColor = UE::MassNavigation::Debug::MixColors(EntityColor, FColor::Black);
+				const FColor LightEntityColor = UE::MassNavigation::Debug::MixColors(EntityColor, FColor::White);
+				
+				const FVector MoveTargetCenter = MoveTarget.Center + ZOffset;
+
+				// MoveTarget slack boundary
+				UE_VLOG_CIRCLE_THICK(this, LogMassNavigation, Log, MoveTargetCenter, FVector::UpVector, LookAheadDistance, EntityColor, /*Thickness*/2,
+					TEXT("%s MoveTgt %s"), *Entity.DebugGetDescription(), *UEnum::GetDisplayValueAsText(MoveTarget.IntentAtGoal).ToString());
+
+				// MoveTarget orientation
+				UE_VLOG_SEGMENT_THICK(this, LogMassNavigation, Log, MoveTargetCenter, MoveTargetCenter + MoveTarget.Forward * LookAheadDistance, EntityColor, /*Thickness*/2, TEXT(""));
+
+				// MoveTarget - current location relation.
+				if (FVector::Dist2D(CurrentLocation, MoveTarget.Center) > LookAheadDistance * 1.5f)
+				{
+					UE_VLOG_SEGMENT_THICK(this, LogMassNavigation, Log, MoveTargetCenter, CurrentLocation + ZOffset, FColor::Red, /*Thickness*/1, TEXT("LOST"));
+				}
+				else
+				{
+					UE_VLOG_SEGMENT_THICK(this, LogMassNavigation, Log, MoveTargetCenter, CurrentLocation + ZOffset, DarkEntityColor, /*Thickness*/1, TEXT(""));
+				}
+
+				// Steering
+				UE_VLOG_SEGMENT_THICK(this, LogMassNavigation, Log, CurrentLocation + ZOffset, CurrentLocation + Steering.DesiredVelocity + ZOffset, LightEntityColor, /*Thickness*/2,
+					TEXT("%s Steer %.1f"), *Entity.DebugGetDescription(), Steering.DesiredVelocity.Length());
+			}
+#endif // WITH_MASSGAMEPLAY_DEBUG
+			
+		}
+	});
+}
+
+#undef UNSAFE_FOR_MT
diff --git a/Engine/Plugins/AI/MassAI/Source/MassNavigation/Private/MassSteeringTrait.cpp b/Engine/Plugins/AI/MassAI/Source/MassNavigation/Private/MassSteeringTrait.cpp
new file mode 100644
--- /dev/null
+++ b/Engine/Plugins/AI/MassAI/Source/MassNavigation/Private/MassSteeringTrait.cpp
@@ -0,0 +1,31 @@
+﻿// Copyright Epic Games, Inc. All Rights Reserved.
+
+#include "Steering/MassSteeringTrait.h"
+#include "MassEntityTemplateRegistry.h"
+#include "MassCommonFragments.h"
+#include "MassMovementFragments.h"
+#include "MassNavigationFragments.h"
+#include "Steering/MassSteeringFragments.h"
+#include "Engine/World.h"
+
+void UMassSteeringTrait::BuildTemplate(FMassEntityTemplateBuildContext& BuildContext, UWorld& World) const
+{
+	UMassEntitySubsystem* EntitySubsystem = UWorld::GetSubsystem<UMassEntitySubsystem>(&World);
+	check(EntitySubsystem);
+
+	BuildContext.AddFragment<FDataFragment_AgentRadius>();
+	BuildContext.AddFragment<FDataFragment_Transform>();
+	BuildContext.AddFragment<FMassVelocityFragment>();
+	BuildContext.AddFragment<FMassForceFragment>();
+	BuildContext.AddFragment<FMassMoveTargetFragment>();
+	
+	BuildContext.AddFragment<FMassSteeringFragment>();
+	BuildContext.AddFragment<FMassStandingSteeringFragment>();
+	BuildContext.AddFragment<FMassGhostLocationFragment>();
+
+	const FConstSharedStruct MovingSteeringFragment = EntitySubsystem->GetOrCreateConstSharedFragment(UE::StructUtils::GetStructCrc32(FConstStructView::Make(MovingSteering)), MovingSteering);
+	BuildContext.AddConstSharedFragment(MovingSteeringFragment);
+
+	const FConstSharedStruct StandingSteeringFragment = EntitySubsystem->GetOrCreateConstSharedFragment(UE::StructUtils::GetStructCrc32(FConstStructView::Make(StandingSteering)), StandingSteering);
+	BuildContext.AddConstSharedFragment(StandingSteeringFragment);
+}
diff --git a/Engine/Plugins/AI/MassAI/Source/MassNavigation/Public/Avoidance/MassAvoidanceProcessors.h b/Engine/Plugins/AI/MassAI/Source/MassNavigation/Public/Avoidance/MassAvoidanceProcessors.h
new file mode 100644
--- /dev/null
+++ b/Engine/Plugins/AI/MassAI/Source/MassNavigation/Public/Avoidance/MassAvoidanceProcessors.h
@@ -0,0 +1,53 @@
+// Copyright Epic Games, Inc. All Rights Reserved.
+
+#pragma once
+
+#include "MassProcessor.h"
+#include "MassAvoidanceProcessors.generated.h"
+
+MASSNAVIGATION_API DECLARE_LOG_CATEGORY_EXTERN(LogAvoidance, Warning, All);
+MASSNAVIGATION_API DECLARE_LOG_CATEGORY_EXTERN(LogAvoidanceVelocities, Warning, All);
+MASSNAVIGATION_API DECLARE_LOG_CATEGORY_EXTERN(LogAvoidanceAgents, Warning, All);
+MASSNAVIGATION_API DECLARE_LOG_CATEGORY_EXTERN(LogAvoidanceObstacles, Warning, All);
+
+class UMassNavigationSubsystem;
+
+/** Experimental: move using cumulative forces to avoid close agents */
+UCLASS()
+class MASSNAVIGATION_API UMassMovingAvoidanceProcessor : public UMassProcessor
+{
+	GENERATED_BODY()
+
+public:
+	UMassMovingAvoidanceProcessor();
+
+protected:
+	virtual void ConfigureQueries() override;
+	virtual void Initialize(UObject& Owner) override;
+	virtual void Execute(UMassEntitySubsystem& EntitySubsystem, FMassExecutionContext& Context) override;
+
+private:
+	TObjectPtr<UWorld> World;
+	TObjectPtr<UMassNavigationSubsystem> NavigationSubsystem;
+	FMassEntityQuery EntityQuery;
+};
+
+/** Avoidance while standing. */
+UCLASS()
+class MASSNAVIGATION_API UMassStandingAvoidanceProcessor : public UMassProcessor
+{
+	GENERATED_BODY()
+
+public:
+	UMassStandingAvoidanceProcessor();
+
+protected:
+	virtual void ConfigureQueries() override;
+	virtual void Initialize(UObject& Owner) override;
+	virtual void Execute(UMassEntitySubsystem& EntitySubsystem, FMassExecutionContext& Context) override;
+
+private:
+	TObjectPtr<UWorld> World;
+	TObjectPtr<UMassNavigationSubsystem> NavigationSubsystem;
+	FMassEntityQuery EntityQuery;
+};
diff --git a/Engine/Plugins/AI/MassAI/Source/MassNavigation/Public/Avoidance/MassAvoidanceTrait.h b/Engine/Plugins/AI/MassAI/Source/MassNavigation/Public/Avoidance/MassAvoidanceTrait.h
new file mode 100644
--- /dev/null
+++ b/Engine/Plugins/AI/MassAI/Source/MassNavigation/Public/Avoidance/MassAvoidanceTrait.h
@@ -0,0 +1,23 @@
+// Copyright Epic Games, Inc. All Rights Reserved.
+
+#pragma once 
+
+#include "MassEntityTraitBase.h"
+#include "MassAvoidanceFragments.h"
+#include "MassAvoidanceTrait.generated.h"
+
+UCLASS(meta = (DisplayName = "Avoidance"))
+class MASSNAVIGATION_API UMassObstacleAvoidanceTrait : public UMassEntityTraitBase
+{
+	GENERATED_BODY()
+
+protected:
+	
+	virtual void BuildTemplate(FMassEntityTemplateBuildContext& BuildContext, UWorld& World) const override;
+
+	UPROPERTY(EditAnywhere, Category="")
+	FMassMovingAvoidanceParameters MovingParameters;
+	
+	UPROPERTY(EditAnywhere, Category="")
+	FMassStandingAvoidanceParameters StandingParameters;
+};
diff --git a/Engine/Plugins/AI/MassAI/Source/MassNavigation/Public/Avoidance/MassNavigationObstacleTrait.h b/Engine/Plugins/AI/MassAI/Source/MassNavigation/Public/Avoidance/MassNavigationObstacleTrait.h
new file mode 100644
--- /dev/null
+++ b/Engine/Plugins/AI/MassAI/Source/MassNavigation/Public/Avoidance/MassNavigationObstacleTrait.h
@@ -0,0 +1,16 @@
+// Copyright Epic Games, Inc. All Rights Reserved.
+
+#pragma once 
+
+#include "MassEntityTraitBase.h"
+#include "MassNavigationObstacleTrait.generated.h"
+
+UCLASS(meta = (DisplayName = "Navigation Obstacle"))
+class MASSNAVIGATION_API UMassNavigationObstacleTrait : public UMassEntityTraitBase
+{
+	GENERATED_BODY()
+
+protected:
+	virtual void BuildTemplate(FMassEntityTemplateBuildContext& BuildContext, UWorld& World) const override;
+
+};
\ No newline at end of file
diff --git a/Engine/Plugins/AI/MassAI/Source/MassNavigation/Public/MassNavigationProcessors.h b/Engine/Plugins/AI/MassAI/Source/MassNavigation/Public/MassNavigationProcessors.h
new file mode 100644
--- /dev/null
+++ b/Engine/Plugins/AI/MassAI/Source/MassNavigation/Public/MassNavigationProcessors.h
@@ -0,0 +1,103 @@
+// Copyright Epic Games, Inc. All Rights Reserved.
+
+#pragma once
+
+#include "MassProcessor.h"
+#include "MassObserverProcessor.h"
+#include "MassNavigationProcessors.generated.h"
+
+class UMassNavigationSubsystem;
+
+/**
+ * Updates Off-LOD entities position to move targets position.
+ */
+UCLASS()
+class MASSNAVIGATION_API UMassOffLODNavigationProcessor : public UMassProcessor
+{
+	GENERATED_BODY()
+
+public:
+	UMassOffLODNavigationProcessor();
+
+protected:
+	virtual void ConfigureQueries() override;
+	virtual void Execute(UMassEntitySubsystem& EntitySubsystem, FMassExecutionContext& Context) override;
+
+private:
+	FMassEntityQuery EntityQuery_Conditional;
+};
+
+/**
+ * Updates entities height to move targets position smoothly.
+ * Does not update Off-LOD entities.
+ */
+UCLASS()
+class MASSNAVIGATION_API UMassNavigationSmoothHeightProcessor : public UMassProcessor
+{
+	GENERATED_BODY()
+
+public:
+	UMassNavigationSmoothHeightProcessor();
+
+protected:
+	virtual void ConfigureQueries() override;
+	virtual void Execute(UMassEntitySubsystem& EntitySubsystem, FMassExecutionContext& Context) override;
+
+private:
+	FMassEntityQuery EntityQuery;
+};
+
+/**
+ * Initializes the move target's location to the agents initial position.
+ */
+UCLASS()
+class MASSNAVIGATION_API UMassMoveTargetFragmentInitializer : public UMassFragmentInitializer
+{
+	GENERATED_BODY()
+public:
+	UMassMoveTargetFragmentInitializer();
+	
+protected:
+	virtual void ConfigureQueries() override;
+	virtual void Execute(UMassEntitySubsystem& EntitySubsystem, FMassExecutionContext& Context) override;
+
+	FMassEntityQuery InitializerQuery;
+};
+
+/** Processor to update obstacle grid */
+UCLASS()
+class MASSNAVIGATION_API UMassNavigationObstacleGridProcessor : public UMassProcessor
+{
+	GENERATED_BODY()
+
+public:
+	UMassNavigationObstacleGridProcessor();
+
+protected:
+	virtual void ConfigureQueries() override;
+	virtual void Initialize(UObject& Owner) override;
+	virtual void Execute(UMassEntitySubsystem& EntitySubsystem, FMassExecutionContext& Context) override;
+
+private:
+	TObjectPtr<UMassNavigationSubsystem> NavigationSubsystem;
+	FMassEntityQuery AddToGridEntityQuery;
+	FMassEntityQuery UpdateGridEntityQuery;
+	FMassEntityQuery RemoveFromGridEntityQuery;
+};
+
+/** Deinitializer processor to remove avoidance obstacles from the avoidance obstacle grid */
+UCLASS()
+class MASSNAVIGATION_API UMassNavigationObstacleRemoverProcessor : public UMassFragmentDeinitializer
+{
+	GENERATED_BODY()
+
+	UMassNavigationObstacleRemoverProcessor();
+	TObjectPtr<UMassNavigationSubsystem> NavigationSubsystem;
+
+protected:
+	virtual void ConfigureQueries() override;
+	virtual void Initialize(UObject& Owner) override;
+	virtual void Execute(UMassEntitySubsystem& EntitySubsystem, FMassExecutionContext& Context) override;
+
+	FMassEntityQuery EntityQuery;
+};
diff --git a/Engine/Plugins/AI/MassAI/Source/MassNavigation/Public/MassNavigationUtils.h b/Engine/Plugins/AI/MassAI/Source/MassNavigation/Public/MassNavigationUtils.h
new file mode 100644
--- /dev/null
+++ b/Engine/Plugins/AI/MassAI/Source/MassNavigation/Public/MassNavigationUtils.h
@@ -0,0 +1,99 @@
+﻿// Copyright Epic Games, Inc. All Rights Reserved.
+
+#pragma once
+
+#include "CoreTypes.h"
+#include "Math/UnrealMathUtility.h"
+
+namespace UE::MassNavigation
+{
+	// Calculates yaw angle from direction vector.
+	inline float GetYawFromQuat(const FQuat Rotation)
+	{
+		const float YawY = 2.f * (Rotation.W * Rotation.Z + Rotation.X * Rotation.Y);
+		const float YawX = (1.f - 2.f * (FMath::Square(Rotation.Y) + FMath::Square(Rotation.Z)));
+		return FMath::Atan2(YawY, YawX);
+	}
+
+	inline float GetYawFromDirection(const FVector Direction)
+	{
+		return FMath::Atan2(Direction.Y, Direction.X);
+	}
+
+	// Wraps and angle to range -PI..PI. Angle in radians.
+	inline float WrapAngle(const float Angle)
+	{
+		float WrappedAngle = FMath::Fmod(Angle, PI*2.0f);
+		WrappedAngle = (WrappedAngle > PI) ? WrappedAngle - PI*2.0f : WrappedAngle;
+		WrappedAngle = (WrappedAngle < -PI) ? WrappedAngle + PI*2.0f : WrappedAngle;
+		return WrappedAngle;
+	}
+
+	// Linearly interpolates between two angles (in Radians).
+	inline float LerpAngle(const float AngleA, const float AngleB, const float T)
+	{
+		const float DeltaAngle = WrapAngle(AngleB - AngleA);
+		return AngleA + DeltaAngle * T;
+	}
+
+	// Exponential smooth from current angle to target angle. Angles in radians.
+	inline float ExponentialSmoothingAngle(const float Angle, const float TargetAngle, const float DeltaTime, const float SmoothingTime)
+	{
+		// Note: based on FMath::ExponentialSmoothingApprox().
+		if (SmoothingTime < KINDA_SMALL_NUMBER)
+		{
+			return TargetAngle;
+		}
+		const float A = DeltaTime / SmoothingTime;
+		const float Exp = FMath::InvExpApprox(A);
+		return TargetAngle + WrapAngle(Angle - TargetAngle) * Exp;
+	}
+
+	// Clamps vectors magnitude to Mag.
+	inline FVector ClampVector(const FVector Vec, const float Mag)
+	{
+		const float Len = Vec.SizeSquared();
+		if (Len > FMath::Square(Mag)) {
+			return Vec * Mag / FMath::Sqrt(Len);
+		}
+		return Vec;
+	}
+
+	// Projects a point to segment and returns the time interpolation value.
+	inline float ProjectPtSeg(const FVector2D Point, const FVector2D Start, const FVector2D End)
+	{
+		const FVector2D Seg = End - Start;
+		const FVector2D Dir = Point - Start;
+		const float d = Seg.SizeSquared();
+		const float t = FVector2D::DotProduct(Seg, Dir);
+		if (t < 0.f) return 0;
+		if (t > d) return 1;
+		return d > 0.f ? (t / d) : 0.f;
+	}
+
+	// Returns the SmoothStep curve for X in range [0..1]. 
+	inline float Smooth(const float X)
+	{
+		return X * X * (3 - 2 * X);
+	}
+
+	// Returns left direction from forward and up directions.
+	inline FVector GetLeftDirection(const FVector Forward, const FVector Up)
+	{
+		return FVector::CrossProduct(Forward, Up);
+	}
+
+	// Computes miter normal in XY plane from two neighbour edge normals.
+	inline FVector ComputeMiterNormal(const FVector NormalA, const FVector NormalB)
+	{
+		FVector Mid = 0.5f * (NormalA + NormalB);
+		const float MidSquared = FVector::DotProduct(Mid, Mid);
+		if (MidSquared > KINDA_SMALL_NUMBER)
+		{
+			const float Scale = FMath::Min(1.f / MidSquared, 20.f);
+			Mid *= Scale;
+		}
+		return Mid;
+	}
+	
+} // UE::MassMovement
\ No newline at end of file
diff --git a/Engine/Plugins/AI/MassAI/Source/MassNavigation/Public/SmoothOrientation/MassSmoothOrientationProcessors.h b/Engine/Plugins/AI/MassAI/Source/MassNavigation/Public/SmoothOrientation/MassSmoothOrientationProcessors.h
new file mode 100644
--- /dev/null
+++ b/Engine/Plugins/AI/MassAI/Source/MassNavigation/Public/SmoothOrientation/MassSmoothOrientationProcessors.h
@@ -0,0 +1,26 @@
+﻿// Copyright Epic Games, Inc. All Rights Reserved.
+
+#pragma once
+
+#include "MassProcessor.h"
+#include "MassSmoothOrientationProcessors.generated.h"
+
+/**
+ * Updates agent's orientation based on current movement.
+ */
+UCLASS()
+class MASSNAVIGATION_API UMassSmoothOrientationProcessor : public UMassProcessor
+{
+	GENERATED_BODY()
+
+public:
+	UMassSmoothOrientationProcessor();
+
+protected:
+	virtual void ConfigureQueries() override;
+	virtual void Execute(UMassEntitySubsystem& EntitySubsystem, FMassExecutionContext& Context) override;
+
+private:
+	FMassEntityQuery HighResEntityQuery;
+	FMassEntityQuery LowResEntityQuery_Conditional;
+};
diff --git a/Engine/Plugins/AI/MassAI/Source/MassNavigation/Public/SmoothOrientation/MassSmoothOrientationTrait.h b/Engine/Plugins/AI/MassAI/Source/MassNavigation/Public/SmoothOrientation/MassSmoothOrientationTrait.h
new file mode 100644
--- /dev/null
+++ b/Engine/Plugins/AI/MassAI/Source/MassNavigation/Public/SmoothOrientation/MassSmoothOrientationTrait.h
@@ -0,0 +1,19 @@
+﻿// Copyright Epic Games, Inc. All Rights Reserved.
+
+#pragma once 
+
+#include "MassEntityTraitBase.h"
+#include "MassSmoothOrientationFragments.h"
+#include "MassSmoothOrientationTrait.generated.h"
+
+UCLASS(meta = (DisplayName = "Smooth Orientation"))
+class MASSNAVIGATION_API UMassSmoothOrientationTrait : public UMassEntityTraitBase
+{
+	GENERATED_BODY()
+
+protected:
+	virtual void BuildTemplate(FMassEntityTemplateBuildContext& BuildContext, UWorld& World) const override;
+
+	UPROPERTY(EditAnywhere, Category="")
+	FMassSmoothOrientationParameters Orientation;
+};
diff --git a/Engine/Plugins/AI/MassAI/Source/MassNavigation/Public/Steering/MassSteeringProcessors.h b/Engine/Plugins/AI/MassAI/Source/MassNavigation/Public/Steering/MassSteeringProcessors.h
new file mode 100644
--- /dev/null
+++ b/Engine/Plugins/AI/MassAI/Source/MassNavigation/Public/Steering/MassSteeringProcessors.h
@@ -0,0 +1,24 @@
+﻿// Copyright Epic Games, Inc. All Rights Reserved.
+
+#pragma once
+
+#include "MassProcessor.h"
+#include "MassObserverProcessor.h"
+#include "MassSteeringProcessors.generated.h"
+
+/** 
+* Processor for updating steering towards MoveTarget.
+*/
+UCLASS()
+class MASSNAVIGATION_API UMassSteerToMoveTargetProcessor : public UMassProcessor
+{
+	GENERATED_BODY()
+
+protected:
+	UMassSteerToMoveTargetProcessor();
+	
+	virtual void ConfigureQueries() override;
+	virtual void Execute(UMassEntitySubsystem& EntitySubsystem, FMassExecutionContext& Context) override;
+
+	FMassEntityQuery EntityQuery;
+};
diff --git a/Engine/Plugins/AI/MassAI/Source/MassNavigation/Public/Steering/MassSteeringTrait.h b/Engine/Plugins/AI/MassAI/Source/MassNavigation/Public/Steering/MassSteeringTrait.h
new file mode 100644
--- /dev/null
+++ b/Engine/Plugins/AI/MassAI/Source/MassNavigation/Public/Steering/MassSteeringTrait.h
@@ -0,0 +1,23 @@
+﻿// Copyright Epic Games, Inc. All Rights Reserved.
+
+#pragma once 
+
+#include "MassEntityTraitBase.h"
+#include "Steering/MassSteeringFragments.h"
+#include "MassSteeringTrait.generated.h"
+
+
+UCLASS(meta = (DisplayName = "Steering"))
+class MASSNAVIGATION_API UMassSteeringTrait : public UMassEntityTraitBase
+{
+	GENERATED_BODY()
+
+protected:
+	virtual void BuildTemplate(FMassEntityTemplateBuildContext& BuildContext, UWorld& World) const override;
+
+	UPROPERTY(Category="Steering", EditAnywhere, meta=(EditInline))
+	FMassMovingSteeringParameters MovingSteering;
+
+	UPROPERTY(Category="Steering", EditAnywhere, meta=(EditInline))
+	FMassStandingSteeringParameters StandingSteering;
+};
diff --git a/Engine/Plugins/AI/MassAI/Source/MassNavigationEditor/Private/MassNavigationTestingActor.h b/Engine/Plugins/AI/MassAI/Source/MassNavigationEditor/Private/MassNavigationTestingActor.h
new file mode 100644
--- /dev/null
+++ b/Engine/Plugins/AI/MassAI/Source/MassNavigationEditor/Private/MassNavigationTestingActor.h
@@ -0,0 +1,121 @@
+// Copyright Epic Games, Inc. All Rights Reserved.
+
+#pragma once
+
+#include "GameFramework/Actor.h"
+#include "Debug/DebugDrawComponent.h"
+#include "DebugRenderSceneProxy.h"
+#include "ZoneGraphTypes.h"
+#include "MassZoneGraphNavigationFragments.h"
+#include "MassNavigationTestingActor.generated.h"
+
+class AZoneGraphData;
+class UZoneGraphTestingComponent;
+class UZoneGraphSubsystem;
+class AMassNavigationTestingActor;
+
+#if !UE_BUILD_SHIPPING && !UE_BUILD_TEST
+class MASSNAVIGATIONEDITOR_API FMassNavigationTestingSceneProxy final : public FDebugRenderSceneProxy
+{
+public:
+	FMassNavigationTestingSceneProxy(const UPrimitiveComponent& InComponent);
+	
+	virtual SIZE_T GetTypeHash() const override;
+	virtual FPrimitiveViewRelevance GetViewRelevance(const FSceneView* View) const override;
+	virtual uint32 GetMemoryFootprint(void) const override;
+};
+#endif
+
+/** Component for testing MassMovement functionality. */
+UCLASS(ClassGroup = Debug)
+class MASSNAVIGATIONEDITOR_API UMassNavigationTestingComponent : public UDebugDrawComponent
+{
+	GENERATED_BODY()
+public:
+	UMassNavigationTestingComponent(const FObjectInitializer& ObjectInitializer = FObjectInitializer::Get());
+
+#if WITH_EDITOR
+	virtual void PostEditChangeProperty(FPropertyChangedEvent& PropertyChangedEvent) override;
+#endif // WITH_EDITOR
+
+	virtual void OnRegister() override;
+	virtual void OnUnregister() override;
+
+	virtual FBoxSphereBounds CalcBounds(const FTransform& LocalToWorld) const override;
+
+#if UE_ENABLE_DEBUG_DRAWING
+	virtual FDebugRenderSceneProxy* CreateDebugSceneProxy() override;
+#endif
+
+	void UpdateTests();
+	void PinLane();
+	void ClearPinnedLane();
+
+protected:
+
+#if WITH_EDITOR
+	void OnZoneGraphDataBuildDone(const struct FZoneGraphBuildData& BuildData);
+#endif
+	void OnZoneGraphDataChanged(const AZoneGraphData* ZoneGraphData);
+
+#if WITH_EDITOR
+	FDelegateHandle OnDataChangedHandle;
+#endif
+	FDelegateHandle OnDataAddedHandle;
+	FDelegateHandle OnDataRemovedHandle;
+
+	UPROPERTY(Transient)
+	UZoneGraphSubsystem* ZoneGraph;
+
+	UPROPERTY(Transient)
+	FZoneGraphLaneLocation LaneLocation;
+
+	UPROPERTY(Transient)
+	FZoneGraphLaneLocation GoalLaneLocation;
+
+	UPROPERTY(EditAnywhere, Category = Default);
+	FVector SearchExtent;
+
+	UPROPERTY(EditAnywhere, Category = Default);
+	float AnticipationDistance = 50.0f;
+
+	UPROPERTY(EditAnywhere, Category = Default);
+	float AgentRadius = 40.0f;
+
+	UPROPERTY(EditAnywhere, Category = Default);
+	bool bHasSpecificEndPoint = true;
+
+	UPROPERTY(EditAnywhere, Category = Default);
+	FZoneGraphTagFilter QueryFilter;
+
+	UPROPERTY(EditAnywhere, Category = Default, meta = (MakeEditWidget=true))
+	FVector GoalPosition;
+
+	FZoneGraphLaneHandle PinnedLane;
+	
+	FMassZoneGraphCachedLaneFragment CachedLane;
+	TArray<FMassZoneGraphShortPathFragment> ShortPaths;
+};
+
+/** Debug actor to visually test zone graph. */
+UCLASS(hidecategories = (Actor, Input, Collision, Rendering, Replication, Partition, HLOD, Cooking))
+class MASSNAVIGATIONEDITOR_API AMassNavigationTestingActor : public AActor
+{
+	GENERATED_BODY()
+public:
+	AMassNavigationTestingActor(const FObjectInitializer& ObjectInitializer = FObjectInitializer::Get());
+
+#if WITH_EDITOR
+	virtual void PostEditMove(bool bFinished) override;
+#endif // WITH_EDITOR
+
+	UFUNCTION(BlueprintCallable, CallInEditor, Category = "Default")
+	void PinLane();
+	
+	UFUNCTION(BlueprintCallable, CallInEditor, Category = "Default")
+	void ClearPinnedLane();
+
+protected:
+	UPROPERTY(Category = Default, VisibleAnywhere, meta = (AllowPrivateAccess = "true"))
+	UMassNavigationTestingComponent* DebugComp;
+};
diff --git a/Engine/Plugins/AI/MassAI/Source/MassNavigationEditor/Public/IMassNavigationEditor.h b/Engine/Plugins/AI/MassAI/Source/MassNavigationEditor/Public/IMassNavigationEditor.h
new file mode 100644
--- /dev/null
+++ b/Engine/Plugins/AI/MassAI/Source/MassNavigationEditor/Public/IMassNavigationEditor.h
@@ -0,0 +1,15 @@
+// Copyright Epic Games, Inc. All Rights Reserved.
+
+#pragma once
+
+#include "CoreMinimal.h"
+#include "Toolkits/AssetEditorToolkit.h"
+
+/** MassNavigation Editor public interface */
+class MASSNAVIGATIONEDITOR_API IMassNavigationEditor : public FAssetEditorToolkit
+{
+public:
+
+};
+
+
diff --git a/Engine/Plugins/AI/MassAI/Source/MassNavigationEditor/Public/MassNavigationEditorModule.h b/Engine/Plugins/AI/MassAI/Source/MassNavigationEditor/Public/MassNavigationEditorModule.h
new file mode 100644
--- /dev/null
+++ b/Engine/Plugins/AI/MassAI/Source/MassNavigationEditor/Public/MassNavigationEditorModule.h
@@ -0,0 +1,23 @@
+// Copyright Epic Games, Inc. All Rights Reserved.
+
+#pragma  once
+
+#include "CoreMinimal.h"
+#include "Modules/ModuleManager.h"
+#include "AssetTypeCategories.h"
+#include "Toolkits/IToolkitHost.h"
+#include "Toolkits/AssetEditorToolkit.h"
+
+
+class IMassNavigationEditor;
+
+class MASSNAVIGATIONEDITOR_API FMassNavigationEditorModule : public IModuleInterface
+{
+public:
+	// Begin IModuleInterface
+	virtual void StartupModule() override;
+	virtual void ShutdownModule() override;
+	// End IModuleInterface
+
+private:
+};
diff --git a/Engine/Plugins/AI/MassAI/Source/MassZoneGraphNavigation/Private/MassZoneGraphNavigationTrait.cpp b/Engine/Plugins/AI/MassAI/Source/MassZoneGraphNavigation/Private/MassZoneGraphNavigationTrait.cpp
new file mode 100644
--- /dev/null
+++ b/Engine/Plugins/AI/MassAI/Source/MassZoneGraphNavigation/Private/MassZoneGraphNavigationTrait.cpp
@@ -0,0 +1,27 @@
+﻿// Copyright Epic Games, Inc. All Rights Reserved.
+#include "MassZoneGraphNavigationTrait.h"
+#include "MassEntityTemplateRegistry.h"
+#include "MassCommonFragments.h"
+#include "MassMovementFragments.h"
+#include "MassNavigationFragments.h"
+#include "MassZoneGraphNavigationFragments.h"
+#include "Engine/World.h"
+
+void UMassZoneGraphNavigationTrait::BuildTemplate(FMassEntityTemplateBuildContext& BuildContext, UWorld& World) const
+{
+	UMassEntitySubsystem* EntitySubsystem = UWorld::GetSubsystem<UMassEntitySubsystem>(&World);
+	check(EntitySubsystem);
+
+	BuildContext.AddFragment<FDataFragment_AgentRadius>();
+	BuildContext.AddFragment<FDataFragment_Transform>();
+	BuildContext.AddFragment<FMassVelocityFragment>();
+	BuildContext.AddFragment<FMassMoveTargetFragment>();
+
+	BuildContext.AddFragment<FMassZoneGraphLaneLocationFragment>();
+	BuildContext.AddFragment<FMassZoneGraphPathRequestFragment>();
+	BuildContext.AddFragment<FMassZoneGraphShortPathFragment>();
+	BuildContext.AddFragment<FMassZoneGraphCachedLaneFragment>();
+
+	const FConstSharedStruct ZGMovementParamsFragment = EntitySubsystem->GetOrCreateConstSharedFragment(UE::StructUtils::GetStructCrc32(FConstStructView::Make(NavigationParameters)), NavigationParameters);
+	BuildContext.AddConstSharedFragment(ZGMovementParamsFragment);
+}
diff --git a/Engine/Plugins/AI/MassAI/Source/MassZoneGraphNavigation/Public/MassZoneGraphNavigationProcessors.h b/Engine/Plugins/AI/MassAI/Source/MassZoneGraphNavigation/Public/MassZoneGraphNavigationProcessors.h
new file mode 100644
--- /dev/null
+++ b/Engine/Plugins/AI/MassAI/Source/MassZoneGraphNavigation/Public/MassZoneGraphNavigationProcessors.h
@@ -0,0 +1,81 @@
+// Copyright Epic Games, Inc. All Rights Reserved.
+
+#pragma once
+
+#include "MassProcessor.h"
+#include "MassObserverProcessor.h"
+#include "MassZoneGraphNavigationProcessors.generated.h"
+
+class UMassSignalSubsystem;
+class UZoneGraphSubsystem;
+
+
+
+/**
+ * Processor for initializing nearest location on ZoneGraph.
+ */
+UCLASS()
+class MASSZONEGRAPHNAVIGATION_API UMassZoneGraphLocationInitializer : public UMassFragmentInitializer
+{
+	GENERATED_BODY()
+	
+public:
+	UMassZoneGraphLocationInitializer();
+
+protected:
+	virtual void Initialize(UObject& Owner) override;
+	virtual void ConfigureQueries() override;
+	virtual void Execute(UMassEntitySubsystem& EntitySubsystem, FMassExecutionContext& Context) override;
+
+	FMassEntityQuery EntityQuery;
+
+	UPROPERTY(Transient)
+	UZoneGraphSubsystem* ZoneGraphSubsystem = nullptr;
+
+	UPROPERTY(Transient)
+	UMassSignalSubsystem* SignalSubsystem = nullptr;
+};
+
+/** 
+ * Processor for updating move target on ZoneGraph path.
+ */
+UCLASS()
+class MASSZONEGRAPHNAVIGATION_API UMassZoneGraphPathFollowProcessor : public UMassProcessor
+{
+	GENERATED_BODY()
+
+protected:
+	UMassZoneGraphPathFollowProcessor();
+	
+	virtual void Initialize(UObject& Owner) override;
+	virtual void ConfigureQueries() override;
+	virtual void Execute(UMassEntitySubsystem& EntitySubsystem, FMassExecutionContext& Context) override;
+
+	FMassEntityQuery EntityQuery_Conditional;
+
+	UPROPERTY(Transient)
+	UZoneGraphSubsystem* ZoneGraphSubsystem = nullptr;
+
+	UPROPERTY(Transient)
+	UMassSignalSubsystem* SignalSubsystem = nullptr;
+};
+
+/** ZoneGraph lane cache boundary processor */
+// @todo MassMovement: Make this signal based.
+UCLASS()
+class MASSZONEGRAPHNAVIGATION_API UMassZoneGraphLaneCacheBoundaryProcessor : public UMassProcessor
+{
+	GENERATED_BODY()
+
+public:
+	UMassZoneGraphLaneCacheBoundaryProcessor();
+
+protected:
+	virtual void ConfigureQueries() override;
+	virtual void Initialize(UObject& Owner) override;
+	virtual void Execute(UMassEntitySubsystem& EntitySubsystem, FMassExecutionContext& Context) override;
+
+private:
+	TWeakObjectPtr<UWorld> WeakWorld;
+	FMassEntityQuery EntityQuery;
+};
diff --git a/Engine/Plugins/AI/MassAI/Source/MassZoneGraphNavigation/Public/MassZoneGraphNavigationTrait.h b/Engine/Plugins/AI/MassAI/Source/MassZoneGraphNavigation/Public/MassZoneGraphNavigationTrait.h
new file mode 100644
--- /dev/null
+++ b/Engine/Plugins/AI/MassAI/Source/MassZoneGraphNavigation/Public/MassZoneGraphNavigationTrait.h
@@ -0,0 +1,20 @@
+﻿// Copyright Epic Games, Inc. All Rights Reserved.
+
+#pragma once 
+
+#include "MassEntityTraitBase.h"
+#include "MassZoneGraphNavigationFragments.h"
+#include "MassZoneGraphNavigationTrait.generated.h"
+
+
+UCLASS(meta = (DisplayName = "ZoneGraph Navigation"))
+class MASSZONEGRAPHNAVIGATION_API UMassZoneGraphNavigationTrait : public UMassEntityTraitBase
+{
+	GENERATED_BODY()
+
+protected:
+	virtual void BuildTemplate(FMassEntityTemplateBuildContext& BuildContext, UWorld& World) const override;
+
+	UPROPERTY(Category="Movement", EditAnywhere)
+	FMassZoneGraphNavigationParameters NavigationParameters;
+};
diff --git a/Engine/Plugins/AI/MassCrowd/Source/MassCrowd/Private/MassCrowdNavigationProcessor.cpp b/Engine/Plugins/AI/MassCrowd/Source/MassCrowd/Private/MassCrowdNavigationProcessor.cpp
--- a/Engine/Plugins/AI/MassCrowd/Source/MassCrowd/Private/MassCrowdNavigationProcessor.cpp
+++ b/Engine/Plugins/AI/MassCrowd/Source/MassCrowd/Private/MassCrowdNavigationProcessor.cpp
@@ -1,17 +1,18 @@
 // Copyright Epic Games, Inc. All Rights Reserved.
 
 #include "MassCrowdNavigationProcessor.h"
+#include "MassCommonFragments.h"
 #include "MassAIBehaviorTypes.h"
 #include "MassCrowdSubsystem.h"
 #include "MassCrowdFragments.h"
 #include "MassCrowdSettings.h"
-#include "MassAIMovementFragments.h"
-#include "MassZoneGraphLaneObstacleAnnotations.h"
+#include "MassNavigationFragments.h"
 #include "ZoneGraphSubsystem.h"
 #include "ZoneGraphAnnotationSubsystem.h"
-#include "MassZoneGraphMovementFragments.h"
+#include "MassZoneGraphNavigationFragments.h"
 #include "Annotations/ZoneGraphDisturbanceAnnotation.h"
 #include "VisualLogger/VisualLogger.h"
+#include "MassSimulationLOD.h"
 
 //----------------------------------------------------------------------//
 // UMassCrowdLaneTrackingSignalProcessor
@@ -108,6 +109,7 @@ void UMassCrowdLaneTrackingDestructor::Execute(UMassEntitySubsystem& EntitySubsy
 //----------------------------------------------------------------------//
 // UMassCrowdDynamicObstacleProcessor
 //----------------------------------------------------------------------//
+
 UMassCrowdDynamicObstacleProcessor::UMassCrowdDynamicObstacleProcessor()
 {
 	bAutoRegisterWithProcessingPhases = true;
@@ -118,7 +120,7 @@ UMassCrowdDynamicObstacleProcessor::UMassCrowdDynamicObstacleProcessor()
 void UMassCrowdDynamicObstacleProcessor::Initialize(UObject& Owner)
 {
 	Super::Initialize(Owner);
-	
+
 	ZoneGraphSubsystem = UWorld::GetSubsystem<UZoneGraphSubsystem>(Owner.GetWorld());
 	ZoneGraphAnnotationSubsystem = UWorld::GetSubsystem<UZoneGraphAnnotationSubsystem>(Owner.GetWorld());
 	
@@ -126,57 +128,86 @@ void UMassCrowdDynamicObstacleProcessor::Initialize(UObject& Owner)
 	checkf(CrowdSettings, TEXT("Settings default object is always expected to be valid"));
 }
 
-void UMassCrowdDynamicObstacleProcessor::OnStop(FMassDynamicObstacleFragment& OutObstacle, const float Radius)
+void UMassCrowdDynamicObstacleProcessor::ConfigureQueries()
 {
-	Super::OnStop(OutObstacle, Radius);
+	EntityQuery_Conditional.AddRequirement<FDataFragment_Transform>(EMassFragmentAccess::ReadOnly);
+	EntityQuery_Conditional.AddRequirement<FDataFragment_AgentRadius>(EMassFragmentAccess::ReadOnly);
+	EntityQuery_Conditional.AddRequirement<FMassCrowdObstacleFragment>(EMassFragmentAccess::ReadWrite);
+	EntityQuery_Conditional.AddChunkRequirement<FMassSimulationVariableTickChunkFragment>(EMassFragmentAccess::ReadOnly, EMassFragmentPresence::Optional);
+	EntityQuery_Conditional.SetChunkFilter(&FMassSimulationVariableTickChunkFragment::ShouldTickChunkThisFrame);
+}
 
-	if (ZoneGraphSubsystem == nullptr || ZoneGraphAnnotationSubsystem == nullptr)
+void UMassCrowdDynamicObstacleProcessor::Execute(UMassEntitySubsystem& EntitySubsystem, FMassExecutionContext& Context)
+{
+	UWorld* World = EntitySubsystem.GetWorld();
+
+	EntityQuery_Conditional.ForEachEntityChunk(EntitySubsystem, Context, [this, World](FMassExecutionContext& Context)
 	{
-		return;
-	}
+		const int32 NumEntities = Context.GetNumEntities();
 
-	ensureMsgf(OutObstacle.LaneObstacleIDs.IsEmpty(), TEXT("Blocked lanes have not been cleared."));
+		const TConstArrayView<FDataFragment_Transform> LocationList = Context.GetFragmentView<FDataFragment_Transform>();
+		const TConstArrayView<FDataFragment_AgentRadius> RadiusList = Context.GetFragmentView<FDataFragment_AgentRadius>();
+		const TArrayView<FMassCrowdObstacleFragment> ObstacleDataList = Context.GetMutableFragmentView<FMassCrowdObstacleFragment>();
 
-	// Only send the a disturbance event if we are fully blocking a pedestrian lanes. 
-	const FZoneGraphTag CrowdTag = CrowdSettings->CrowdTag;
-	FZoneGraphTagFilter TagFilter;
-	TagFilter.AnyTags.Add(CrowdTag);
-	TArray<FZoneGraphLaneSection> OutLaneSections;
-	ZoneGraphSubsystem->FindLaneOverlaps(OutObstacle.LastPosition, Radius, TagFilter, OutLaneSections);
-	if (!OutLaneSections.IsEmpty())
-	{
-		// Add an obstacle disturbance.
-		FZoneGraphObstacleDisturbanceArea Disturbance;
-		Disturbance.Position = OutObstacle.LastPosition;
-		constexpr float EffectRadius = 1000.f;
-		Disturbance.Radius = EffectRadius;
-		Disturbance.ObstacleRadius = Radius;
-		Disturbance.ObstacleID = FMassLaneObstacleID::GetNextUniqueID();
-		Disturbance.Action = EZoneGraphObstacleDisturbanceAreaAction::Add;
-		ZoneGraphAnnotationSubsystem->SendEvent(Disturbance);
-
-		// Keep track of our obstacles.
-		OutObstacle.LaneObstacleIDs.Add(Disturbance.ObstacleID);
-	}
-}
+		const float CurrentTime = World->GetTimeSeconds();
 
-void UMassCrowdDynamicObstacleProcessor::OnMove(FMassDynamicObstacleFragment& OutObstacle)
-{
-	Super::OnMove(OutObstacle);
+		const FDateTime Now = FDateTime::UtcNow();
 
-	if (ZoneGraphAnnotationSubsystem == nullptr)
-	{
-		return;
-	}
+		for (int32 EntityIndex = 0; EntityIndex < NumEntities; ++EntityIndex)
+		{
+			// @todo: limit update frequency, this does not need to occur every frame
+			const FVector Position = LocationList[EntityIndex].GetTransform().GetLocation();
+			const float Radius = RadiusList[EntityIndex].Radius;
+			FMassCrowdObstacleFragment& Obstacle = ObstacleDataList[EntityIndex];
 
-	// Unblock previously blocked lanes.
-	for (const FMassLaneObstacleID ID : OutObstacle.LaneObstacleIDs)
-	{
-		FZoneGraphObstacleDisturbanceArea Disturbance;
-		Disturbance.ObstacleID = ID;
-		Disturbance.Action = EZoneGraphObstacleDisturbanceAreaAction::Remove;
-		ZoneGraphAnnotationSubsystem->SendEvent(Disturbance);
-	}
+			UE_VLOG_LOCATION(this, LogMassNavigationObstacle, Display, Position, Radius, Obstacle.bHasStopped ? FColor::Red : FColor::Green, TEXT(""));
 
-	OutObstacle.LaneObstacleIDs.Reset();
+			if ((Position - Obstacle.LastPosition).SquaredLength() < FMath::Square(DistanceBuffer))
+			{
+				const float TimeElapsed = CurrentTime - Obstacle.LastMovedTimeStamp;
+				if (TimeElapsed > DelayBeforeStopNotification && !Obstacle.bHasStopped)
+				{
+					// The obstacle hasn't moved for a while.
+					Obstacle.bHasStopped = true;
+
+					ensureMsgf(Obstacle.LaneObstacleID.IsValid() == false, TEXT("Obstacle should not have been set."));
+
+					Obstacle.LaneObstacleID = FMassLaneObstacleID::GetNextUniqueID();
+
+					constexpr float EffectRadius = 1000.f;
+
+					// Add an obstacle disturbance.
+					FZoneGraphObstacleDisturbanceArea Disturbance;
+					Disturbance.Position = Obstacle.LastPosition;
+					Disturbance.Radius = EffectRadius;
+					Disturbance.ObstacleRadius = Radius;
+					Disturbance.ObstacleID = Obstacle.LaneObstacleID;
+					Disturbance.Action = EZoneGraphObstacleDisturbanceAreaAction::Add;
+					ZoneGraphAnnotationSubsystem->SendEvent(Disturbance);
+				}
+			}
+			else
+			{
+				// Update position and time stamp
+				Obstacle.LastPosition = Position;
+				Obstacle.LastMovedTimeStamp = CurrentTime;
+
+				// If the obstacle had stopped, signal the move.
+				if (Obstacle.bHasStopped)
+				{
+					Obstacle.bHasStopped = false;
+
+					if (ensureMsgf(Obstacle.LaneObstacleID.IsValid(), TEXT("Obstacle should have been set.")))
+					{
+						FZoneGraphObstacleDisturbanceArea Disturbance;
+						Disturbance.ObstacleID = Obstacle.LaneObstacleID;
+						Disturbance.Action = EZoneGraphObstacleDisturbanceAreaAction::Remove;
+						ZoneGraphAnnotationSubsystem->SendEvent(Disturbance);
+
+						Obstacle.LaneObstacleID = FMassLaneObstacleID();
+					}
+				}
+			}
+		}
+	});
 }
diff --git a/Engine/Plugins/AI/MassCrowd/Source/MassCrowd/Private/MassCrowdRepresentationProcessor.cpp b/Engine/Plugins/AI/MassCrowd/Source/MassCrowd/Private/MassCrowdRepresentationProcessor.cpp
--- a/Engine/Plugins/AI/MassCrowd/Source/MassCrowd/Private/MassCrowdRepresentationProcessor.cpp
+++ b/Engine/Plugins/AI/MassCrowd/Source/MassCrowd/Private/MassCrowdRepresentationProcessor.cpp
@@ -2,7 +2,7 @@
 
 #include "MassCrowdRepresentationProcessor.h"
 #include "MassCrowdFragments.h"
-#include "MassAIMovementFragments.h"
+#include "MassMovementFragments.h"
 #include "Components/SkeletalMeshComponent.h"
 #include "Components/CapsuleComponent.h"
 #include "GameFramework/CharacterMovementComponent.h"
diff --git a/Engine/Plugins/AI/MassCrowd/Source/MassCrowd/Private/Tasks/MassCrowdClaimWaitSlotTask.cpp b/Engine/Plugins/AI/MassCrowd/Source/MassCrowd/Private/Tasks/MassCrowdClaimWaitSlotTask.cpp
--- a/Engine/Plugins/AI/MassCrowd/Source/MassCrowd/Private/Tasks/MassCrowdClaimWaitSlotTask.cpp
+++ b/Engine/Plugins/AI/MassCrowd/Source/MassCrowd/Private/Tasks/MassCrowdClaimWaitSlotTask.cpp
@@ -3,8 +3,8 @@
 #include "Tasks/MassCrowdClaimWaitSlotTask.h"
 #include "StateTreeExecutionContext.h"
 #include "MassCrowdSubsystem.h"
-#include "MassZoneGraphMovementFragments.h"
-#include "MassAIMovementFragments.h"
+#include "MassZoneGraphNavigationFragments.h"
+#include "MassNavigationFragments.h"
 #include "MassStateTreeExecutionContext.h"
 
 bool FMassCrowdClaimWaitSlotTask::Link(FStateTreeLinker& Linker)
diff --git a/Engine/Plugins/AI/MassCrowd/Source/MassCrowd/Private/Tasks/MassZoneGraphFindWanderTarget.cpp b/Engine/Plugins/AI/MassCrowd/Source/MassCrowd/Private/Tasks/MassZoneGraphFindWanderTarget.cpp
--- a/Engine/Plugins/AI/MassCrowd/Source/MassCrowd/Private/Tasks/MassZoneGraphFindWanderTarget.cpp
+++ b/Engine/Plugins/AI/MassCrowd/Source/MassCrowd/Private/Tasks/MassZoneGraphFindWanderTarget.cpp
@@ -3,11 +3,10 @@
 #include "Tasks/MassZoneGraphFindWanderTarget.h"
 #include "Templates/Tuple.h"
 #include "StateTreeExecutionContext.h"
-#include "MassStateTreeSubsystem.h"
 #include "ZoneGraphAnnotationSubsystem.h"
 #include "ZoneGraphSubsystem.h"
 #include "ZoneGraphQuery.h"
-#include "MassZoneGraphMovementFragments.h"
+#include "MassZoneGraphNavigationFragments.h"
 #include "MassAIBehaviorTypes.h"
 #include "MassCrowdSettings.h"
 #include "MassCrowdSubsystem.h"
diff --git a/Engine/Plugins/AI/MassCrowd/Source/MassCrowd/Private/ZoneGraphCrowdLaneAnnotations.cpp b/Engine/Plugins/AI/MassCrowd/Source/MassCrowd/Private/ZoneGraphCrowdLaneAnnotations.cpp
--- a/Engine/Plugins/AI/MassCrowd/Source/MassCrowd/Private/ZoneGraphCrowdLaneAnnotations.cpp
+++ b/Engine/Plugins/AI/MassCrowd/Source/MassCrowd/Private/ZoneGraphCrowdLaneAnnotations.cpp
@@ -7,7 +7,7 @@
 #include "ZoneGraphTypes.h"
 #include "ZoneGraphQuery.h"
 #include "ZoneGraphSettings.h"
-#include "MassAIMovementTypes.h"
+#include "MassNavigationTypes.h"
 #include "VisualLogger/VisualLogger.h"
 
 #if !UE_BUILD_SHIPPING && !UE_BUILD_TEST
diff --git a/Engine/Plugins/AI/MassCrowd/Source/MassCrowd/Public/MassCrowdFragments.h b/Engine/Plugins/AI/MassCrowd/Source/MassCrowd/Public/MassCrowdFragments.h
--- a/Engine/Plugins/AI/MassCrowd/Source/MassCrowd/Public/MassCrowdFragments.h
+++ b/Engine/Plugins/AI/MassCrowd/Source/MassCrowd/Public/MassCrowdFragments.h
@@ -4,9 +4,7 @@
 
 #include "MassEntityTypes.h"
 #include "ZoneGraphTypes.h"
-#include "Containers/StaticArray.h"
-#include "MassLODManager.h"
-
+#include "ZoneGraphAnnotationTypes.h"
 #include "MassCrowdFragments.generated.h"
 
 /**
@@ -30,3 +28,20 @@ struct MASSCROWD_API FMassCrowdLaneTrackingFragment : public FMassFragment
 };
 
 
+USTRUCT()
+struct MASSCROWD_API FMassCrowdObstacleFragment : public FMassFragment
+{
+	GENERATED_BODY()
+
+	/** Obstacle ID reported to the obstruction annotation. */
+	FMassLaneObstacleID LaneObstacleID;
+
+	/** Position of the dynamic obstacle when it last moved. */
+	FVector LastPosition = FVector::ZeroVector;
+
+	/** Time stamp when that obstacle stopped moving. */
+	float LastMovedTimeStamp = 0.0f;
+
+	/** Has this dynamic obstacle stopped moving. */
+	bool bHasStopped = true;
+};
diff --git a/Engine/Plugins/AI/MassCrowd/Source/MassCrowd/Public/MassCrowdNavigationProcessor.h b/Engine/Plugins/AI/MassCrowd/Source/MassCrowd/Public/MassCrowdNavigationProcessor.h
--- a/Engine/Plugins/AI/MassCrowd/Source/MassCrowd/Public/MassCrowdNavigationProcessor.h
+++ b/Engine/Plugins/AI/MassCrowd/Source/MassCrowd/Public/MassCrowdNavigationProcessor.h
@@ -4,7 +4,7 @@
 
 #include "MassSignalProcessorBase.h"
 #include "MassObserverProcessor.h"
-#include "NavigationProcessor.h"
+#include "MassCrowdFragments.h"
 #include "MassCrowdNavigationProcessor.generated.h"
 
 class UZoneGraphSubsystem;
@@ -50,7 +50,7 @@ protected:
 };
 
 UCLASS()
-class MASSCROWD_API UMassCrowdDynamicObstacleProcessor : public UMassDynamicObstacleProcessor
+class MASSCROWD_API UMassCrowdDynamicObstacleProcessor : public UMassProcessor
 {
 	GENERATED_BODY()
 
@@ -59,9 +59,21 @@ public:
 
 protected:
 	virtual void Initialize(UObject& Owner) override;
+	virtual void ConfigureQueries() override;
+	virtual void Execute(UMassEntitySubsystem& EntitySubsystem, FMassExecutionContext& Context);
+
+	void OnStop(FMassCrowdObstacleFragment& OutObstacle, const float Radius);
+	void OnMove(FMassCrowdObstacleFragment& OutObstacle);
+
+	/** Delay before sending the stop notification once the entity has stop moving. */
+	UPROPERTY(Category = Settings, EditAnywhere, meta = (ClampMin = "0"))
+	float DelayBeforeStopNotification = 0.3f;
+
+	/** Distance within which the obstacle is considered not moving. */
+	UPROPERTY(Category = Settings, EditAnywhere, meta = (ClampMin = "0"))
+	float DistanceBuffer = 10.f;
 
-	virtual void OnStop(FMassDynamicObstacleFragment& OutObstacle, const float Radius) override;
-	virtual void OnMove(FMassDynamicObstacleFragment& OutObstacle) override;
+	FMassEntityQuery EntityQuery_Conditional;
 
 	UPROPERTY(Transient)
 	UZoneGraphSubsystem* ZoneGraphSubsystem;
diff --git a/Engine/Plugins/Runtime/MassGameplay/MassGameplay.uplugin b/Engine/Plugins/Runtime/MassGameplay/MassGameplay.uplugin
--- a/Engine/Plugins/Runtime/MassGameplay/MassGameplay.uplugin
+++ b/Engine/Plugins/Runtime/MassGameplay/MassGameplay.uplugin
@@ -63,6 +63,11 @@
 		{
 			"Name": "MassGameplayEditor",
 			"Type": "Editor"
+		},
+		{
+			"Name": "MassMovementEditor",
+			"Type": "Editor",
+			"LoadingPhase": "PostEngineInit"
 		}
 	],
 	"Plugins": [
diff --git a/Engine/Plugins/Runtime/MassGameplay/Source/MassActors/Private/Translators/MassCapsuleComponentTranslators.cpp b/Engine/Plugins/Runtime/MassGameplay/Source/MassActors/Private/Translators/MassCapsuleComponentTranslators.cpp
--- a/Engine/Plugins/Runtime/MassGameplay/Source/MassActors/Private/Translators/MassCapsuleComponentTranslators.cpp
+++ b/Engine/Plugins/Runtime/MassGameplay/Source/MassActors/Private/Translators/MassCapsuleComponentTranslators.cpp
@@ -4,7 +4,6 @@
 #include "MassCommonTypes.h"
 #include "MassEntitySubsystem.h"
 #include "Components/CapsuleComponent.h"
-#include "MassMovementTypes.h"
 
 
 //----------------------------------------------------------------------//
@@ -48,7 +47,7 @@ UMassTransformToActorCapsuleTranslator::UMassTransformToActorCapsuleTranslator()
 	ExecutionFlags = (int32)EProcessorExecutionFlags::All;
 	RequiredTags.Add<FMassCapsuleTransformCopyToActorTag>();
 	ExecutionOrder.ExecuteInGroup = UE::Mass::ProcessorGroupNames::UpdateWorldFromMass;
-	ExecutionOrder.ExecuteAfter.Add(UE::Mass::ProcessorGroupNames::Avoidance);
+	ExecutionOrder.ExecuteAfter.Add(UE::Mass::ProcessorGroupNames::Movement);
 }
 
 void UMassTransformToActorCapsuleTranslator::ConfigureQueries()
diff --git a/Engine/Plugins/Runtime/MassGameplay/Source/MassActors/Private/Translators/MassCharacterMovementTranslators.cpp b/Engine/Plugins/Runtime/MassGameplay/Source/MassActors/Private/Translators/MassCharacterMovementTranslators.cpp
--- a/Engine/Plugins/Runtime/MassGameplay/Source/MassActors/Private/Translators/MassCharacterMovementTranslators.cpp
+++ b/Engine/Plugins/Runtime/MassGameplay/Source/MassActors/Private/Translators/MassCharacterMovementTranslators.cpp
@@ -6,7 +6,6 @@
 #include "VisualLogger/VisualLogger.h"
 #include "MassEntitySubsystem.h"
 #include "MassCommonTypes.h"
-#include "MassMovementTypes.h"
 #include "MassMovementFragments.h"
 #include "MassSpawnerTypes.h"
 
@@ -57,7 +56,7 @@ UMassCharacterMovementToActorTranslator::UMassCharacterMovementToActorTranslator
 {
 	ExecutionFlags = (int32)EProcessorExecutionFlags::All;
 	ExecutionOrder.ExecuteInGroup = UE::Mass::ProcessorGroupNames::UpdateWorldFromMass;
-	ExecutionOrder.ExecuteAfter.Add(UE::Mass::ProcessorGroupNames::Avoidance);
+	ExecutionOrder.ExecuteAfter.Add(UE::Mass::ProcessorGroupNames::Movement);
 	RequiredTags.Add<FMassCharacterMovementCopyToActorTag>();
 }
 
@@ -133,7 +132,7 @@ UMassCharacterOrientationToActorTranslator::UMassCharacterOrientationToActorTran
 {
 	ExecutionFlags = (int32)EProcessorExecutionFlags::All;
 	ExecutionOrder.ExecuteInGroup = UE::Mass::ProcessorGroupNames::UpdateWorldFromMass;
-	ExecutionOrder.ExecuteAfter.Add(UE::Mass::ProcessorGroupNames::Avoidance);
+	ExecutionOrder.ExecuteAfter.Add(UE::Mass::ProcessorGroupNames::Movement);
 	RequiredTags.Add<FMassCharacterOrientationCopyToActorTag>();
 }
 
diff --git a/Engine/Plugins/Runtime/MassGameplay/Source/MassActors/Private/Translators/MassSceneComponentLocationTranslator.cpp b/Engine/Plugins/Runtime/MassGameplay/Source/MassActors/Private/Translators/MassSceneComponentLocationTranslator.cpp
--- a/Engine/Plugins/Runtime/MassGameplay/Source/MassActors/Private/Translators/MassSceneComponentLocationTranslator.cpp
+++ b/Engine/Plugins/Runtime/MassGameplay/Source/MassActors/Private/Translators/MassSceneComponentLocationTranslator.cpp
@@ -1,7 +1,7 @@
 // Copyright Epic Games, Inc. All Rights Reserved.
 
 #include "Translators/MassSceneComponentLocationTranslator.h"
-#include "MassMovementTypes.h"
+#include "MassCommonTypes.h"
 #include "Components/SceneComponent.h"
 #include "MassEntitySubsystem.h"
 #include "MassCommonTypes.h"
@@ -48,7 +48,7 @@ UMassSceneComponentLocationToActorTranslator::UMassSceneComponentLocationToActor
 {
 	ExecutionFlags = (int32)EProcessorExecutionFlags::All;
 	ExecutionOrder.ExecuteInGroup = UE::Mass::ProcessorGroupNames::UpdateWorldFromMass;
-	ExecutionOrder.ExecuteAfter.Add(UE::Mass::ProcessorGroupNames::Avoidance);
+	ExecutionOrder.ExecuteAfter.Add(UE::Mass::ProcessorGroupNames::Movement);
 	RequiredTags.Add<FMassSceneComponentLocationCopyToActorTag>();
 }
 
diff --git a/Engine/Plugins/Runtime/MassGameplay/Source/MassActors/Public/Translators/MassCharacterMovementTranslators.h b/Engine/Plugins/Runtime/MassGameplay/Source/MassActors/Public/Translators/MassCharacterMovementTranslators.h
--- a/Engine/Plugins/Runtime/MassGameplay/Source/MassActors/Public/Translators/MassCharacterMovementTranslators.h
+++ b/Engine/Plugins/Runtime/MassGameplay/Source/MassActors/Public/Translators/MassCharacterMovementTranslators.h
@@ -8,10 +8,6 @@
 #include "MassCharacterMovementTranslators.generated.h"
 
 class UCharacterMovementComponent;
-struct FDataFragment_CharacterMovementComponentWrapper;
-struct FNavLocationComponent;
-struct FDataFragment_Transform;
-struct FMassVelocityFragment;
 
 USTRUCT()
 struct FDataFragment_CharacterMovementComponentWrapper : public FDataFragment_ObjectWrapper
diff --git a/Engine/Plugins/Runtime/MassGameplay/Source/MassCommon/Public/MassCommonTypes.h b/Engine/Plugins/Runtime/MassGameplay/Source/MassCommon/Public/MassCommonTypes.h
--- a/Engine/Plugins/Runtime/MassGameplay/Source/MassCommon/Public/MassCommonTypes.h
+++ b/Engine/Plugins/Runtime/MassGameplay/Source/MassCommon/Public/MassCommonTypes.h
@@ -35,6 +35,8 @@ namespace UE::Mass::ProcessorGroupNames
 	const FName SyncWorldToMass = FName(TEXT("SyncWorldToMass"));
 	const FName Behavior = FName(TEXT("Behavior"));
 	const FName Tasks = FName(TEXT("Tasks"));
+	const FName Avoidance = FName(TEXT("Avoidance"));
+	const FName Movement = FName(TEXT("Movement"));
 }
 
 USTRUCT()
diff --git a/Engine/Plugins/Runtime/MassGameplay/Source/MassMovement/Private/MassMovementFragments.cpp b/Engine/Plugins/Runtime/MassGameplay/Source/MassMovement/Private/MassMovementFragments.cpp
--- a/Engine/Plugins/Runtime/MassGameplay/Source/MassMovement/Private/MassMovementFragments.cpp
+++ b/Engine/Plugins/Runtime/MassGameplay/Source/MassMovement/Private/MassMovementFragments.cpp
@@ -1,36 +1,42 @@
 // Copyright Epic Games, Inc. All Rights Reserved.
 
 #include "MassMovementFragments.h"
-#include "GameFramework/GameStateBase.h"
-#include "Engine/World.h"
 
-//void FMassMoveTargetFragment::CreateNewAction(const EMassMovementAction InAction, const UWorld& InWorld)
-//{
-//	ensureMsgf(InWorld.GetNetMode() != NM_Client, TEXT("This version of SetDesiredAction should only be called on the authority."));
-//	CurrentAction = InAction;
-//	CurrentActionID++;
-//	CurrentActionWorldStartTime = InWorld.TimeSeconds;
-//
-//	const AGameStateBase* GameState = InWorld.GetGameState();
-//	CurrentActionServerStartTime = GameState != nullptr ? GameState->GetServerWorldTimeSeconds() : CurrentActionWorldStartTime;
-//
-//	MarkNetDirty();
-//}
-//
-//void FMassMoveTargetFragment::CreateReplicatedAction(const EMassMovementAction InAction, const uint16 InActionID, const float InWorldStartTime, const float InServerStartTime)
-//{
-//	CurrentAction = InAction;
-//	CurrentActionID = InActionID;
-//	CurrentActionWorldStartTime = InWorldStartTime;
-//	CurrentActionServerStartTime = InServerStartTime;
-//}
-//
-//FString FMassMoveTargetFragment::ToString() const
-//{
-//	return FString::Printf(TEXT("ActionID:%d Action:%s StartTime: World:%.1f Server:%.1f DesiredSpeed:%.1f "),
-//		CurrentActionID,
-//		*UEnum::GetValueAsString(CurrentAction),
-//		CurrentActionWorldStartTime,
-//		CurrentActionServerStartTime,
-//		DesiredSpeed.Get());
-//}
\ No newline at end of file
+void FMassMovementParameters::Update()
+{
+	for (FMassMovementStyleParameters& Style : MovementStyles)
+	{
+		if (Style.DesiredSpeeds.IsEmpty())
+		{
+			continue;
+		}
+
+		// Calculate probability threshold for the speeds, so that a speed can be looked up based in a float in range [0...1]. 
+		float Total = 0.0f;
+		for (const FMassMovementStyleSpeedParameters& Speed : Style.DesiredSpeeds)
+		{
+			Total += Speed.Probability;
+		}
+
+		if (Total > KINDA_SMALL_NUMBER)
+		{
+			const float Scale = 1.0f / Total;
+			float Sum = 0.0f;
+			for (FMassMovementStyleSpeedParameters& Speed : Style.DesiredSpeeds)
+			{
+				Sum += Speed.Probability;
+				Speed.ProbabilityThreshold = FMath::Min(Sum * Scale, 1.0f);
+			}
+		}
+		else
+		{
+			const float Scale = 1.0f / Style.DesiredSpeeds.Num();
+			float Sum = 0.0f;
+			for (FMassMovementStyleSpeedParameters& Speed : Style.DesiredSpeeds)
+			{
+				Sum += 1.0f;
+				Speed.ProbabilityThreshold = FMath::Min(Sum * Scale, 1.0f);
+			}
+		}
+	}
+}
diff --git a/Engine/Plugins/Runtime/MassGameplay/Source/MassMovement/Private/MassMovementProcessors.cpp b/Engine/Plugins/Runtime/MassGameplay/Source/MassMovement/Private/MassMovementProcessors.cpp
new file mode 100644
--- /dev/null
+++ b/Engine/Plugins/Runtime/MassGameplay/Source/MassMovement/Private/MassMovementProcessors.cpp
@@ -0,0 +1,76 @@
+﻿// Copyright Epic Games, Inc. All Rights Reserved.
+
+#include "Movement/MassMovementProcessors.h"
+#include "MassCommonUtils.h"
+#include "MassCommandBuffer.h"
+#include "MassCommonFragments.h"
+#include "MassMovementFragments.h"
+#include "Math/UnrealMathUtility.h"
+#include "MassSimulationLOD.h"
+
+#define UNSAFE_FOR_MT 0
+#define MOVEMENT_DEBUGDRAW 0	// Set to 1 to see heading debugdraw
+
+//----------------------------------------------------------------------//
+//  UMassApplyMovementProcessor
+//----------------------------------------------------------------------//
+
+UMassApplyMovementProcessor::UMassApplyMovementProcessor()
+{
+	ExecutionFlags = (int32)EProcessorExecutionFlags::All;
+	ExecutionOrder.ExecuteInGroup = UE::Mass::ProcessorGroupNames::Movement;
+	ExecutionOrder.ExecuteAfter.Add(UE::Mass::ProcessorGroupNames::Avoidance);
+}
+
+void UMassApplyMovementProcessor::ConfigureQueries()
+{
+	EntityQuery.AddRequirement<FMassVelocityFragment>(EMassFragmentAccess::ReadWrite);
+	EntityQuery.AddRequirement<FDataFragment_Transform>(EMassFragmentAccess::ReadWrite);
+	EntityQuery.AddRequirement<FMassForceFragment>(EMassFragmentAccess::ReadWrite);
+	EntityQuery.AddTagRequirement<FMassOffLODTag>(EMassFragmentPresence::None);
+	EntityQuery.AddConstSharedRequirement<FMassMovementParameters>(EMassFragmentPresence::All);
+}
+
+void UMassApplyMovementProcessor::Execute(UMassEntitySubsystem& EntitySubsystem,
+													FMassExecutionContext& Context)
+{
+	// Clamp max delta time to avoid force explosion on large time steps (i.e. during initialization).
+	const float DeltaTime = FMath::Min(0.1f, Context.GetDeltaTimeSeconds());
+
+	QUICK_SCOPE_CYCLE_COUNTER(HighRes);
+
+	EntityQuery.ForEachEntityChunk(EntitySubsystem, Context, [this, DeltaTime](FMassExecutionContext& Context)
+	{
+		const int32 NumEntities = Context.GetNumEntities();
+
+		const FMassMovementParameters& MovementParams = Context.GetConstSharedFragment<FMassMovementParameters>();
+
+		const TArrayView<FDataFragment_Transform> LocationList = Context.GetMutableFragmentView<FDataFragment_Transform>();
+		const TArrayView<FMassForceFragment> ForceList = Context.GetMutableFragmentView<FMassForceFragment>();
+		const TArrayView<FMassVelocityFragment> VelocityList = Context.GetMutableFragmentView<FMassVelocityFragment>();
+
+		for (int32 EntityIndex = 0; EntityIndex < NumEntities; ++EntityIndex)
+		{
+			FMassForceFragment& Force = ForceList[EntityIndex];
+			FMassVelocityFragment& Velocity = VelocityList[EntityIndex];
+			FTransform& CurrentTransform = LocationList[EntityIndex].GetMutableTransform();
+
+			// Update velocity from steering forces.
+			Velocity.Value += Force.Value * DeltaTime;
+
+#if WITH_MASSGAMEPLAY_DEBUG
+			if (UE::MassMovement::bFreezeMovement)
+			{
+				Velocity.Value = FVector::ZeroVector;
+			}
+#endif // WITH_MASSGAMEPLAY_DEBUG
+
+			FVector CurrentLocation = CurrentTransform.GetLocation();
+			CurrentLocation += Velocity.Value * DeltaTime;
+			CurrentTransform.SetTranslation(CurrentLocation);
+
+		}
+	});
+}
+
+#undef UNSAFE_FOR_MT
diff --git a/Engine/Plugins/Runtime/MassGameplay/Source/MassMovement/Private/MassMovementTrait.cpp b/Engine/Plugins/Runtime/MassGameplay/Source/MassMovement/Private/MassMovementTrait.cpp
new file mode 100644
--- /dev/null
+++ b/Engine/Plugins/Runtime/MassGameplay/Source/MassMovement/Private/MassMovementTrait.cpp
@@ -0,0 +1,22 @@
+﻿// Copyright Epic Games, Inc. All Rights Reserved.
+#include "Movement/MassMovementTrait.h"
+#include "MassEntityTemplateRegistry.h"
+#include "MassCommonFragments.h"
+#include "MassMovementFragments.h"
+#include "MassMovementTypes.h"
+#include "Engine/World.h"
+
+void UMassMovementTrait::BuildTemplate(FMassEntityTemplateBuildContext& BuildContext, UWorld& World) const
+{
+	UMassEntitySubsystem* EntitySubsystem = UWorld::GetSubsystem<UMassEntitySubsystem>(&World);
+	check(EntitySubsystem);
+
+	BuildContext.AddFragment<FDataFragment_AgentRadius>();
+	BuildContext.AddFragment<FDataFragment_Transform>();
+
+	BuildContext.AddFragment<FMassVelocityFragment>();
+	BuildContext.AddFragment<FMassForceFragment>();
+
+	const FConstSharedStruct MovementFragment = EntitySubsystem->GetOrCreateConstSharedFragment(UE::StructUtils::GetStructCrc32(FConstStructView::Make(Movement)), Movement);
+	BuildContext.AddConstSharedFragment(MovementFragment);
+}
diff --git a/Engine/Plugins/Runtime/MassGameplay/Source/MassMovement/Private/MassMovementTypes.cpp b/Engine/Plugins/Runtime/MassGameplay/Source/MassMovement/Private/MassMovementTypes.cpp
new file mode 100644
--- /dev/null
+++ b/Engine/Plugins/Runtime/MassGameplay/Source/MassMovement/Private/MassMovementTypes.cpp
@@ -0,0 +1,11 @@
+﻿// Copyright Epic Games, Inc. All Rights Reserved.
+
+#include "MassMovementTypes.h"
+#include "HAL/IConsoleManager.h"
+
+namespace UE::MassMovement
+{
+	int32 bFreezeMovement = 0;
+	FAutoConsoleVariableRef CVarFreezeMovement(TEXT("ai.debug.mass.FreezeMovement"), bFreezeMovement, TEXT("Freeze any movement by common movement processors."));
+
+} // UE::MassMovement
diff --git a/Engine/Plugins/Runtime/MassGameplay/Source/MassMovement/Private/MassSimpleMovementTrait.cpp b/Engine/Plugins/Runtime/MassGameplay/Source/MassMovement/Private/MassSimpleMovementTrait.cpp
--- a/Engine/Plugins/Runtime/MassGameplay/Source/MassMovement/Private/MassSimpleMovementTrait.cpp
+++ b/Engine/Plugins/Runtime/MassGameplay/Source/MassMovement/Private/MassSimpleMovementTrait.cpp
@@ -1,9 +1,8 @@
 // Copyright Epic Games, Inc. All Rights Reserved.
 
-#include "MassSimpleMovementTrait.h"
+#include "Example/MassSimpleMovementTrait.h"
 #include "MassEntityTemplateRegistry.h"
 #include "Engine/World.h"
-#include "MassMovementTypes.h"
 #include "MassMovementFragments.h"
 #include "MassCommonFragments.h"
 #include "MassSimulationLOD.h"
diff --git a/Engine/Plugins/Runtime/MassGameplay/Source/MassMovement/Private/MassVelocityRandomizerTrait.cpp b/Engine/Plugins/Runtime/MassGameplay/Source/MassMovement/Private/MassVelocityRandomizerTrait.cpp
--- a/Engine/Plugins/Runtime/MassGameplay/Source/MassMovement/Private/MassVelocityRandomizerTrait.cpp
+++ b/Engine/Plugins/Runtime/MassGameplay/Source/MassMovement/Private/MassVelocityRandomizerTrait.cpp
@@ -1,6 +1,6 @@
 // Copyright Epic Games, Inc. All Rights Reserved.
 
-#include "MassVelocityRandomizerTrait.h"
+#include "Example/MassVelocityRandomizerTrait.h"
 #include "MassEntityTemplateRegistry.h"
 #include "Engine/World.h"
 #include "MassMovementFragments.h"
diff --git a/Engine/Plugins/Runtime/MassGameplay/Source/MassMovement/Public/MassMovementFragments.h b/Engine/Plugins/Runtime/MassGameplay/Source/MassMovement/Public/MassMovementFragments.h
--- a/Engine/Plugins/Runtime/MassGameplay/Source/MassMovement/Public/MassMovementFragments.h
+++ b/Engine/Plugins/Runtime/MassGameplay/Source/MassMovement/Public/MassMovementFragments.h
@@ -4,7 +4,8 @@
 
 #include "CoreMinimal.h"
 #include "MassCommonTypes.h"
-
+#include "MassMovementTypes.h"
+#include "RandomSequence.h"
 #include "MassMovementFragments.generated.h"
 
 USTRUCT()
@@ -14,3 +15,79 @@ struct MASSMOVEMENT_API FMassVelocityFragment : public FMassFragment
 
 	FVector Value = FVector::ZeroVector;
 };
+
+USTRUCT()
+struct MASSMOVEMENT_API FMassForceFragment : public FMassFragment
+{
+	GENERATED_BODY()
+
+	FVector Value = FVector::ZeroVector;
+};
+
+USTRUCT()
+struct MASSMOVEMENT_API FMassMovementParameters : public FMassSharedFragment
+{
+	GENERATED_BODY()
+
+	FMassMovementParameters GetValidated() const
+	{
+		FMassMovementParameters Copy = *this;
+		Copy.Update();
+		return Copy;
+	}
+
+	/** Updates internal values for faster desired speed generation. */
+	void Update();
+
+	/** Generates desired speed based on style and unique id. The id is used deterministically assign a specific speed range. */
+	float GenerateDesiredSpeed(const FMassMovementStyleRef& Style, const int32 UniqueId) const
+	{
+		float DesiredSpeed = DefaultDesiredSpeed;
+		float DesiredSpeedVariance = DefaultDesiredSpeedVariance;
+		
+		const FMassMovementStyleParameters* StyleParams = MovementStyles.FindByPredicate([&Style](const FMassMovementStyleParameters& Config)
+			{
+				return Config.Style.ID == Style.ID;
+			});
+		
+		if (StyleParams != nullptr)
+		{
+			const float Prob = UE::RandomSequence::FRand(UniqueId);
+			for (const FMassMovementStyleSpeedParameters& Speed : StyleParams->DesiredSpeeds)
+			{
+				if (Prob < Speed.ProbabilityThreshold)
+				{
+					DesiredSpeed = Speed.Speed;
+					DesiredSpeedVariance = Speed.Variance;
+					break;
+				}
+			}
+		}
+		
+		return DesiredSpeed * FMath::RandRange(1.0f - DesiredSpeedVariance, 1.0f + DesiredSpeedVariance);
+	}
+	
+	/** Maximum speed (cm/s). */
+	UPROPERTY(EditAnywhere, Category = "Movement", meta = (ClampMin = "0", ForceUnits="cm/s"))
+	float MaxSpeed = 200.f;
+
+	/** 200..600 Smaller steering maximum acceleration makes the steering more \"calm\" but less opportunistic, may not find solution, or gets stuck. */
+	UPROPERTY(config, EditAnywhere, Category = "Movement", meta = (UIMin = 0.0, ClampMin = 0.0, ForceUnits="cm/s^2"))
+	float MaxAcceleration = 250.f;
+
+	/** Default desired speed (cm/s). */
+	UPROPERTY(EditAnywhere, Category = "Movement", meta = (ClampMin = "0", ForceUnits="cm/s"))
+	float DefaultDesiredSpeed = 140.f;
+
+	/** How much default desired speed is varied randomly. */
+	UPROPERTY(EditAnywhere, Category = "Movement", meta = (ClampMin = "0", ClampMax = "1"))
+	float DefaultDesiredSpeedVariance = 0.1f;
+
+	/** The time it takes the entity position to catchup to the requested height. */
+	UPROPERTY(EditAnywhere, Category = "Movement", meta = (ClampMin = "0.0", ForceUnits="s"))
+	float HeightSmoothingTime = 0.2f;
+
+	/** List of supported movement styles for this configuration. */
+	UPROPERTY(EditAnywhere, Category = "Movement")
+	TArray<FMassMovementStyleParameters> MovementStyles;
+};
