Optimized how we set the HasColliderData when updating the avoidance obstacle processor.

#rnx

diff --git a/Engine/Plugins/AI/MassAI/Source/MassNavigation/Private/MassNavigationProcessors.cpp b/Engine/Plugins/AI/MassAI/Source/MassNavigation/Private/MassNavigationProcessors.cpp
--- a/Engine/Plugins/AI/MassAI/Source/MassNavigation/Private/MassNavigationProcessors.cpp
+++ b/Engine/Plugins/AI/MassAI/Source/MassNavigation/Private/MassNavigationProcessors.cpp
@@ -164,18 +164,22 @@ UMassNavigationObstacleGridProcessor::UMassNavigationObstacleGridProcessor()
 
 void UMassNavigationObstacleGridProcessor::ConfigureQueries()
 {
-	AddToGridEntityQuery.AddRequirement<FDataFragment_Transform>(EMassFragmentAccess::ReadOnly);
-	AddToGridEntityQuery.AddRequirement<FDataFragment_AgentRadius>(EMassFragmentAccess::ReadOnly);
-	AddToGridEntityQuery.AddRequirement<FMassNavigationObstacleGridCellLocationFragment>(EMassFragmentAccess::ReadWrite);
-	UpdateGridEntityQuery = AddToGridEntityQuery;
-	RemoveFromGridEntityQuery = AddToGridEntityQuery;
+	FMassEntityQuery BaseEntityQuery;
+	BaseEntityQuery.AddRequirement<FDataFragment_Transform>(EMassFragmentAccess::ReadOnly);
+	BaseEntityQuery.AddRequirement<FDataFragment_AgentRadius>(EMassFragmentAccess::ReadOnly);
+	BaseEntityQuery.AddRequirement<FMassNavigationObstacleGridCellLocationFragment>(EMassFragmentAccess::ReadWrite);
 
+	AddToGridEntityQuery = BaseEntityQuery;
+	AddToGridEntityQuery.AddRequirement<FMassAvoidanceColliderFragment>(EMassFragmentAccess::ReadOnly, EMassFragmentPresence::Optional);
 	AddToGridEntityQuery.AddTagRequirement<FMassOffLODTag>(EMassFragmentPresence::None);
 	AddToGridEntityQuery.AddTagRequirement<FMassInNavigationObstacleGridTag>(EMassFragmentPresence::None);
 
+	UpdateGridEntityQuery = BaseEntityQuery;
+	UpdateGridEntityQuery.AddRequirement<FMassAvoidanceColliderFragment>(EMassFragmentAccess::ReadOnly, EMassFragmentPresence::Optional);
 	UpdateGridEntityQuery.AddTagRequirement<FMassOffLODTag>(EMassFragmentPresence::None);
 	UpdateGridEntityQuery.AddTagRequirement<FMassInNavigationObstacleGridTag>(EMassFragmentPresence::All);
 
+	RemoveFromGridEntityQuery = BaseEntityQuery;
 	RemoveFromGridEntityQuery.AddTagRequirement<FMassOffLODTag>(EMassFragmentPresence::All);
 	RemoveFromGridEntityQuery.AddTagRequirement<FMassInNavigationObstacleGridTag>(EMassFragmentPresence::All);
 }
@@ -202,6 +206,7 @@ void UMassNavigationObstacleGridProcessor::Execute(UMassEntitySubsystem& EntityS
 		TConstArrayView<FDataFragment_Transform> LocationList = Context.GetFragmentView<FDataFragment_Transform>();
 		TConstArrayView<FDataFragment_AgentRadius> RadiiList = Context.GetFragmentView<FDataFragment_AgentRadius>();
 		TArrayView<FMassNavigationObstacleGridCellLocationFragment> NavigationObstacleCellLocationList = Context.GetMutableFragmentView<FMassNavigationObstacleGridCellLocationFragment>();
+		const bool bHasColliderData = Context.GetFragmentView<FMassAvoidanceColliderFragment>().Num() > 0;
 
 		for (int32 EntityIndex = 0; EntityIndex < NumEntities; ++EntityIndex)
 		{
@@ -211,12 +216,7 @@ void UMassNavigationObstacleGridProcessor::Execute(UMassEntitySubsystem& EntityS
 
 			FMassNavigationObstacleItem ObstacleItem;
 			ObstacleItem.Entity = Context.GetEntity(EntityIndex);
-			FMassEntityView EntityView(EntitySubsystem, ObstacleItem.Entity);
-			const FMassAvoidanceColliderFragment* Collider = EntityView.GetFragmentDataPtr<FMassAvoidanceColliderFragment>();
-			if (Collider)
-			{
-				ObstacleItem.ItemFlags |= EMassNavigationObstacleFlags::HasColliderData;
-			}
+			ObstacleItem.ItemFlags |= bHasColliderData ? EMassNavigationObstacleFlags::HasColliderData : EMassNavigationObstacleFlags::None;
 			
 			const FBox NewBounds(NewPos - FVector(Radius, Radius, 0.f), NewPos + FVector(Radius, Radius, 0.f));
 			NavigationObstacleCellLocationList[EntityIndex].CellLoc = NavigationSubsystem->GetObstacleGridMutable().Add(ObstacleItem, NewBounds);
@@ -232,6 +232,7 @@ void UMassNavigationObstacleGridProcessor::Execute(UMassEntitySubsystem& EntityS
 		TConstArrayView<FDataFragment_Transform> LocationList = Context.GetFragmentView<FDataFragment_Transform>();
 		TConstArrayView<FDataFragment_AgentRadius> RadiiList = Context.GetFragmentView<FDataFragment_AgentRadius>();
 		TArrayView<FMassNavigationObstacleGridCellLocationFragment> NavigationObstacleCellLocationList = Context.GetMutableFragmentView<FMassNavigationObstacleGridCellLocationFragment>();
+		const bool bHasColliderData = Context.GetFragmentView<FMassAvoidanceColliderFragment>().Num() > 0;
 
 		for (int32 EntityIndex = 0; EntityIndex < NumEntities; ++EntityIndex)
 		{
@@ -240,12 +241,8 @@ void UMassNavigationObstacleGridProcessor::Execute(UMassEntitySubsystem& EntityS
 			const float Radius = RadiiList[EntityIndex].Radius;
 			FMassNavigationObstacleItem ObstacleItem;
 			ObstacleItem.Entity = Context.GetEntity(EntityIndex);
-			FMassEntityView EntityView(EntitySubsystem, ObstacleItem.Entity);
-			const FMassAvoidanceColliderFragment* Collider = EntityView.GetFragmentDataPtr<FMassAvoidanceColliderFragment>();
-			if (Collider)
-			{
-				ObstacleItem.ItemFlags |= EMassNavigationObstacleFlags::HasColliderData;
-			}
+			ObstacleItem.ItemFlags |= bHasColliderData ? EMassNavigationObstacleFlags::HasColliderData : EMassNavigationObstacleFlags::None;
+
 			const FBox NewBounds(NewPos - FVector(Radius, Radius, 0.f), NewPos + FVector(Radius, Radius, 0.f));
 			NavigationObstacleCellLocationList[EntityIndex].CellLoc = NavigationSubsystem->GetObstacleGridMutable().Move(ObstacleItem, NavigationObstacleCellLocationList[EntityIndex].CellLoc, NewBounds);
 
