Switched all fragment de/initializer processors over to derive from the common parrent, MassObserverProcessor

diff --git a/Engine/Plugins/AI/MassAI/Source/MassAIBehavior/Private/MassStateTreeProcessors.cpp b/Engine/Plugins/AI/MassAI/Source/MassAIBehavior/Private/MassStateTreeProcessors.cpp
--- a/Engine/Plugins/AI/MassAI/Source/MassAIBehavior/Private/MassStateTreeProcessors.cpp
+++ b/Engine/Plugins/AI/MassAI/Source/MassAIBehavior/Private/MassStateTreeProcessors.cpp
@@ -170,7 +170,8 @@ void ForEachEntityInChunk(
 UMassStateTreeFragmentDestructor::UMassStateTreeFragmentDestructor()
 {
 	ExecutionFlags = (int32)(EProcessorExecutionFlags::Standalone | EProcessorExecutionFlags::Server);
-	FragmentType = FMassStateTreeFragment::StaticStruct();
+	ObservedType = FMassStateTreeFragment::StaticStruct();
+	Operation = EMassObservedOperation::Remove;
 }
 
 void UMassStateTreeFragmentDestructor::Initialize(UObject& Owner)
diff --git a/Engine/Plugins/AI/MassAI/Source/MassAIBehavior/Private/MassZoneGraphAnnotationProcessors.cpp b/Engine/Plugins/AI/MassAI/Source/MassAIBehavior/Private/MassZoneGraphAnnotationProcessors.cpp
--- a/Engine/Plugins/AI/MassAI/Source/MassAIBehavior/Private/MassZoneGraphAnnotationProcessors.cpp
+++ b/Engine/Plugins/AI/MassAI/Source/MassAIBehavior/Private/MassZoneGraphAnnotationProcessors.cpp
@@ -16,7 +16,8 @@
 //----------------------------------------------------------------------//
 UMassZoneGraphAnnotationTagsInitializer::UMassZoneGraphAnnotationTagsInitializer()
 {
-	FragmentType = FMassZoneGraphAnnotationTagsFragment::StaticStruct();
+	ObservedType = FMassZoneGraphAnnotationTagsFragment::StaticStruct();
+	Operation = EMassObservedOperation::Add;
 }
 
 void UMassZoneGraphAnnotationTagsInitializer::Initialize(UObject& Owner)
diff --git a/Engine/Plugins/AI/MassAI/Source/MassAIBehavior/Public/MassStateTreeProcessors.h b/Engine/Plugins/AI/MassAI/Source/MassAIBehavior/Public/MassStateTreeProcessors.h
--- a/Engine/Plugins/AI/MassAI/Source/MassAIBehavior/Public/MassStateTreeProcessors.h
+++ b/Engine/Plugins/AI/MassAI/Source/MassAIBehavior/Public/MassStateTreeProcessors.h
@@ -14,7 +14,7 @@ struct FMassStateTreeExecutionContext;
  * Processor to stop and uninitialize StateTrees on entities.
  */
 UCLASS()
-class MASSAIBEHAVIOR_API UMassStateTreeFragmentDestructor : public UMassFragmentDeinitializer
+class MASSAIBEHAVIOR_API UMassStateTreeFragmentDestructor : public UMassObserverProcessor
 {
 	GENERATED_BODY()
 
diff --git a/Engine/Plugins/AI/MassAI/Source/MassAIBehavior/Public/MassZoneGraphAnnotationProcessors.h b/Engine/Plugins/AI/MassAI/Source/MassAIBehavior/Public/MassZoneGraphAnnotationProcessors.h
--- a/Engine/Plugins/AI/MassAI/Source/MassAIBehavior/Public/MassZoneGraphAnnotationProcessors.h
+++ b/Engine/Plugins/AI/MassAI/Source/MassAIBehavior/Public/MassZoneGraphAnnotationProcessors.h
@@ -16,7 +16,7 @@ struct FMassEntityHandle;
  * Processor for initializing ZoneGraph annotation tags.
  */
 UCLASS()
-class MASSAIBEHAVIOR_API UMassZoneGraphAnnotationTagsInitializer : public UMassFragmentInitializer
+class MASSAIBEHAVIOR_API UMassZoneGraphAnnotationTagsInitializer : public UMassObserverProcessor
 {
 	GENERATED_BODY()
 
diff --git a/Engine/Plugins/AI/MassAI/Source/MassNavigation/Private/MassNavigationProcessors.cpp b/Engine/Plugins/AI/MassAI/Source/MassNavigation/Private/MassNavigationProcessors.cpp
--- a/Engine/Plugins/AI/MassAI/Source/MassNavigation/Private/MassNavigationProcessors.cpp
+++ b/Engine/Plugins/AI/MassAI/Source/MassNavigation/Private/MassNavigationProcessors.cpp
@@ -122,7 +122,8 @@ void UMassNavigationSmoothHeightProcessor::Execute(UMassEntitySubsystem& EntityS
 
 UMassMoveTargetFragmentInitializer::UMassMoveTargetFragmentInitializer()
 {
-	FragmentType = FMassMoveTargetFragment::StaticStruct();
+	ObservedType = FMassMoveTargetFragment::StaticStruct();
+	Operation = EMassObservedOperation::Add;
 }
 
 void UMassMoveTargetFragmentInitializer::ConfigureQueries()
@@ -280,7 +281,8 @@ void UMassNavigationObstacleGridProcessor::Execute(UMassEntitySubsystem& EntityS
 //----------------------------------------------------------------------//
 UMassNavigationObstacleRemoverProcessor::UMassNavigationObstacleRemoverProcessor()
 {
-	FragmentType = FMassNavigationObstacleGridCellLocationFragment::StaticStruct();
+	ObservedType = FMassNavigationObstacleGridCellLocationFragment::StaticStruct();
+	Operation = EMassObservedOperation::Remove;
 	ExecutionFlags = (int32)(EProcessorExecutionFlags::All);
 }
 
diff --git a/Engine/Plugins/AI/MassAI/Source/MassNavigation/Public/MassNavigationProcessors.h b/Engine/Plugins/AI/MassAI/Source/MassNavigation/Public/MassNavigationProcessors.h
--- a/Engine/Plugins/AI/MassAI/Source/MassNavigation/Public/MassNavigationProcessors.h
+++ b/Engine/Plugins/AI/MassAI/Source/MassNavigation/Public/MassNavigationProcessors.h
@@ -51,7 +51,7 @@ private:
  * Initializes the move target's location to the agents initial position.
  */
 UCLASS()
-class MASSNAVIGATION_API UMassMoveTargetFragmentInitializer : public UMassFragmentInitializer
+class MASSNAVIGATION_API UMassMoveTargetFragmentInitializer : public UMassObserverProcessor
 {
 	GENERATED_BODY()
 public:
@@ -87,7 +87,7 @@ private:
 
 /** Deinitializer processor to remove avoidance obstacles from the avoidance obstacle grid */
 UCLASS()
-class MASSNAVIGATION_API UMassNavigationObstacleRemoverProcessor : public UMassFragmentDeinitializer
+class MASSNAVIGATION_API UMassNavigationObstacleRemoverProcessor : public UMassObserverProcessor
 {
 	GENERATED_BODY()
 
diff --git a/Engine/Plugins/AI/MassAI/Source/MassZoneGraphNavigation/Private/MassZoneGraphNavigationProcessors.cpp b/Engine/Plugins/AI/MassAI/Source/MassZoneGraphNavigation/Private/MassZoneGraphNavigationProcessors.cpp
--- a/Engine/Plugins/AI/MassAI/Source/MassZoneGraphNavigation/Private/MassZoneGraphNavigationProcessors.cpp
+++ b/Engine/Plugins/AI/MassAI/Source/MassZoneGraphNavigation/Private/MassZoneGraphNavigationProcessors.cpp
@@ -37,7 +37,8 @@ namespace UE::MassNavigation::Debug
 //----------------------------------------------------------------------//
 UMassZoneGraphLocationInitializer::UMassZoneGraphLocationInitializer()
 {
-	FragmentType = FMassZoneGraphLaneLocationFragment::StaticStruct();
+	ObservedType = FMassZoneGraphLaneLocationFragment::StaticStruct();
+	Operation = EMassObservedOperation::Add;
 	ExecutionFlags = (int32)(EProcessorExecutionFlags::All);
 }
 
diff --git a/Engine/Plugins/AI/MassAI/Source/MassZoneGraphNavigation/Public/MassZoneGraphNavigationProcessors.h b/Engine/Plugins/AI/MassAI/Source/MassZoneGraphNavigation/Public/MassZoneGraphNavigationProcessors.h
--- a/Engine/Plugins/AI/MassAI/Source/MassZoneGraphNavigation/Public/MassZoneGraphNavigationProcessors.h
+++ b/Engine/Plugins/AI/MassAI/Source/MassZoneGraphNavigation/Public/MassZoneGraphNavigationProcessors.h
@@ -15,7 +15,7 @@ class UZoneGraphSubsystem;
  * Processor for initializing nearest location on ZoneGraph.
  */
 UCLASS()
-class MASSZONEGRAPHNAVIGATION_API UMassZoneGraphLocationInitializer : public UMassFragmentInitializer
+class MASSZONEGRAPHNAVIGATION_API UMassZoneGraphLocationInitializer : public UMassObserverProcessor
 {
 	GENERATED_BODY()
 	
diff --git a/Engine/Plugins/AI/MassCrowd/Source/MassCrowd/Private/MassCrowdNavigationProcessor.cpp b/Engine/Plugins/AI/MassCrowd/Source/MassCrowd/Private/MassCrowdNavigationProcessor.cpp
--- a/Engine/Plugins/AI/MassCrowd/Source/MassCrowd/Private/MassCrowdNavigationProcessor.cpp
+++ b/Engine/Plugins/AI/MassCrowd/Source/MassCrowd/Private/MassCrowdNavigationProcessor.cpp
@@ -71,7 +71,8 @@ void UMassCrowdLaneTrackingSignalProcessor::SignalEntities(UMassEntitySubsystem&
 UMassCrowdLaneTrackingDestructor::UMassCrowdLaneTrackingDestructor()
 {
 	ExecutionFlags = (int32)(EProcessorExecutionFlags::Standalone | EProcessorExecutionFlags::Server);
-	FragmentType = FMassCrowdLaneTrackingFragment::StaticStruct();
+	ObservedType = FMassCrowdLaneTrackingFragment::StaticStruct();
+	Operation = EMassObservedOperation::Remove;
 }
 
 void UMassCrowdLaneTrackingDestructor::Initialize(UObject& Owner)
diff --git a/Engine/Plugins/AI/MassCrowd/Source/MassCrowd/Public/MassCrowdNavigationProcessor.h b/Engine/Plugins/AI/MassCrowd/Source/MassCrowd/Public/MassCrowdNavigationProcessor.h
--- a/Engine/Plugins/AI/MassCrowd/Source/MassCrowd/Public/MassCrowdNavigationProcessor.h
+++ b/Engine/Plugins/AI/MassCrowd/Source/MassCrowd/Public/MassCrowdNavigationProcessor.h
@@ -31,7 +31,7 @@ protected:
 
 /** Processors that cleans up the lane tracking on entity destruction. */
 UCLASS()
-class MASSCROWD_API UMassCrowdLaneTrackingDestructor : public UMassFragmentDeinitializer
+class MASSCROWD_API UMassCrowdLaneTrackingDestructor : public UMassObserverProcessor
 {
 	GENERATED_BODY()
 
diff --git a/Engine/Plugins/Runtime/MassGameplay/Source/MassActors/Private/Translators/MassCharacterMovementTranslators.cpp b/Engine/Plugins/Runtime/MassGameplay/Source/MassActors/Private/Translators/MassCharacterMovementTranslators.cpp
--- a/Engine/Plugins/Runtime/MassGameplay/Source/MassActors/Private/Translators/MassCharacterMovementTranslators.cpp
+++ b/Engine/Plugins/Runtime/MassGameplay/Source/MassActors/Private/Translators/MassCharacterMovementTranslators.cpp
@@ -172,7 +172,8 @@ void UMassCharacterOrientationToActorTranslator::Execute(UMassEntitySubsystem& E
 //----------------------------------------------------------------------//
 UMassFragmentInitializer_NavLocation::UMassFragmentInitializer_NavLocation()
 {
-	FragmentType = FDataFragment_NavLocation::StaticStruct();
+	ObservedType = FDataFragment_NavLocation::StaticStruct();
+	Operation = EMassObservedOperation::Add;
 }
 
 void UMassFragmentInitializer_NavLocation::ConfigureQueries() 
diff --git a/Engine/Plugins/Runtime/MassGameplay/Source/MassActors/Public/Translators/MassCharacterMovementTranslators.h b/Engine/Plugins/Runtime/MassGameplay/Source/MassActors/Public/Translators/MassCharacterMovementTranslators.h
--- a/Engine/Plugins/Runtime/MassGameplay/Source/MassActors/Public/Translators/MassCharacterMovementTranslators.h
+++ b/Engine/Plugins/Runtime/MassGameplay/Source/MassActors/Public/Translators/MassCharacterMovementTranslators.h
@@ -104,7 +104,7 @@ protected:
 };
 
 UCLASS()
-class MASSACTORS_API UMassFragmentInitializer_NavLocation : public UMassFragmentInitializer
+class MASSACTORS_API UMassFragmentInitializer_NavLocation : public UMassObserverProcessor
 {
 	GENERATED_BODY()
 public:
diff --git a/Engine/Plugins/Runtime/MassGameplay/Source/MassGameplayDebug/Private/AssignDebugVisProcessor.cpp b/Engine/Plugins/Runtime/MassGameplay/Source/MassGameplayDebug/Private/AssignDebugVisProcessor.cpp
--- a/Engine/Plugins/Runtime/MassGameplay/Source/MassGameplayDebug/Private/AssignDebugVisProcessor.cpp
+++ b/Engine/Plugins/Runtime/MassGameplay/Source/MassGameplayDebug/Private/AssignDebugVisProcessor.cpp
@@ -12,7 +12,8 @@
 UAssignDebugVisProcessor::UAssignDebugVisProcessor()
 {
 	bAutoRegisterWithProcessingPhases = false;
-	FragmentType = FSimDebugVisComponent::StaticStruct();
+	ObservedType = FSimDebugVisComponent::StaticStruct();
+	Operation = EMassObservedOperation::Add;
 }
 
 void UAssignDebugVisProcessor::ConfigureQueries()
diff --git a/Engine/Plugins/Runtime/MassGameplay/Source/MassGameplayDebug/Public/AssignDebugVisProcessor.h b/Engine/Plugins/Runtime/MassGameplay/Source/MassGameplayDebug/Public/AssignDebugVisProcessor.h
--- a/Engine/Plugins/Runtime/MassGameplay/Source/MassGameplayDebug/Public/AssignDebugVisProcessor.h
+++ b/Engine/Plugins/Runtime/MassGameplay/Source/MassGameplayDebug/Public/AssignDebugVisProcessor.h
@@ -10,7 +10,7 @@ class UMassDebugVisualizationComponent;
 struct FSimDebugVisComponent;
 
 UCLASS()
-class MASSGAMEPLAYDEBUG_API UAssignDebugVisProcessor : public UMassFragmentInitializer
+class MASSGAMEPLAYDEBUG_API UAssignDebugVisProcessor : public UMassObserverProcessor
 {
 	GENERATED_BODY()
 public:
diff --git a/Engine/Plugins/Runtime/MassGameplay/Source/MassMovement/Private/MassVelocityRandomizerTrait.cpp b/Engine/Plugins/Runtime/MassGameplay/Source/MassMovement/Private/MassVelocityRandomizerTrait.cpp
--- a/Engine/Plugins/Runtime/MassGameplay/Source/MassMovement/Private/MassVelocityRandomizerTrait.cpp
+++ b/Engine/Plugins/Runtime/MassGameplay/Source/MassMovement/Private/MassVelocityRandomizerTrait.cpp
@@ -24,7 +24,8 @@ void UMassVelocityRandomizerTrait::BuildTemplate(FMassEntityTemplateBuildContext
 //----------------------------------------------------------------------//
 UMassRandomVelocityInitializer::UMassRandomVelocityInitializer()
 {
-	FragmentType = FMassVelocityFragment::StaticStruct();
+	ObservedType = FMassVelocityFragment::StaticStruct();
+	Operation = EMassObservedOperation::Add;
 }
 
 void UMassRandomVelocityInitializer::ConfigureQueries()
diff --git a/Engine/Plugins/Runtime/MassGameplay/Source/MassMovement/Public/Example/MassVelocityRandomizerTrait.h b/Engine/Plugins/Runtime/MassGameplay/Source/MassMovement/Public/Example/MassVelocityRandomizerTrait.h
--- a/Engine/Plugins/Runtime/MassGameplay/Source/MassMovement/Public/Example/MassVelocityRandomizerTrait.h
+++ b/Engine/Plugins/Runtime/MassGameplay/Source/MassMovement/Public/Example/MassVelocityRandomizerTrait.h
@@ -31,7 +31,7 @@ protected:
 
 
 UCLASS()
-class MASSMOVEMENT_API UMassRandomVelocityInitializer : public UMassFragmentInitializer
+class MASSMOVEMENT_API UMassRandomVelocityInitializer : public UMassObserverProcessor
 {
 	GENERATED_BODY()
 public:
diff --git a/Engine/Plugins/Runtime/MassGameplay/Source/MassReplication/Private/MassReplicationFragments.cpp b/Engine/Plugins/Runtime/MassGameplay/Source/MassReplication/Private/MassReplicationFragments.cpp
--- a/Engine/Plugins/Runtime/MassGameplay/Source/MassReplication/Private/MassReplicationFragments.cpp
+++ b/Engine/Plugins/Runtime/MassGameplay/Source/MassReplication/Private/MassReplicationFragments.cpp
@@ -10,7 +10,8 @@
 UMassNetworkIDFragmentInitializer::UMassNetworkIDFragmentInitializer()
 {
 	ExecutionFlags = (int32)(EProcessorExecutionFlags::Standalone | EProcessorExecutionFlags::Server);
-	FragmentType = FMassNetworkIDFragment::StaticStruct();
+	ObservedType = FMassNetworkIDFragment::StaticStruct();
+	Operation = EMassObservedOperation::Add;
 }
 
 void UMassNetworkIDFragmentInitializer::ConfigureQueries()
diff --git a/Engine/Plugins/Runtime/MassGameplay/Source/MassReplication/Public/MassReplicationFragments.h b/Engine/Plugins/Runtime/MassGameplay/Source/MassReplication/Public/MassReplicationFragments.h
--- a/Engine/Plugins/Runtime/MassGameplay/Source/MassReplication/Public/MassReplicationFragments.h
+++ b/Engine/Plugins/Runtime/MassGameplay/Source/MassReplication/Public/MassReplicationFragments.h
@@ -88,7 +88,7 @@ struct MASSREPLICATION_API FMassReplicationLODFragment : public FMassFragment
 };
 
 UCLASS()
-class MASSREPLICATION_API UMassNetworkIDFragmentInitializer : public UMassFragmentInitializer
+class MASSREPLICATION_API UMassNetworkIDFragmentInitializer : public UMassObserverProcessor
 {
 	GENERATED_BODY()
 public:
diff --git a/Engine/Plugins/Runtime/MassGameplay/Source/MassRepresentation/Private/MassRepresentationProcessor.cpp b/Engine/Plugins/Runtime/MassGameplay/Source/MassRepresentation/Private/MassRepresentationProcessor.cpp
--- a/Engine/Plugins/Runtime/MassGameplay/Source/MassRepresentation/Private/MassRepresentationProcessor.cpp
+++ b/Engine/Plugins/Runtime/MassGameplay/Source/MassRepresentation/Private/MassRepresentationProcessor.cpp
@@ -365,9 +365,8 @@ void UMassVisualizationProcessor::Execute(UMassEntitySubsystem& EntitySubsystem,
 //----------------------------------------------------------------------//
 UMassRepresentationFragmentDestructor::UMassRepresentationFragmentDestructor()
 {
-	// By putting this to null, this Deinitializer needs to be explicitly added via a Entity traits and will no longer be automatically called
-	// This FragmentType should be deprecated and we should always use the traits.
-	FragmentType = FMassRepresentationFragment::StaticStruct();
+	ObservedType = FMassRepresentationFragment::StaticStruct();
+	Operation = EMassObservedOperation::Remove;
 	ExecutionFlags = (int32)(EProcessorExecutionFlags::All);
 }
 
diff --git a/Engine/Plugins/Runtime/MassGameplay/Source/MassRepresentation/Public/MassRepresentationProcessor.h b/Engine/Plugins/Runtime/MassGameplay/Source/MassRepresentation/Public/MassRepresentationProcessor.h
--- a/Engine/Plugins/Runtime/MassGameplay/Source/MassRepresentation/Public/MassRepresentationProcessor.h
+++ b/Engine/Plugins/Runtime/MassGameplay/Source/MassRepresentation/Public/MassRepresentationProcessor.h
@@ -110,7 +110,7 @@ protected:
 
 
 UCLASS()
-class MASSREPRESENTATION_API UMassRepresentationFragmentDestructor : public UMassFragmentDeinitializer
+class MASSREPRESENTATION_API UMassRepresentationFragmentDestructor : public UMassObserverProcessor
 {
 	GENERATED_BODY()
 
diff --git a/Engine/Plugins/Runtime/MassGameplay/Source/MassSmartObjects/Private/MassSmartObjectProcessor.cpp b/Engine/Plugins/Runtime/MassGameplay/Source/MassSmartObjects/Private/MassSmartObjectProcessor.cpp
--- a/Engine/Plugins/Runtime/MassGameplay/Source/MassSmartObjects/Private/MassSmartObjectProcessor.cpp
+++ b/Engine/Plugins/Runtime/MassGameplay/Source/MassSmartObjects/Private/MassSmartObjectProcessor.cpp
@@ -358,7 +358,8 @@ void UMassSmartObjectTimedBehaviorProcessor::Execute(UMassEntitySubsystem& Entit
 //----------------------------------------------------------------------//
 UMassSmartObjectUserFragmentDeinitializer::UMassSmartObjectUserFragmentDeinitializer()
 {
-	FragmentType = FMassSmartObjectUserFragment::StaticStruct();
+	ObservedType = FMassSmartObjectUserFragment::StaticStruct();
+	Operation = EMassObservedOperation::Remove;
 	ExecutionFlags = (int32)(EProcessorExecutionFlags::All);
 }
 
diff --git a/Engine/Plugins/Runtime/MassGameplay/Source/MassSmartObjects/Public/MassSmartObjectProcessor.h b/Engine/Plugins/Runtime/MassGameplay/Source/MassSmartObjects/Public/MassSmartObjectProcessor.h
--- a/Engine/Plugins/Runtime/MassGameplay/Source/MassSmartObjects/Public/MassSmartObjectProcessor.h
+++ b/Engine/Plugins/Runtime/MassGameplay/Source/MassSmartObjects/Public/MassSmartObjectProcessor.h
@@ -72,7 +72,7 @@ protected:
 
 /** Deinitializer processor to unregister slot invalidation callback when SmartObjectUser fragment gets removed */
 UCLASS()
-class MASSSMARTOBJECTS_API UMassSmartObjectUserFragmentDeinitializer : public UMassFragmentDeinitializer
+class MASSSMARTOBJECTS_API UMassSmartObjectUserFragmentDeinitializer : public UMassObserverProcessor
 {
 	GENERATED_BODY()
 
