Mass Crowd fixes:
- fix crowd dancing while waiting for to cross a street
- fixed crowd player avoidance

diff --git a/Engine/Plugins/AI/MassAI/Source/MassNavigation/Private/MassSteeringProcessors.cpp b/Engine/Plugins/AI/MassAI/Source/MassNavigation/Private/MassSteeringProcessors.cpp
--- a/Engine/Plugins/AI/MassAI/Source/MassNavigation/Private/MassSteeringProcessors.cpp
+++ b/Engine/Plugins/AI/MassAI/Source/MassNavigation/Private/MassSteeringProcessors.cpp
@@ -244,15 +244,15 @@ void UMassSteerToMoveTargetProcessor::Execute(UMassEntitySubsystem& EntitySubsys
 					if (StandingSteering.bEnteredFromMoveAction)
 					{
 						// If the current steering target is from approaching a move target, use the same speed logic as movement to ensure smooth transition.
-						const float SpeedFade = FMath::Clamp(Distance / LookAheadDistance, 0.0f, 1.0f);
+						const float Range = FMath::Max(1.0f, LookAheadDistance - StandingSteeringParams.DeadZoneRadius);
+						const float SpeedFade = FMath::Clamp((Distance - StandingSteeringParams.DeadZoneRadius) / Range, 0.0f, 1.0f);
 						DesiredSpeed = MoveTarget.DesiredSpeed.Get() * UE::MassNavigation::CalcDirectionalSpeedScale(CurrentForward, SteerDirection) * UE::MassNavigation::ArrivalSpeedEnvelope(SpeedFade);
 					}
 					else
 					{
-						// More aggressive movement when doing adjustments.
-						const float AdjustmentLookAheadDistance = LookAheadDistance * 0.5f;
-						// Intentionally not taking dead zone into account here, so that the speed does not drop to zero.
-						const float SpeedFade = FMath::Clamp(Distance / AdjustmentLookAheadDistance, 0.0f, 1.0f);
+						const float Range = FMath::Max(1.0f, LookAheadDistance - StandingSteeringParams.DeadZoneRadius);
+						const float SpeedFade = FMath::Clamp((Distance - StandingSteeringParams.DeadZoneRadius) / Range, 0.0f, 1.0f);
+						// Not using the directional scaling so that the steps we take to avoid are done quickly, and the behavior is reactive.
 						DesiredSpeed = MoveTarget.DesiredSpeed.Get() * UE::MassNavigation::ArrivalSpeedEnvelope(SpeedFade);
 					}
 					
diff --git a/Engine/Plugins/AI/MassAI/Source/MassNavigation/Public/Steering/MassSteeringFragments.h b/Engine/Plugins/AI/MassAI/Source/MassNavigation/Public/Steering/MassSteeringFragments.h
--- a/Engine/Plugins/AI/MassAI/Source/MassNavigation/Public/Steering/MassSteeringFragments.h
+++ b/Engine/Plugins/AI/MassAI/Source/MassNavigation/Public/Steering/MassSteeringFragments.h
@@ -93,5 +93,5 @@ struct FMassStandingSteeringParameters : public FMassSharedFragment
 	float TargetSelectionCooldownVariance = 0.5f;
 
 	UPROPERTY(EditAnywhere, Category = "Standing", meta = (ForceUnits="cm"))
-	float DeadZoneRadius = 5.0f;
+	float DeadZoneRadius = 15.0f;
 };
