[Mass] fixed shared fragment support in StateTree

diff --git a/Engine/Plugins/AI/MassAI/Source/MassAIBehavior/Private/MassStateTreeTrait.cpp b/Engine/Plugins/AI/MassAI/Source/MassAIBehavior/Private/MassStateTreeTrait.cpp
--- a/Engine/Plugins/AI/MassAI/Source/MassAIBehavior/Private/MassStateTreeTrait.cpp
+++ b/Engine/Plugins/AI/MassAI/Source/MassAIBehavior/Private/MassStateTreeTrait.cpp
@@ -70,6 +70,11 @@ void UMassStateTreeTrait::ValidateTemplate(FMassEntityTemplateBuildContext& Buil
 				const bool bContainsFragment = BuildContext.HasFragment(*CastChecked<UScriptStruct>(ItemDesc.Struct));
 				UE_CVLOG(!bContainsFragment, MassStateTreeSubsystem, LogMassBehavior, Error, TEXT("StateTree %s: Could not find required fragment %s"), *GetNameSafe(StateTree), *GetNameSafe(ItemDesc.Struct));
 			}
+			else if (ItemDesc.Struct->IsChildOf(FMassSharedFragment::StaticStruct()))
+			{
+				const bool bContainsFragment = BuildContext.HasSharedFragment(*CastChecked<UScriptStruct>(ItemDesc.Struct));
+				UE_CVLOG(!bContainsFragment, MassStateTreeSubsystem, LogMassBehavior, Error, TEXT("StateTree %s: Could not find required shared fragment %s"), *GetNameSafe(StateTree), *GetNameSafe(ItemDesc.Struct));
+			}
 			else
 			{
 				UE_VLOG(MassStateTreeSubsystem, LogMassBehavior, Error, TEXT("StateTree %s: Unsupported requirement %s"), *GetNameSafe(StateTree), *GetNameSafe(ItemDesc.Struct));
diff --git a/Engine/Plugins/Runtime/MassGameplay/Source/MassSpawner/Public/MassEntityTemplate.h b/Engine/Plugins/Runtime/MassGameplay/Source/MassSpawner/Public/MassEntityTemplate.h
--- a/Engine/Plugins/Runtime/MassGameplay/Source/MassSpawner/Public/MassEntityTemplate.h
+++ b/Engine/Plugins/Runtime/MassGameplay/Source/MassSpawner/Public/MassEntityTemplate.h
@@ -229,6 +229,11 @@ struct MASSSPAWNER_API FMassEntityTemplate
 		return Composition.SharedFragments.Contains<T>();
 	}
 
+	bool HasSharedFragment(const UScriptStruct& ScriptStruct) const
+	{
+		return Composition.SharedFragments.Contains(ScriptStruct);
+	}
+
 	void Sort()
 	{
 		SharedFragmentValues.Sort();
diff --git a/Engine/Plugins/Runtime/MassGameplay/Source/MassSpawner/Public/MassEntityTemplateRegistry.h b/Engine/Plugins/Runtime/MassGameplay/Source/MassSpawner/Public/MassEntityTemplateRegistry.h
--- a/Engine/Plugins/Runtime/MassGameplay/Source/MassSpawner/Public/MassEntityTemplateRegistry.h
+++ b/Engine/Plugins/Runtime/MassGameplay/Source/MassSpawner/Public/MassEntityTemplateRegistry.h
@@ -101,6 +101,11 @@ struct FMassEntityTemplateBuildContext
 		return Template.HasSharedFragment<T>();
 	}
 
+	bool HasSharedFragment(const UScriptStruct& ScriptStruct) const
+	{
+		return Template.HasSharedFragment(ScriptStruct);
+	}
+
 	//----------------------------------------------------------------------//
 	// Translators
 	//----------------------------------------------------------------------//
