Removed improper reference in Mass Framework for 5.0 release

diff --git a/Engine/Plugins/AI/MassAI/Config/BaseMassAI.ini b/Engine/Plugins/AI/MassAI/Config/BaseMassAI.ini
--- a/Engine/Plugins/AI/MassAI/Config/BaseMassAI.ini
+++ b/Engine/Plugins/AI/MassAI/Config/BaseMassAI.ini
@@ -37,3 +37,9 @@
 +PropertyRedirects=(OldName="MassEntityConfig.Features",NewName="Traits")
 
 +PropertyRedirects=(OldName="MassZoneGraphFindEscapeTarget.FleeTag",NewName="DisturbanceAnnotationTag")
+
+;MassAIBehavior
++StructRedirects=(OldName="MassZoneGraphAnnotationTagsFragment",NewName="MassZoneGraphAnnotationFragment")
++StructRedirects=(OldName="MassStateTreeActivated",NewName="MassStateTreeActivatedTag")
++StructRedirects=(OldName="/Script/MassAIBehavior.MassZoneGraphAnnotationTagsFragment",NewName="/Script/MassAIBehavior.MassZoneGraphAnnotationFragment")
++StructRedirects=(OldName="/Script/MassAIBehavior.MassStateTreeActivated",NewName="/Script/MassAIBehavior.MassStateTreeActivatedTag")
diff --git a/Engine/Plugins/AI/MassAI/Source/MassAIBehavior/Private/Evaluators/MassStateTreeSmartObjectEvaluator.cpp b/Engine/Plugins/AI/MassAI/Source/MassAIBehavior/Private/Evaluators/MassStateTreeSmartObjectEvaluator.cpp
--- a/Engine/Plugins/AI/MassAI/Source/MassAIBehavior/Private/Evaluators/MassStateTreeSmartObjectEvaluator.cpp
+++ b/Engine/Plugins/AI/MassAI/Source/MassAIBehavior/Private/Evaluators/MassStateTreeSmartObjectEvaluator.cpp
@@ -132,7 +132,7 @@ void FMassStateTreeSmartObjectEvaluator::Evaluate(FStateTreeExecutionContext& Co
 		}
 		else
 		{
-			const FDataFragment_Transform& TransformFragment = Context.GetExternalData(EntityTransformHandle);
+			const FTransformFragment& TransformFragment = Context.GetExternalData(EntityTransformHandle);
 			SearchRequestID = MassSmartObjectHandler.FindCandidatesAsync(RequestingEntity, TransformFragment.GetTransform().GetLocation());
 		}
 	}
diff --git a/Engine/Plugins/AI/MassAI/Source/MassAIBehavior/Private/Evaluators/MassZoneGraphAnnotationEvaluator.cpp b/Engine/Plugins/AI/MassAI/Source/MassAIBehavior/Private/Evaluators/MassZoneGraphAnnotationEvaluator.cpp
--- a/Engine/Plugins/AI/MassAI/Source/MassAIBehavior/Private/Evaluators/MassZoneGraphAnnotationEvaluator.cpp
+++ b/Engine/Plugins/AI/MassAI/Source/MassAIBehavior/Private/Evaluators/MassZoneGraphAnnotationEvaluator.cpp
@@ -20,7 +20,7 @@ bool FMassZoneGraphAnnotationEvaluator::Link(FStateTreeLinker& Linker)
 
 void FMassZoneGraphAnnotationEvaluator::Evaluate(FStateTreeExecutionContext& Context, const EStateTreeEvaluationType EvalType, const float DeltaTime) const
 {
-	const FMassZoneGraphAnnotationTagsFragment& AnnotationTagsFragment = Context.GetExternalData(AnnotationTagsFragmentHandle);
+	const FMassZoneGraphAnnotationFragment& AnnotationTagsFragment = Context.GetExternalData(AnnotationTagsFragmentHandle);
 	FZoneGraphTagMask& AnnotationTags = Context.GetInstanceData(AnnotationTagsHandle);
 	AnnotationTags = AnnotationTagsFragment.Tags;
 }
diff --git a/Engine/Plugins/AI/MassAI/Source/MassAIBehavior/Private/MassLookAtProcessors.cpp b/Engine/Plugins/AI/MassAI/Source/MassAIBehavior/Private/MassLookAtProcessors.cpp
--- a/Engine/Plugins/AI/MassAI/Source/MassAIBehavior/Private/MassLookAtProcessors.cpp
+++ b/Engine/Plugins/AI/MassAI/Source/MassAIBehavior/Private/MassLookAtProcessors.cpp
@@ -109,7 +109,7 @@ void UMassLookAtProcessor::Initialize(UObject& Owner)
 void UMassLookAtProcessor::ConfigureQueries()
 {
 	EntityQuery_Conditional.AddRequirement<FMassLookAtFragment>(EMassFragmentAccess::ReadWrite);
-	EntityQuery_Conditional.AddRequirement<FDataFragment_Transform>(EMassFragmentAccess::ReadOnly);
+	EntityQuery_Conditional.AddRequirement<FTransformFragment>(EMassFragmentAccess::ReadOnly);
 	EntityQuery_Conditional.AddRequirement<FMassMoveTargetFragment>(EMassFragmentAccess::ReadOnly);
 	EntityQuery_Conditional.AddRequirement<FMassZoneGraphLaneLocationFragment>(EMassFragmentAccess::ReadOnly, EMassFragmentPresence::Optional);
 	EntityQuery_Conditional.AddRequirement<FMassLookAtTrajectoryFragment>(EMassFragmentAccess::ReadWrite, EMassFragmentPresence::Optional);
@@ -131,7 +131,7 @@ void UMassLookAtProcessor::Execute(UMassEntitySubsystem& EntitySubsystem, FMassE
 		{
 			const int32 NumEntities = Context.GetNumEntities();
 			const TArrayView<FMassLookAtFragment> LookAtList = Context.GetMutableFragmentView<FMassLookAtFragment>();
-			const TConstArrayView<FDataFragment_Transform> TransformList = Context.GetFragmentView<FDataFragment_Transform>();
+			const TConstArrayView<FTransformFragment> TransformList = Context.GetFragmentView<FTransformFragment>();
 			const TConstArrayView<FMassMoveTargetFragment> MoveTargetList = Context.GetFragmentView<FMassMoveTargetFragment>();
 			const TConstArrayView<FMassZoneGraphLaneLocationFragment> ZoneGraphLocationList = Context.GetFragmentView<FMassZoneGraphLaneLocationFragment>();
 			const TConstArrayView<FMassZoneGraphShortPathFragment> ShortPathList = Context.GetFragmentView<FMassZoneGraphShortPathFragment>();
@@ -141,7 +141,7 @@ void UMassLookAtProcessor::Execute(UMassEntitySubsystem& EntitySubsystem, FMassE
 			{
 				FMassLookAtFragment& LookAt = LookAtList[i];
 				const FMassMoveTargetFragment& MoveTarget = MoveTargetList[i];
-				const FDataFragment_Transform& TransformFragment = TransformList[i];
+				const FTransformFragment& TransformFragment = TransformList[i];
 
 				const bool bHasLookAtTrajectory = ZoneGraphLocationList.Num() > 0 && LookAtTrajectoryList.Num() > 0 && ShortPathList.Num() > 0;
 
@@ -284,7 +284,7 @@ void UMassLookAtProcessor::FindNewGazeTarget(const UMassEntitySubsystem& EntityS
 			}
 
 			// TargetTag is added through the LookAtTargetTrait and Transform was added with it
-			const FDataFragment_Transform& TargetTransform = EntityView.GetFragmentData<FDataFragment_Transform>();
+			const FTransformFragment& TargetTransform = EntityView.GetFragmentData<FTransformFragment>();
 			const FVector TargetLocation = TargetTransform.GetTransform().GetLocation();
 			FVector Direction = (TargetLocation - Location).GetSafeNormal();
 			Direction = Transform.InverseTransformVector(Direction);
@@ -358,7 +358,7 @@ void UMassLookAtProcessor::UpdateLookAtTrackedEntity(const UMassEntitySubsystem&
 	// Update direction toward target
 	if (EntitySubsystem.IsEntityValid(LookAt.TrackedEntity))
 	{
-		if (const FDataFragment_Transform* TargetTransform = EntitySubsystem.GetFragmentDataPtr<FDataFragment_Transform>(LookAt.TrackedEntity))
+		if (const FTransformFragment* TargetTransform = EntitySubsystem.GetFragmentDataPtr<FTransformFragment>(LookAt.TrackedEntity))
 		{
 			const FVector AgentPosition = Transform.GetLocation();
 			const FVector NewGlobalDirection = (TargetTransform->GetTransform().GetLocation() - AgentPosition).GetSafeNormal();
@@ -382,7 +382,7 @@ bool UMassLookAtProcessor::UpdateGazeTrackedEntity(const UMassEntitySubsystem& E
 	// Update direction toward gaze target
 	if (LookAt.GazeTrackedEntity.IsSet() && EntitySubsystem.IsEntityValid(LookAt.GazeTrackedEntity))
 	{
-		if (const FDataFragment_Transform* TargetTransform = EntitySubsystem.GetFragmentDataPtr<FDataFragment_Transform>(LookAt.GazeTrackedEntity))
+		if (const FTransformFragment* TargetTransform = EntitySubsystem.GetFragmentDataPtr<FTransformFragment>(LookAt.GazeTrackedEntity))
 		{
 			const FVector AgentPosition = Transform.GetLocation();
 			const FVector NewGlobalDirection = (TargetTransform->GetTransform().GetLocation() - AgentPosition).GetSafeNormal();
diff --git a/Engine/Plugins/AI/MassAI/Source/MassAIBehavior/Private/MassLookAtTargetTrait.cpp b/Engine/Plugins/AI/MassAI/Source/MassAIBehavior/Private/MassLookAtTargetTrait.cpp
--- a/Engine/Plugins/AI/MassAI/Source/MassAIBehavior/Private/MassLookAtTargetTrait.cpp
+++ b/Engine/Plugins/AI/MassAI/Source/MassAIBehavior/Private/MassLookAtTargetTrait.cpp
@@ -8,5 +8,5 @@
 void UMassLookAtTargetTrait::BuildTemplate(FMassEntityTemplateBuildContext& BuildContext, UWorld& World) const
 {
 	BuildContext.AddTag<FMassLookAtTargetTag>();
-	BuildContext.AddFragment<FDataFragment_Transform>();
+	BuildContext.AddFragment<FTransformFragment>();
 }
diff --git a/Engine/Plugins/AI/MassAI/Source/MassAIBehavior/Private/MassStateTreeProcessors.cpp b/Engine/Plugins/AI/MassAI/Source/MassAIBehavior/Private/MassStateTreeProcessors.cpp
--- a/Engine/Plugins/AI/MassAI/Source/MassAIBehavior/Private/MassStateTreeProcessors.cpp
+++ b/Engine/Plugins/AI/MassAI/Source/MassAIBehavior/Private/MassStateTreeProcessors.cpp
@@ -231,7 +231,7 @@ void UMassStateTreeActivationProcessor::Initialize(UObject& Owner)
 void UMassStateTreeActivationProcessor::ConfigureQueries()
 {
 	EntityQuery.AddRequirement<FMassStateTreeFragment>(EMassFragmentAccess::ReadOnly);
-	EntityQuery.AddTagRequirement<FMassStateTreeActivated>(EMassFragmentPresence::None);
+	EntityQuery.AddTagRequirement<FMassStateTreeActivatedTag>(EMassFragmentPresence::None);
 	EntityQuery.AddChunkRequirement<FMassSimulationVariableTickChunkFragment>(EMassFragmentAccess::ReadOnly, EMassFragmentPresence::Optional);
 }
 
@@ -278,7 +278,7 @@ void UMassStateTreeActivationProcessor::Execute(UMassEntitySubsystem& EntitySubs
 		// Adding a tag on each entities to remember we have sent the state tree initialization signal
 		for (int32 i = 0; i < NumEntities; ++i)
 		{
-			Context.Defer().AddTag<FMassStateTreeActivated>(Context.GetEntity(i));
+			Context.Defer().AddTag<FMassStateTreeActivatedTag>(Context.GetEntity(i));
 		}
 	});
 	// Signal all entities inside the consolidated list
diff --git a/Engine/Plugins/AI/MassAI/Source/MassAIBehavior/Private/MassZoneGraphAnnotationProcessors.cpp b/Engine/Plugins/AI/MassAI/Source/MassAIBehavior/Private/MassZoneGraphAnnotationProcessors.cpp
--- a/Engine/Plugins/AI/MassAI/Source/MassAIBehavior/Private/MassZoneGraphAnnotationProcessors.cpp
+++ b/Engine/Plugins/AI/MassAI/Source/MassAIBehavior/Private/MassZoneGraphAnnotationProcessors.cpp
@@ -16,7 +16,7 @@
 //----------------------------------------------------------------------//
 UMassZoneGraphAnnotationTagsInitializer::UMassZoneGraphAnnotationTagsInitializer()
 {
-	ObservedType = FMassZoneGraphAnnotationTagsFragment::StaticStruct();
+	ObservedType = FMassZoneGraphAnnotationFragment::StaticStruct();
 	Operation = EMassObservedOperation::Add;
 }
 
@@ -28,7 +28,7 @@ void UMassZoneGraphAnnotationTagsInitializer::Initialize(UObject& Owner)
 
 void UMassZoneGraphAnnotationTagsInitializer::ConfigureQueries()
 {
-	EntityQuery.AddRequirement<FMassZoneGraphAnnotationTagsFragment>(EMassFragmentAccess::ReadWrite);
+	EntityQuery.AddRequirement<FMassZoneGraphAnnotationFragment>(EMassFragmentAccess::ReadWrite);
 	EntityQuery.AddRequirement<FMassZoneGraphLaneLocationFragment>(EMassFragmentAccess::ReadOnly);
 }
 
@@ -37,12 +37,12 @@ void UMassZoneGraphAnnotationTagsInitializer::Execute(UMassEntitySubsystem& Enti
 	EntityQuery.ForEachEntityChunk(EntitySubsystem, Context, [this](FMassExecutionContext& Context)
 	{
 		const int32 NumEntities = Context.GetNumEntities();
-		const TArrayView<FMassZoneGraphAnnotationTagsFragment> AnnotationTagsList = Context.GetMutableFragmentView<FMassZoneGraphAnnotationTagsFragment>();
+		const TArrayView<FMassZoneGraphAnnotationFragment> AnnotationTagsList = Context.GetMutableFragmentView<FMassZoneGraphAnnotationFragment>();
 		const TConstArrayView<FMassZoneGraphLaneLocationFragment> LaneLocationList = Context.GetFragmentView<FMassZoneGraphLaneLocationFragment>();
 
 		for (int32 EntityIndex = 0; EntityIndex < NumEntities; ++EntityIndex)
 		{
-			FMassZoneGraphAnnotationTagsFragment& AnnotationTags = AnnotationTagsList[EntityIndex];
+			FMassZoneGraphAnnotationFragment& AnnotationTags = AnnotationTagsList[EntityIndex];
 			const FMassZoneGraphLaneLocationFragment& LaneLocation = LaneLocationList[EntityIndex];
 
 			if (!LaneLocation.LaneHandle.IsValid())
@@ -78,7 +78,7 @@ void UMassZoneGraphAnnotationTagUpdateProcessor::Initialize(UObject& Owner)
 void UMassZoneGraphAnnotationTagUpdateProcessor::ConfigureQueries()
 {
 	Super::ConfigureQueries();
-	EntityQuery.AddRequirement<FMassZoneGraphAnnotationTagsFragment>(EMassFragmentAccess::ReadWrite);
+	EntityQuery.AddRequirement<FMassZoneGraphAnnotationFragment>(EMassFragmentAccess::ReadWrite);
 	EntityQuery.AddRequirement<FMassZoneGraphLaneLocationFragment>(EMassFragmentAccess::ReadOnly);
 	EntityQuery.AddChunkRequirement<FMassZoneGraphAnnotationVariableTickChunkFragment>(EMassFragmentAccess::ReadOnly);
 	EntityQuery.AddChunkRequirement<FMassSimulationVariableTickChunkFragment>(EMassFragmentAccess::ReadOnly, EMassFragmentPresence::Optional);
@@ -100,12 +100,12 @@ void UMassZoneGraphAnnotationTagUpdateProcessor::Execute(UMassEntitySubsystem& E
 		}
 
 		const int32 NumEntities = Context.GetNumEntities();
-		const TArrayView<FMassZoneGraphAnnotationTagsFragment> AnnotationTagsList = Context.GetMutableFragmentView<FMassZoneGraphAnnotationTagsFragment>();
+		const TArrayView<FMassZoneGraphAnnotationFragment> AnnotationTagsList = Context.GetMutableFragmentView<FMassZoneGraphAnnotationFragment>();
 		const TConstArrayView<FMassZoneGraphLaneLocationFragment> LaneLocationList = Context.GetFragmentView<FMassZoneGraphLaneLocationFragment>();
 
 		for (int32 EntityIndex = 0; EntityIndex < NumEntities; ++EntityIndex)
 		{
-			FMassZoneGraphAnnotationTagsFragment& AnnotationTags = AnnotationTagsList[EntityIndex];
+			FMassZoneGraphAnnotationFragment& AnnotationTags = AnnotationTagsList[EntityIndex];
 			const FMassZoneGraphLaneLocationFragment& LaneLocation = LaneLocationList[EntityIndex];
 
 			UpdateAnnotationTags(AnnotationTags, LaneLocation, Context.GetEntity(EntityIndex));
@@ -118,7 +118,7 @@ void UMassZoneGraphAnnotationTagUpdateProcessor::Execute(UMassEntitySubsystem& E
 	}
 }
 
-void UMassZoneGraphAnnotationTagUpdateProcessor::UpdateAnnotationTags(FMassZoneGraphAnnotationTagsFragment& AnnotationTags, const FMassZoneGraphLaneLocationFragment& LaneLocation, const FMassEntityHandle Entity)
+void UMassZoneGraphAnnotationTagUpdateProcessor::UpdateAnnotationTags(FMassZoneGraphAnnotationFragment& AnnotationTags, const FMassZoneGraphLaneLocationFragment& LaneLocation, const FMassEntityHandle Entity)
 {
 	const FZoneGraphTagMask OldTags = AnnotationTags.Tags;
 
@@ -148,11 +148,11 @@ void UMassZoneGraphAnnotationTagUpdateProcessor::SignalEntities(UMassEntitySubsy
 	{
 		const int32 NumEntities = Context.GetNumEntities();
 		const TConstArrayView<FMassZoneGraphLaneLocationFragment> LaneLocationList = Context.GetFragmentView<FMassZoneGraphLaneLocationFragment>();
-		const TArrayView<FMassZoneGraphAnnotationTagsFragment> AnnotationTagsList = Context.GetMutableFragmentView<FMassZoneGraphAnnotationTagsFragment>();
+		const TArrayView<FMassZoneGraphAnnotationFragment> AnnotationTagsList = Context.GetMutableFragmentView<FMassZoneGraphAnnotationFragment>();
 
 		for (int32 EntityIndex = 0; EntityIndex < NumEntities; ++EntityIndex)
 		{
-			FMassZoneGraphAnnotationTagsFragment& AnnotationTags = AnnotationTagsList[EntityIndex];
+			FMassZoneGraphAnnotationFragment& AnnotationTags = AnnotationTagsList[EntityIndex];
 			const FMassZoneGraphLaneLocationFragment& LaneLocation = LaneLocationList[EntityIndex];
 
 			UpdateAnnotationTags(AnnotationTags, LaneLocation, Context.GetEntity(EntityIndex));
diff --git a/Engine/Plugins/AI/MassAI/Source/MassAIBehavior/Private/MassZoneGraphAnnotationTrait.cpp b/Engine/Plugins/AI/MassAI/Source/MassAIBehavior/Private/MassZoneGraphAnnotationTrait.cpp
--- a/Engine/Plugins/AI/MassAI/Source/MassAIBehavior/Private/MassZoneGraphAnnotationTrait.cpp
+++ b/Engine/Plugins/AI/MassAI/Source/MassAIBehavior/Private/MassZoneGraphAnnotationTrait.cpp
@@ -7,6 +7,6 @@
 
 void UMassZoneGraphAnnotationTrait::BuildTemplate(FMassEntityTemplateBuildContext& BuildContext, UWorld& World) const
 {
-	BuildContext.AddFragment<FMassZoneGraphAnnotationTagsFragment>();
+	BuildContext.AddFragment<FMassZoneGraphAnnotationFragment>();
 	BuildContext.AddChunkFragment<FMassZoneGraphAnnotationVariableTickChunkFragment>();
 }
diff --git a/Engine/Plugins/AI/MassAI/Source/MassAIBehavior/Private/Tasks/MassStateTreeSmartObjectTask.cpp b/Engine/Plugins/AI/MassAI/Source/MassAIBehavior/Private/Tasks/MassStateTreeSmartObjectTask.cpp
--- a/Engine/Plugins/AI/MassAI/Source/MassAIBehavior/Private/Tasks/MassStateTreeSmartObjectTask.cpp
+++ b/Engine/Plugins/AI/MassAI/Source/MassAIBehavior/Private/Tasks/MassStateTreeSmartObjectTask.cpp
@@ -144,7 +144,7 @@ EStateTreeRunStatus FMassUseSmartObjectTask::EnterState(FStateTreeExecutionConte
 	USmartObjectSubsystem& SmartObjectSubsystem = Context.GetExternalData(SmartObjectSubsystemHandle);
 	UMassSignalSubsystem& SignalSubsystem = Context.GetExternalData(MassSignalSubsystemHandle);
 	FMassSmartObjectUserFragment& SOUser = Context.GetExternalData(SmartObjectUserHandle);
-	const FDataFragment_Transform& TransformFragment = Context.GetExternalData(EntityTransformHandle);
+	const FTransformFragment& TransformFragment = Context.GetExternalData(EntityTransformHandle);
 	FMassMoveTargetFragment& MoveTarget = Context.GetExternalData(MoveTargetHandle);
 
 	// Setup MassSmartObject handler and start interaction
diff --git a/Engine/Plugins/AI/MassAI/Source/MassAIBehavior/Private/Tasks/MassZoneGraphPathFollowTask.cpp b/Engine/Plugins/AI/MassAI/Source/MassAIBehavior/Private/Tasks/MassZoneGraphPathFollowTask.cpp
--- a/Engine/Plugins/AI/MassAI/Source/MassAIBehavior/Private/Tasks/MassZoneGraphPathFollowTask.cpp
+++ b/Engine/Plugins/AI/MassAI/Source/MassAIBehavior/Private/Tasks/MassZoneGraphPathFollowTask.cpp
@@ -33,7 +33,7 @@ bool FMassZoneGraphPathFollowTask::RequestPath(FMassStateTreeExecutionContext& C
 {
 	const UZoneGraphSubsystem& ZoneGraphSubsystem = Context.GetExternalData(ZoneGraphSubsystemHandle);
 	const FMassZoneGraphLaneLocationFragment& LaneLocation = Context.GetExternalData(LocationHandle);
-	const FDataFragment_AgentRadius& AgentRadius = Context.GetExternalData(AgentRadiusHandle);
+	const FAgentRadiusFragment& AgentRadius = Context.GetExternalData(AgentRadiusHandle);
 	FMassZoneGraphShortPathFragment& ShortPath = Context.GetExternalData(ShortPathHandle);
 	FMassZoneGraphCachedLaneFragment& CachedLane = Context.GetExternalData(CachedLaneHandle);
 	FMassMoveTargetFragment& MoveTarget = Context.GetExternalData(MoveTargetHandle);
diff --git a/Engine/Plugins/AI/MassAI/Source/MassAIBehavior/Public/Evaluators/MassStateTreeSmartObjectEvaluator.h b/Engine/Plugins/AI/MassAI/Source/MassAIBehavior/Public/Evaluators/MassStateTreeSmartObjectEvaluator.h
--- a/Engine/Plugins/AI/MassAI/Source/MassAIBehavior/Public/Evaluators/MassStateTreeSmartObjectEvaluator.h
+++ b/Engine/Plugins/AI/MassAI/Source/MassAIBehavior/Public/Evaluators/MassStateTreeSmartObjectEvaluator.h
@@ -9,7 +9,7 @@
 #include "MassStateTreeSmartObjectEvaluator.generated.h"
 
 struct FStateTreeExecutionContext;
-struct FDataFragment_Transform;
+struct FTransformFragment;
 struct FMassZoneGraphLaneLocationFragment;
 struct FMassSmartObjectUserFragment;
 
@@ -60,7 +60,7 @@ protected:
 
 	TStateTreeExternalDataHandle<USmartObjectSubsystem> SmartObjectSubsystemHandle;
 	TStateTreeExternalDataHandle<UMassSignalSubsystem> MassSignalSubsystemHandle;
-	TStateTreeExternalDataHandle<FDataFragment_Transform> EntityTransformHandle;
+	TStateTreeExternalDataHandle<FTransformFragment> EntityTransformHandle;
 	TStateTreeExternalDataHandle<FMassSmartObjectUserFragment> SmartObjectUserHandle;
 	TStateTreeExternalDataHandle<FMassZoneGraphLaneLocationFragment, EStateTreeExternalDataRequirement::Optional> LocationHandle;
 
diff --git a/Engine/Plugins/AI/MassAI/Source/MassAIBehavior/Public/Evaluators/MassZoneGraphAnnotationEvaluator.h b/Engine/Plugins/AI/MassAI/Source/MassAIBehavior/Public/Evaluators/MassZoneGraphAnnotationEvaluator.h
--- a/Engine/Plugins/AI/MassAI/Source/MassAIBehavior/Public/Evaluators/MassZoneGraphAnnotationEvaluator.h
+++ b/Engine/Plugins/AI/MassAI/Source/MassAIBehavior/Public/Evaluators/MassZoneGraphAnnotationEvaluator.h
@@ -7,7 +7,7 @@
 #include "MassZoneGraphAnnotationEvaluator.generated.h"
 
 struct FStateTreeExecutionContext;
-struct FMassZoneGraphAnnotationTagsFragment;
+struct FMassZoneGraphAnnotationFragment;
 /**
  * Evaluator to expose ZoneGraph Annotation Tags for decision making.
  */
@@ -32,7 +32,7 @@ protected:
 	virtual const UStruct* GetInstanceDataType() const override { return FMassZoneGraphAnnotationEvaluatorInstanceData::StaticStruct(); }
 	virtual void Evaluate(FStateTreeExecutionContext& Context, const EStateTreeEvaluationType EvalType, const float DeltaTime) const override;
 
-	TStateTreeExternalDataHandle<FMassZoneGraphAnnotationTagsFragment> AnnotationTagsFragmentHandle;
+	TStateTreeExternalDataHandle<FMassZoneGraphAnnotationFragment> AnnotationTagsFragmentHandle;
 
 	TStateTreeInstanceDataPropertyHandle<FZoneGraphTagMask> AnnotationTagsHandle;
 };
diff --git a/Engine/Plugins/AI/MassAI/Source/MassAIBehavior/Public/MassStateTreeProcessors.h b/Engine/Plugins/AI/MassAI/Source/MassAIBehavior/Public/MassStateTreeProcessors.h
--- a/Engine/Plugins/AI/MassAI/Source/MassAIBehavior/Public/MassStateTreeProcessors.h
+++ b/Engine/Plugins/AI/MassAI/Source/MassAIBehavior/Public/MassStateTreeProcessors.h
@@ -36,7 +36,7 @@ protected:
  * Special tag to know if the state tree has been activated
  */
 USTRUCT()
-struct MASSAIBEHAVIOR_API FMassStateTreeActivated : public FMassTag
+struct MASSAIBEHAVIOR_API FMassStateTreeActivatedTag : public FMassTag
 {
 	GENERATED_BODY()
 };
diff --git a/Engine/Plugins/AI/MassAI/Source/MassAIBehavior/Public/MassZoneGraphAnnotationFragments.h b/Engine/Plugins/AI/MassAI/Source/MassAIBehavior/Public/MassZoneGraphAnnotationFragments.h
--- a/Engine/Plugins/AI/MassAI/Source/MassAIBehavior/Public/MassZoneGraphAnnotationFragments.h
+++ b/Engine/Plugins/AI/MassAI/Source/MassAIBehavior/Public/MassZoneGraphAnnotationFragments.h
@@ -9,7 +9,7 @@
 struct FMassExecutionContext;
 
 USTRUCT()
-struct MASSAIBEHAVIOR_API FMassZoneGraphAnnotationTagsFragment : public FMassFragment
+struct MASSAIBEHAVIOR_API FMassZoneGraphAnnotationFragment : public FMassFragment
 {
 	GENERATED_BODY()
 
diff --git a/Engine/Plugins/AI/MassAI/Source/MassAIBehavior/Public/MassZoneGraphAnnotationProcessors.h b/Engine/Plugins/AI/MassAI/Source/MassAIBehavior/Public/MassZoneGraphAnnotationProcessors.h
--- a/Engine/Plugins/AI/MassAI/Source/MassAIBehavior/Public/MassZoneGraphAnnotationProcessors.h
+++ b/Engine/Plugins/AI/MassAI/Source/MassAIBehavior/Public/MassZoneGraphAnnotationProcessors.h
@@ -8,7 +8,7 @@
 
 class UMassSignalSubsystem;
 class UZoneGraphAnnotationSubsystem;
-struct FMassZoneGraphAnnotationTagsFragment;
+struct FMassZoneGraphAnnotationFragment;
 struct FMassZoneGraphLaneLocationFragment;
 struct FMassEntityHandle;
 
@@ -50,7 +50,7 @@ protected:
 	virtual void ConfigureQueries() override;
 	virtual void Execute(UMassEntitySubsystem& EntitySubsystem, FMassExecutionContext& Context) override;
 
-	void UpdateAnnotationTags(FMassZoneGraphAnnotationTagsFragment& AnnotationTags, const FMassZoneGraphLaneLocationFragment& LaneLocation, FMassEntityHandle Entity);
+	void UpdateAnnotationTags(FMassZoneGraphAnnotationFragment& AnnotationTags, const FMassZoneGraphLaneLocationFragment& LaneLocation, FMassEntityHandle Entity);
 
 	virtual void SignalEntities(UMassEntitySubsystem& EntitySubsystem, FMassExecutionContext& Context, FMassSignalNameLookup& EntitySignals) override;
 
diff --git a/Engine/Plugins/AI/MassAI/Source/MassAIBehavior/Public/Tasks/MassStateTreeSmartObjectTask.h b/Engine/Plugins/AI/MassAI/Source/MassAIBehavior/Public/Tasks/MassStateTreeSmartObjectTask.h
--- a/Engine/Plugins/AI/MassAI/Source/MassAIBehavior/Public/Tasks/MassStateTreeSmartObjectTask.h
+++ b/Engine/Plugins/AI/MassAI/Source/MassAIBehavior/Public/Tasks/MassStateTreeSmartObjectTask.h
@@ -12,7 +12,7 @@ struct FStateTreeExecutionContext;
 struct FMassSmartObjectUserFragment;
 class USmartObjectSubsystem;
 class UMassSignalSubsystem;
-struct FDataFragment_Transform;
+struct FTransformFragment;
 struct FMassMoveTargetFragment;
 
 /**
@@ -75,7 +75,7 @@ protected:
 
 	TStateTreeExternalDataHandle<USmartObjectSubsystem> SmartObjectSubsystemHandle;
 	TStateTreeExternalDataHandle<UMassSignalSubsystem> MassSignalSubsystemHandle;
-	TStateTreeExternalDataHandle<FDataFragment_Transform> EntityTransformHandle;
+	TStateTreeExternalDataHandle<FTransformFragment> EntityTransformHandle;
 	TStateTreeExternalDataHandle<FMassSmartObjectUserFragment> SmartObjectUserHandle;
 	TStateTreeExternalDataHandle<FMassMoveTargetFragment> MoveTargetHandle;
 
diff --git a/Engine/Plugins/AI/MassAI/Source/MassAIBehavior/Public/Tasks/MassZoneGraphPathFollowTask.h b/Engine/Plugins/AI/MassAI/Source/MassAIBehavior/Public/Tasks/MassZoneGraphPathFollowTask.h
--- a/Engine/Plugins/AI/MassAI/Source/MassAIBehavior/Public/Tasks/MassZoneGraphPathFollowTask.h
+++ b/Engine/Plugins/AI/MassAI/Source/MassAIBehavior/Public/Tasks/MassZoneGraphPathFollowTask.h
@@ -15,7 +15,7 @@ struct FMassMoveTargetFragment;
 struct FMassZoneGraphPathRequestFragment;
 struct FMassZoneGraphShortPathFragment;
 struct FMassZoneGraphCachedLaneFragment;
-struct FDataFragment_AgentRadius;
+struct FAgentRadiusFragment;
 struct FMassMovementParameters;
 class UZoneGraphSubsystem;
 
@@ -100,7 +100,7 @@ protected:
 	TStateTreeExternalDataHandle<FMassZoneGraphPathRequestFragment> PathRequestHandle;
 	TStateTreeExternalDataHandle<FMassZoneGraphShortPathFragment> ShortPathHandle;
 	TStateTreeExternalDataHandle<FMassZoneGraphCachedLaneFragment> CachedLaneHandle;
-	TStateTreeExternalDataHandle<FDataFragment_AgentRadius> AgentRadiusHandle;
+	TStateTreeExternalDataHandle<FAgentRadiusFragment> AgentRadiusHandle;
 	TStateTreeExternalDataHandle<FMassMovementParameters> MovementParamsHandle;
 	TStateTreeExternalDataHandle<UZoneGraphSubsystem> ZoneGraphSubsystemHandle;
 
diff --git a/Engine/Plugins/AI/MassAI/Source/MassAIDebug/Private/GameplayDebuggerCategory_Mass.cpp b/Engine/Plugins/AI/MassAI/Source/MassAIDebug/Private/GameplayDebuggerCategory_Mass.cpp
--- a/Engine/Plugins/AI/MassAI/Source/MassAIDebug/Private/GameplayDebuggerCategory_Mass.cpp
+++ b/Engine/Plugins/AI/MassAI/Source/MassAIDebug/Private/GameplayDebuggerCategory_Mass.cpp
@@ -222,7 +222,7 @@ void FGameplayDebuggerCategory_Mass::CollectData(APlayerController* OwnerPC, AAc
 
 		if (CachedEntity.IsSet() && bMarkEntityBeingDebugged)
 		{
-			if (const FDataFragment_Transform* TransformFragment = EntitySystem->GetFragmentDataPtr<FDataFragment_Transform>(CachedEntity))
+			if (const FTransformFragment* TransformFragment = EntitySystem->GetFragmentDataPtr<FTransformFragment>(CachedEntity))
 			{
 				const FVector Location = TransformFragment->GetTransform().GetLocation();
 				AddShape(FGameplayDebuggerShape::MakeBox(Location, FVector(3,3,250), FColor::Purple));
@@ -323,7 +323,7 @@ void FGameplayDebuggerCategory_Mass::CollectData(APlayerController* OwnerPC, AAc
 
 				}
 
-				const FDataFragment_Transform& TransformFragment = EntitySystem->GetFragmentDataChecked<FDataFragment_Transform>(CachedEntity);
+				const FTransformFragment& TransformFragment = EntitySystem->GetFragmentDataChecked<FTransformFragment>(CachedEntity);
 				constexpr float CapsuleRadius = 50.f;
 				AddShape(FGameplayDebuggerShape::MakeCapsule(TransformFragment.GetTransform().GetLocation() + 2.f * CapsuleRadius * FVector::UpVector, CapsuleRadius, CapsuleRadius * 2.f, FColor::Orange));
 			}
@@ -347,8 +347,8 @@ void FGameplayDebuggerCategory_Mass::CollectData(APlayerController* OwnerPC, AAc
 
 		FMassEntityQuery EntityQuery;
 		EntityQuery.AddRequirement<FMassStateTreeFragment>(EMassFragmentAccess::ReadOnly);
-		EntityQuery.AddRequirement<FDataFragment_Transform>(EMassFragmentAccess::ReadOnly);
-		EntityQuery.AddRequirement<FDataFragment_AgentRadius>(EMassFragmentAccess::ReadOnly);
+		EntityQuery.AddRequirement<FTransformFragment>(EMassFragmentAccess::ReadOnly);
+		EntityQuery.AddRequirement<FAgentRadiusFragment>(EMassFragmentAccess::ReadOnly);
 		EntityQuery.AddRequirement<FMassSteeringFragment>(EMassFragmentAccess::ReadOnly);
 		EntityQuery.AddRequirement<FMassStandingSteeringFragment>(EMassFragmentAccess::ReadOnly);
 		EntityQuery.AddRequirement<FMassGhostLocationFragment>(EMassFragmentAccess::ReadOnly);
@@ -372,8 +372,8 @@ void FGameplayDebuggerCategory_Mass::CollectData(APlayerController* OwnerPC, AAc
 			{
 				const int32 NumEntities = Context.GetNumEntities();
 				const TConstArrayView<FMassStateTreeFragment> StateTreeList = Context.GetFragmentView<FMassStateTreeFragment>();
-				const TConstArrayView<FDataFragment_Transform> TransformList = Context.GetFragmentView<FDataFragment_Transform>();
-				const TConstArrayView<FDataFragment_AgentRadius> RadiusList = Context.GetFragmentView<FDataFragment_AgentRadius>();
+				const TConstArrayView<FTransformFragment> TransformList = Context.GetFragmentView<FTransformFragment>();
+				const TConstArrayView<FAgentRadiusFragment> RadiusList = Context.GetFragmentView<FAgentRadiusFragment>();
 				const TConstArrayView<FMassSteeringFragment> SteeringList = Context.GetFragmentView<FMassSteeringFragment>();
 				const TConstArrayView<FMassStandingSteeringFragment> StandingSteeringList = Context.GetFragmentView<FMassStandingSteeringFragment>();
 				const TConstArrayView<FMassGhostLocationFragment> GhostList = Context.GetFragmentView<FMassGhostLocationFragment>();
@@ -394,7 +394,7 @@ void FGameplayDebuggerCategory_Mass::CollectData(APlayerController* OwnerPC, AAc
 
 				for (int32 EntityIndex = 0; EntityIndex < NumEntities; ++EntityIndex)
 				{
-					const FDataFragment_Transform& Transform = TransformList[EntityIndex];
+					const FTransformFragment& Transform = TransformList[EntityIndex];
 					const FVector EntityLocation = Transform.GetTransform().GetLocation();
 					
 					// Cull entities
@@ -410,7 +410,7 @@ void FGameplayDebuggerCategory_Mass::CollectData(APlayerController* OwnerPC, AAc
 						continue;
 					}
 
-					const FDataFragment_AgentRadius& Radius = RadiusList[EntityIndex];
+					const FAgentRadiusFragment& Radius = RadiusList[EntityIndex];
 					const FMassSteeringFragment& Steering = SteeringList[EntityIndex];
 					const FMassStandingSteeringFragment& StandingSteering = StandingSteeringList[EntityIndex];
 					const FMassGhostLocationFragment& Ghost = GhostList[EntityIndex];
@@ -451,7 +451,7 @@ void FGameplayDebuggerCategory_Mass::CollectData(APlayerController* OwnerPC, AAc
 						bool bLookArrowDrawn = false;
 						if (LookAt.LookAtMode == EMassLookAtMode::LookAtEntity && EntitySystem->IsEntityValid(LookAt.TrackedEntity))
 						{
-							if (const FDataFragment_Transform* TargetTransform = EntitySystem->GetFragmentDataPtr<FDataFragment_Transform>(LookAt.TrackedEntity))
+							if (const FTransformFragment* TargetTransform = EntitySystem->GetFragmentDataPtr<FTransformFragment>(LookAt.TrackedEntity))
 							{
 								FVector TargetPosition = TargetTransform->GetTransform().GetLocation();
 								TargetPosition.Z = BasePos.Z;
@@ -465,7 +465,7 @@ void FGameplayDebuggerCategory_Mass::CollectData(APlayerController* OwnerPC, AAc
 
 						if (LookAt.bRandomGazeEntities && EntitySystem->IsEntityValid(LookAt.GazeTrackedEntity))
 						{
-							if (const FDataFragment_Transform* TargetTransform = EntitySystem->GetFragmentDataPtr<FDataFragment_Transform>(LookAt.GazeTrackedEntity))
+							if (const FTransformFragment* TargetTransform = EntitySystem->GetFragmentDataPtr<FTransformFragment>(LookAt.GazeTrackedEntity))
 							{
 								FVector TargetPosition = TargetTransform->GetTransform().GetLocation();
 								TargetPosition.Z = BasePos.Z;
@@ -587,17 +587,17 @@ void FGameplayDebuggerCategory_Mass::CollectData(APlayerController* OwnerPC, AAc
 		{
 			FMassEntityQuery EntityColliderQuery;
 			EntityColliderQuery.AddRequirement<FMassAvoidanceColliderFragment>(EMassFragmentAccess::ReadOnly);
-			EntityColliderQuery.AddRequirement<FDataFragment_Transform>(EMassFragmentAccess::ReadOnly);
+			EntityColliderQuery.AddRequirement<FTransformFragment>(EMassFragmentAccess::ReadOnly);
 			FMassExecutionContext Context(0.f);
 			EntityColliderQuery.ForEachEntityChunk(*EntitySystem, Context, [this, ViewLocation, ViewDirection](const FMassExecutionContext& Context)
 			{
 				const int32 NumEntities = Context.GetNumEntities();
-				const TConstArrayView<FDataFragment_Transform> TransformList = Context.GetFragmentView<FDataFragment_Transform>();
+				const TConstArrayView<FTransformFragment> TransformList = Context.GetFragmentView<FTransformFragment>();
 				const TConstArrayView<FMassAvoidanceColliderFragment> CollidersList = Context.GetFragmentView<FMassAvoidanceColliderFragment>();
 
 				for (int32 EntityIndex = 0; EntityIndex < NumEntities; ++EntityIndex)
 				{
-					const FDataFragment_Transform& Transform = TransformList[EntityIndex];
+					const FTransformFragment& Transform = TransformList[EntityIndex];
 					const FVector EntityLocation = Transform.GetTransform().GetLocation();
 					const FVector EntityForward = Transform.GetTransform().GetRotation().GetForwardVector();
 					
diff --git a/Engine/Plugins/AI/MassAI/Source/MassAIDebug/Private/MassDebugStateTreeProcessor.cpp b/Engine/Plugins/AI/MassAI/Source/MassAIDebug/Private/MassDebugStateTreeProcessor.cpp
--- a/Engine/Plugins/AI/MassAI/Source/MassAIDebug/Private/MassDebugStateTreeProcessor.cpp
+++ b/Engine/Plugins/AI/MassAI/Source/MassAIDebug/Private/MassDebugStateTreeProcessor.cpp
@@ -21,7 +21,7 @@ UMassDebugStateTreeProcessor::UMassDebugStateTreeProcessor()
 void UMassDebugStateTreeProcessor::ConfigureQueries()
 {
 	EntityQuery.AddRequirement<FMassStateTreeFragment>(EMassFragmentAccess::ReadOnly);
-	EntityQuery.AddRequirement<FDataFragment_Transform>(EMassFragmentAccess::ReadOnly);
+	EntityQuery.AddRequirement<FTransformFragment>(EMassFragmentAccess::ReadOnly);
 }
 
 void UMassDebugStateTreeProcessor::Execute(UMassEntitySubsystem& EntitySubsystem, FMassExecutionContext& Context)
@@ -62,7 +62,7 @@ void UMassDebugStateTreeProcessor::Execute(UMassEntitySubsystem& EntitySubsystem
 			const FMassEntityHandle SelectedEntity = Debugger->GetSelectedEntity();
 			const int32 NumEntities = Context.GetNumEntities();
 			const TConstArrayView<FMassStateTreeFragment> StateTreeList = Context.GetFragmentView<FMassStateTreeFragment>();
-			const TConstArrayView<FDataFragment_Transform> TransformList = Context.GetFragmentView<FDataFragment_Transform>();
+			const TConstArrayView<FTransformFragment> TransformList = Context.GetFragmentView<FTransformFragment>();
 			const UStateTree* StateTree = MassStateTreeSubsystem->GetRegisteredStateTreeAsset(StateTreeList[0].StateTreeHandle);
 
 			// Not reporting error since this processor is a debug tool 
@@ -90,7 +90,7 @@ void UMassDebugStateTreeProcessor::Execute(UMassEntitySubsystem& EntitySubsystem
 				const bool bDisplayDebug = UE::Mass::Debug::IsDebuggingEntity(Entity, &EntityColor);
 				if (bDisplayDebug)
 				{
-					const FDataFragment_Transform& Transform = TransformList[i];
+					const FTransformFragment& Transform = TransformList[i];
 					
 					const FVector ZOffset(0,0,50);
 					const FVector Position = Transform.GetTransform().GetLocation() + ZOffset;
diff --git a/Engine/Plugins/AI/MassAI/Source/MassAIReplication/Private/MassReplicationPathHandlers.cpp b/Engine/Plugins/AI/MassAI/Source/MassAIReplication/Private/MassReplicationPathHandlers.cpp
--- a/Engine/Plugins/AI/MassAI/Source/MassAIReplication/Private/MassReplicationPathHandlers.cpp
+++ b/Engine/Plugins/AI/MassAI/Source/MassAIReplication/Private/MassReplicationPathHandlers.cpp
@@ -117,7 +117,7 @@ void FReplicatedAgentPathData::ApplyToEntity(const UWorld& InWorld, const FMassE
 	FMassZoneGraphShortPathFragment& ShortPath = InEntityView.GetFragmentData<FMassZoneGraphShortPathFragment>();
 	FMassZoneGraphCachedLaneFragment& CachedLane = InEntityView.GetFragmentData<FMassZoneGraphCachedLaneFragment>();
 	FMassZoneGraphLaneLocationFragment& LaneLocation = InEntityView.GetFragmentData<FMassZoneGraphLaneLocationFragment>();
-	const FDataFragment_AgentRadius& AgentRadius = InEntityView.GetFragmentData<FDataFragment_AgentRadius>();
+	const FAgentRadiusFragment& AgentRadius = InEntityView.GetFragmentData<FAgentRadiusFragment>();
 
 	MoveTarget.CreateReplicatedAction(Action, ActionID, InWorld.GetTimeSeconds(), ActionServerStartTime);
 	MoveTarget.DesiredSpeed = DesiredSpeed;
diff --git a/Engine/Plugins/AI/MassAI/Source/MassNavigation/Private/MassAvoidanceProcessors.cpp b/Engine/Plugins/AI/MassAI/Source/MassNavigation/Private/MassAvoidanceProcessors.cpp
--- a/Engine/Plugins/AI/MassAI/Source/MassNavigation/Private/MassAvoidanceProcessors.cpp
+++ b/Engine/Plugins/AI/MassAI/Source/MassNavigation/Private/MassAvoidanceProcessors.cpp
@@ -357,9 +357,9 @@ void UMassMovingAvoidanceProcessor::ConfigureQueries()
 	EntityQuery.AddRequirement<FMassForceFragment>(EMassFragmentAccess::ReadWrite);
 	EntityQuery.AddRequirement<FMassNavigationEdgesFragment>(EMassFragmentAccess::ReadOnly);
 	EntityQuery.AddRequirement<FMassMoveTargetFragment>(EMassFragmentAccess::ReadOnly);
-	EntityQuery.AddRequirement<FDataFragment_Transform>(EMassFragmentAccess::ReadOnly);
+	EntityQuery.AddRequirement<FTransformFragment>(EMassFragmentAccess::ReadOnly);
 	EntityQuery.AddRequirement<FMassVelocityFragment>(EMassFragmentAccess::ReadOnly);
-	EntityQuery.AddRequirement<FDataFragment_AgentRadius>(EMassFragmentAccess::ReadOnly);
+	EntityQuery.AddRequirement<FAgentRadiusFragment>(EMassFragmentAccess::ReadOnly);
 	EntityQuery.AddTagRequirement<FMassMediumLODTag>(EMassFragmentPresence::None);
 	EntityQuery.AddTagRequirement<FMassLowLODTag>(EMassFragmentPresence::None);
 	EntityQuery.AddTagRequirement<FMassOffLODTag>(EMassFragmentPresence::None);
@@ -392,9 +392,9 @@ void UMassMovingAvoidanceProcessor::Execute(UMassEntitySubsystem& EntitySubsyste
 		
 		const TArrayView<FMassForceFragment> ForceList = Context.GetMutableFragmentView<FMassForceFragment>();
 		const TConstArrayView<FMassNavigationEdgesFragment> NavEdgesList = Context.GetFragmentView<FMassNavigationEdgesFragment>();
-		const TConstArrayView<FDataFragment_Transform> LocationList = Context.GetFragmentView<FDataFragment_Transform>();
+		const TConstArrayView<FTransformFragment> LocationList = Context.GetFragmentView<FTransformFragment>();
 		const TConstArrayView<FMassVelocityFragment> VelocityList = Context.GetFragmentView<FMassVelocityFragment>();
-		const TConstArrayView<FDataFragment_AgentRadius> RadiusList = Context.GetFragmentView<FDataFragment_AgentRadius>();
+		const TConstArrayView<FAgentRadiusFragment> RadiusList = Context.GetFragmentView<FAgentRadiusFragment>();
 		const TConstArrayView<FMassMoveTargetFragment> MoveTargetList = Context.GetFragmentView<FMassMoveTargetFragment>();
 		const FMassMovingAvoidanceParameters& MovingAvoidanceParams = Context.GetConstSharedFragment<FMassMovingAvoidanceParameters>();
 		const FMassMovementParameters& MovementParams = Context.GetConstSharedFragment<FMassMovementParameters>();
@@ -445,9 +445,9 @@ void UMassMovingAvoidanceProcessor::Execute(UMassEntitySubsystem& EntitySubsyste
 
 			FMassEntityHandle Entity = Context.GetEntity(EntityIndex);
 			const FMassNavigationEdgesFragment& NavEdges = NavEdgesList[EntityIndex];
-			const FDataFragment_Transform& Location = LocationList[EntityIndex];
+			const FTransformFragment& Location = LocationList[EntityIndex];
 			const FMassVelocityFragment& Velocity = VelocityList[EntityIndex];
-			const FDataFragment_AgentRadius& Radius = RadiusList[EntityIndex];
+			const FAgentRadiusFragment& Radius = RadiusList[EntityIndex];
 			FMassForceFragment& Force = ForceList[EntityIndex];
 
 			// Smaller steering max accel makes the steering more "calm" but less opportunistic, may not find solution, or gets stuck.
@@ -745,7 +745,7 @@ void UMassMovingAvoidanceProcessor::Execute(UMassEntitySubsystem& EntitySubsyste
 				}
 				
 				// Skip too far
-				const FTransform& Transform = EntitySubsystem.GetFragmentDataChecked<FDataFragment_Transform>(OtherEntity.Entity).GetTransform();
+				const FTransform& Transform = EntitySubsystem.GetFragmentDataChecked<FTransformFragment>(OtherEntity.Entity).GetTransform();
 				const FVector OtherLocation = Transform.GetLocation();
 				
 				const float SqDist = FVector::DistSquared(AgentLocation, OtherLocation);
@@ -832,7 +832,7 @@ void UMassMovingAvoidanceProcessor::Execute(UMassEntitySubsystem& EntitySubsyste
 					FCollider& Collider = Colliders.AddDefaulted_GetRef();
 					Collider.Location = Obstacle.LocationCached;
 					Collider.Velocity = OtherVelocity;
-					Collider.Radius = OtherEntityView.GetFragmentData<FDataFragment_AgentRadius>().Radius;
+					Collider.Radius = OtherEntityView.GetFragmentData<FAgentRadiusFragment>().Radius;
 					Collider.bCanAvoid = bCanAvoid;
 					Collider.bIsMoving = bOtherIsMoving;
 				}
@@ -1020,8 +1020,8 @@ void UMassStandingAvoidanceProcessor::ConfigureQueries()
 	EntityQuery.AddRequirement<FMassGhostLocationFragment>(EMassFragmentAccess::ReadWrite);
 	EntityQuery.AddRequirement<FMassNavigationEdgesFragment>(EMassFragmentAccess::ReadOnly);
 	EntityQuery.AddRequirement<FMassMoveTargetFragment>(EMassFragmentAccess::ReadOnly);
-	EntityQuery.AddRequirement<FDataFragment_Transform>(EMassFragmentAccess::ReadOnly);
-	EntityQuery.AddRequirement<FDataFragment_AgentRadius>(EMassFragmentAccess::ReadOnly);
+	EntityQuery.AddRequirement<FTransformFragment>(EMassFragmentAccess::ReadOnly);
+	EntityQuery.AddRequirement<FAgentRadiusFragment>(EMassFragmentAccess::ReadOnly);
 	EntityQuery.AddTagRequirement<FMassMediumLODTag>(EMassFragmentPresence::None);
 	EntityQuery.AddTagRequirement<FMassLowLODTag>(EMassFragmentPresence::None);
 	EntityQuery.AddTagRequirement<FMassOffLODTag>(EMassFragmentPresence::None);
@@ -1052,8 +1052,8 @@ void UMassStandingAvoidanceProcessor::Execute(UMassEntitySubsystem& EntitySubsys
 		const float DeltaTime = Context.GetDeltaTimeSeconds();
 
 		const TArrayView<FMassGhostLocationFragment> GhostList = Context.GetMutableFragmentView<FMassGhostLocationFragment>();
-		const TConstArrayView<FDataFragment_Transform> LocationList = Context.GetFragmentView<FDataFragment_Transform>();
-		const TConstArrayView<FDataFragment_AgentRadius> RadiusList = Context.GetFragmentView<FDataFragment_AgentRadius>();
+		const TConstArrayView<FTransformFragment> LocationList = Context.GetFragmentView<FTransformFragment>();
+		const TConstArrayView<FAgentRadiusFragment> RadiusList = Context.GetFragmentView<FAgentRadiusFragment>();
 		const TConstArrayView<FMassMoveTargetFragment> MoveTargetList = Context.GetFragmentView<FMassMoveTargetFragment>();
 		const FMassStandingAvoidanceParameters& StandingParams = Context.GetConstSharedFragment<FMassStandingAvoidanceParameters>();
 
@@ -1094,8 +1094,8 @@ void UMassStandingAvoidanceProcessor::Execute(UMassEntitySubsystem& EntitySubsys
 				continue;
 			}
 
-			const FDataFragment_Transform& Location = LocationList[EntityIndex];
-			const FDataFragment_AgentRadius& Radius = RadiusList[EntityIndex];
+			const FTransformFragment& Location = LocationList[EntityIndex];
+			const FAgentRadiusFragment& Radius = RadiusList[EntityIndex];
 
 			FMassEntityHandle Entity = Context.GetEntity(EntityIndex);
 			const FVector AgentLocation = Location.GetTransform().GetTranslation();
@@ -1143,7 +1143,7 @@ void UMassStandingAvoidanceProcessor::Execute(UMassEntitySubsystem& EntitySubsys
 				}
 
 				// Skip too far
-				const FDataFragment_Transform& OtherTransform = EntitySubsystem.GetFragmentDataChecked<FDataFragment_Transform>(OtherEntity.Entity);
+				const FTransformFragment& OtherTransform = EntitySubsystem.GetFragmentDataChecked<FTransformFragment>(OtherEntity.Entity);
 				const FVector OtherLocation = OtherTransform.GetTransform().GetLocation();
 				const float DistSq = FVector::DistSquared(AgentLocation, OtherLocation);
 				if (DistSq > DistanceCutOffSqr)
@@ -1165,7 +1165,7 @@ void UMassStandingAvoidanceProcessor::Execute(UMassEntitySubsystem& EntitySubsys
 				FSortedObstacle& OtherAgent = ClosestObstacles[Index];
 				FMassEntityView OtherEntityView(EntitySubsystem, OtherAgent.Entity);
 
-				const float OtherRadius = OtherEntityView.GetFragmentData<FDataFragment_AgentRadius>().Radius;
+				const float OtherRadius = OtherEntityView.GetFragmentData<FAgentRadiusFragment>().Radius;
 				const float TotalRadius = GhostRadius + OtherRadius;
 
 				// @todo: this is heavy fragment to access, see if we could handle this differently.
diff --git a/Engine/Plugins/AI/MassAI/Source/MassNavigation/Private/MassAvoidanceTrait.cpp b/Engine/Plugins/AI/MassAI/Source/MassNavigation/Private/MassAvoidanceTrait.cpp
--- a/Engine/Plugins/AI/MassAI/Source/MassNavigation/Private/MassAvoidanceTrait.cpp
+++ b/Engine/Plugins/AI/MassAI/Source/MassNavigation/Private/MassAvoidanceTrait.cpp
@@ -13,9 +13,9 @@ void UMassObstacleAvoidanceTrait::BuildTemplate(FMassEntityTemplateBuildContext&
 	UMassEntitySubsystem* EntitySubsystem = UWorld::GetSubsystem<UMassEntitySubsystem>(&World);
 	check(EntitySubsystem);
 
-	BuildContext.AddFragment<FDataFragment_AgentRadius>();
+	BuildContext.AddFragment<FAgentRadiusFragment>();
 	BuildContext.AddFragment<FMassNavigationEdgesFragment>();
-	BuildContext.AddFragment<FDataFragment_Transform>();
+	BuildContext.AddFragment<FTransformFragment>();
 	BuildContext.AddFragment<FMassVelocityFragment>();
 	BuildContext.AddFragment<FMassForceFragment>();
 	BuildContext.AddFragment<FMassMoveTargetFragment>();
diff --git a/Engine/Plugins/AI/MassAI/Source/MassNavigation/Private/MassNavigationObstacleTrait.cpp b/Engine/Plugins/AI/MassAI/Source/MassNavigation/Private/MassNavigationObstacleTrait.cpp
--- a/Engine/Plugins/AI/MassAI/Source/MassNavigation/Private/MassNavigationObstacleTrait.cpp
+++ b/Engine/Plugins/AI/MassAI/Source/MassNavigation/Private/MassNavigationObstacleTrait.cpp
@@ -6,7 +6,7 @@
 
 void UMassNavigationObstacleTrait::BuildTemplate(FMassEntityTemplateBuildContext& BuildContext, UWorld& World) const
 {
-	BuildContext.AddFragment<FDataFragment_AgentRadius>();
+	BuildContext.AddFragment<FAgentRadiusFragment>();
 
 	BuildContext.AddFragment<FMassNavigationObstacleGridCellLocationFragment>();
 }
diff --git a/Engine/Plugins/AI/MassAI/Source/MassNavigation/Private/MassNavigationProcessors.cpp b/Engine/Plugins/AI/MassAI/Source/MassNavigation/Private/MassNavigationProcessors.cpp
--- a/Engine/Plugins/AI/MassAI/Source/MassNavigation/Private/MassNavigationProcessors.cpp
+++ b/Engine/Plugins/AI/MassAI/Source/MassNavigation/Private/MassNavigationProcessors.cpp
@@ -29,7 +29,7 @@ UMassOffLODNavigationProcessor::UMassOffLODNavigationProcessor()
 
 void UMassOffLODNavigationProcessor::ConfigureQueries()
 {
-	EntityQuery_Conditional.AddRequirement<FDataFragment_Transform>(EMassFragmentAccess::ReadWrite);
+	EntityQuery_Conditional.AddRequirement<FTransformFragment>(EMassFragmentAccess::ReadWrite);
 	EntityQuery_Conditional.AddRequirement<FMassMoveTargetFragment>(EMassFragmentAccess::ReadOnly);
 	EntityQuery_Conditional.AddTagRequirement<FMassOffLODTag>(EMassFragmentPresence::All);
 	EntityQuery_Conditional.AddChunkRequirement<FMassSimulationVariableTickChunkFragment>(EMassFragmentAccess::ReadOnly, EMassFragmentPresence::Optional);
@@ -49,7 +49,7 @@ void UMassOffLODNavigationProcessor::Execute(UMassEntitySubsystem& EntitySubsyst
 #endif // WITH_MASSGAMEPLAY_DEBUG
 		const int32 NumEntities = Context.GetNumEntities();
 
-		const TArrayView<FDataFragment_Transform> LocationList = Context.GetMutableFragmentView<FDataFragment_Transform>();
+		const TArrayView<FTransformFragment> LocationList = Context.GetMutableFragmentView<FTransformFragment>();
 		const TConstArrayView<FMassMoveTargetFragment> MoveTargetList = Context.GetFragmentView<FMassMoveTargetFragment>();
 
 		for (int32 EntityIndex = 0; EntityIndex < NumEntities; ++EntityIndex)
@@ -76,7 +76,7 @@ UMassNavigationSmoothHeightProcessor::UMassNavigationSmoothHeightProcessor()
 
 void UMassNavigationSmoothHeightProcessor::ConfigureQueries()
 {
-	EntityQuery.AddRequirement<FDataFragment_Transform>(EMassFragmentAccess::ReadWrite);
+	EntityQuery.AddRequirement<FTransformFragment>(EMassFragmentAccess::ReadWrite);
 	EntityQuery.AddRequirement<FMassMoveTargetFragment>(EMassFragmentAccess::ReadOnly);
 	EntityQuery.AddTagRequirement<FMassOffLODTag>(EMassFragmentPresence::None);
 	EntityQuery.AddConstSharedRequirement<FMassMovementParameters>(EMassFragmentPresence::All);
@@ -97,7 +97,7 @@ void UMassNavigationSmoothHeightProcessor::Execute(UMassEntitySubsystem& EntityS
 		const float DeltaTime = Context.GetDeltaTimeSeconds();
 
 		const FMassMovementParameters& MovementParams = Context.GetConstSharedFragment<FMassMovementParameters>();
-		const TArrayView<FDataFragment_Transform> LocationList = Context.GetMutableFragmentView<FDataFragment_Transform>();
+		const TArrayView<FTransformFragment> LocationList = Context.GetMutableFragmentView<FTransformFragment>();
 		const TConstArrayView<FMassMoveTargetFragment> MoveTargetList = Context.GetFragmentView<FMassMoveTargetFragment>();
 
 		for (int32 EntityIndex = 0; EntityIndex < NumEntities; ++EntityIndex)
@@ -129,7 +129,7 @@ UMassMoveTargetFragmentInitializer::UMassMoveTargetFragmentInitializer()
 void UMassMoveTargetFragmentInitializer::ConfigureQueries()
 {
 	InitializerQuery.AddRequirement<FMassMoveTargetFragment>(EMassFragmentAccess::ReadWrite);
-	InitializerQuery.AddRequirement<FDataFragment_Transform>(EMassFragmentAccess::ReadOnly);
+	InitializerQuery.AddRequirement<FTransformFragment>(EMassFragmentAccess::ReadOnly);
 }
 
 void UMassMoveTargetFragmentInitializer::Execute(UMassEntitySubsystem& EntitySubsystem, FMassExecutionContext& Context)
@@ -138,12 +138,12 @@ void UMassMoveTargetFragmentInitializer::Execute(UMassEntitySubsystem& EntitySub
 	{
 		const int32 NumEntities = Context.GetNumEntities();
 		const TArrayView<FMassMoveTargetFragment> MoveTargetList = Context.GetMutableFragmentView<FMassMoveTargetFragment>();
-		const TConstArrayView<FDataFragment_Transform> LocationList = Context.GetFragmentView<FDataFragment_Transform>();
+		const TConstArrayView<FTransformFragment> LocationList = Context.GetFragmentView<FTransformFragment>();
 
 		for (int32 EntityIndex = 0; EntityIndex < NumEntities; ++EntityIndex)
 		{
 			FMassMoveTargetFragment& MoveTarget = MoveTargetList[EntityIndex];
-			const FDataFragment_Transform& Location = LocationList[EntityIndex];
+			const FTransformFragment& Location = LocationList[EntityIndex];
 
 			MoveTarget.Center = Location.GetTransform().GetLocation();
 			MoveTarget.Forward = Location.GetTransform().GetRotation().Vector();
@@ -166,8 +166,8 @@ UMassNavigationObstacleGridProcessor::UMassNavigationObstacleGridProcessor()
 void UMassNavigationObstacleGridProcessor::ConfigureQueries()
 {
 	FMassEntityQuery BaseEntityQuery;
-	BaseEntityQuery.AddRequirement<FDataFragment_Transform>(EMassFragmentAccess::ReadOnly);
-	BaseEntityQuery.AddRequirement<FDataFragment_AgentRadius>(EMassFragmentAccess::ReadOnly);
+	BaseEntityQuery.AddRequirement<FTransformFragment>(EMassFragmentAccess::ReadOnly);
+	BaseEntityQuery.AddRequirement<FAgentRadiusFragment>(EMassFragmentAccess::ReadOnly);
 	BaseEntityQuery.AddRequirement<FMassNavigationObstacleGridCellLocationFragment>(EMassFragmentAccess::ReadWrite);
 
 	AddToGridEntityQuery = BaseEntityQuery;
@@ -204,8 +204,8 @@ void UMassNavigationObstacleGridProcessor::Execute(UMassEntitySubsystem& EntityS
 	{
 		const int32 NumEntities = Context.GetNumEntities();
 
-		TConstArrayView<FDataFragment_Transform> LocationList = Context.GetFragmentView<FDataFragment_Transform>();
-		TConstArrayView<FDataFragment_AgentRadius> RadiiList = Context.GetFragmentView<FDataFragment_AgentRadius>();
+		TConstArrayView<FTransformFragment> LocationList = Context.GetFragmentView<FTransformFragment>();
+		TConstArrayView<FAgentRadiusFragment> RadiiList = Context.GetFragmentView<FAgentRadiusFragment>();
 		TArrayView<FMassNavigationObstacleGridCellLocationFragment> NavigationObstacleCellLocationList = Context.GetMutableFragmentView<FMassNavigationObstacleGridCellLocationFragment>();
 		const bool bHasColliderData = Context.GetFragmentView<FMassAvoidanceColliderFragment>().Num() > 0;
 
@@ -230,8 +230,8 @@ void UMassNavigationObstacleGridProcessor::Execute(UMassEntitySubsystem& EntityS
 	{
 		const int32 NumEntities = Context.GetNumEntities();
 
-		TConstArrayView<FDataFragment_Transform> LocationList = Context.GetFragmentView<FDataFragment_Transform>();
-		TConstArrayView<FDataFragment_AgentRadius> RadiiList = Context.GetFragmentView<FDataFragment_AgentRadius>();
+		TConstArrayView<FTransformFragment> LocationList = Context.GetFragmentView<FTransformFragment>();
+		TConstArrayView<FAgentRadiusFragment> RadiiList = Context.GetFragmentView<FAgentRadiusFragment>();
 		TArrayView<FMassNavigationObstacleGridCellLocationFragment> NavigationObstacleCellLocationList = Context.GetMutableFragmentView<FMassNavigationObstacleGridCellLocationFragment>();
 		const bool bHasColliderData = Context.GetFragmentView<FMassAvoidanceColliderFragment>().Num() > 0;
 
diff --git a/Engine/Plugins/AI/MassAI/Source/MassNavigation/Private/MassSmoothOrientationProcessors.cpp b/Engine/Plugins/AI/MassAI/Source/MassNavigation/Private/MassSmoothOrientationProcessors.cpp
--- a/Engine/Plugins/AI/MassAI/Source/MassNavigation/Private/MassSmoothOrientationProcessors.cpp
+++ b/Engine/Plugins/AI/MassAI/Source/MassNavigation/Private/MassSmoothOrientationProcessors.cpp
@@ -26,11 +26,11 @@ void UMassSmoothOrientationProcessor::ConfigureQueries()
 {
 	HighResEntityQuery.AddRequirement<FMassMoveTargetFragment>(EMassFragmentAccess::ReadOnly);
 	HighResEntityQuery.AddRequirement<FMassVelocityFragment>(EMassFragmentAccess::ReadWrite);
-	HighResEntityQuery.AddRequirement<FDataFragment_Transform>(EMassFragmentAccess::ReadWrite);
+	HighResEntityQuery.AddRequirement<FTransformFragment>(EMassFragmentAccess::ReadWrite);
 	HighResEntityQuery.AddTagRequirement<FMassOffLODTag>(EMassFragmentPresence::None);
 	HighResEntityQuery.AddConstSharedRequirement<FMassSmoothOrientationParameters>(EMassFragmentPresence::All);
 
-	LowResEntityQuery_Conditional.AddRequirement<FDataFragment_Transform>(EMassFragmentAccess::ReadWrite);
+	LowResEntityQuery_Conditional.AddRequirement<FTransformFragment>(EMassFragmentAccess::ReadWrite);
 	LowResEntityQuery_Conditional.AddRequirement<FMassMoveTargetFragment>(EMassFragmentAccess::ReadOnly);
 	LowResEntityQuery_Conditional.AddTagRequirement<FMassOffLODTag>(EMassFragmentPresence::All);
 	LowResEntityQuery_Conditional.AddChunkRequirement<FMassSimulationVariableTickChunkFragment>(EMassFragmentAccess::ReadOnly, EMassFragmentPresence::Optional);
@@ -53,7 +53,7 @@ void UMassSmoothOrientationProcessor::Execute(UMassEntitySubsystem& EntitySubsys
 			const FMassSmoothOrientationParameters& OrientationParams = Context.GetConstSharedFragment<FMassSmoothOrientationParameters>();
 
 			const TConstArrayView<FMassMoveTargetFragment> MoveTargetList = Context.GetFragmentView<FMassMoveTargetFragment>();
-			const TArrayView<FDataFragment_Transform> LocationList = Context.GetMutableFragmentView<FDataFragment_Transform>();
+			const TArrayView<FTransformFragment> LocationList = Context.GetMutableFragmentView<FTransformFragment>();
 			const TArrayView<FMassVelocityFragment> VelocityList = Context.GetMutableFragmentView<FMassVelocityFragment>();
 
 			for (int32 EntityIndex = 0; EntityIndex < NumEntities; ++EntityIndex)
@@ -119,7 +119,7 @@ void UMassSmoothOrientationProcessor::Execute(UMassEntitySubsystem& EntitySubsys
 		{
 			const int32 NumEntities = Context.GetNumEntities();
 
-			const TArrayView<FDataFragment_Transform> LocationList = Context.GetMutableFragmentView<FDataFragment_Transform>();
+			const TArrayView<FTransformFragment> LocationList = Context.GetMutableFragmentView<FTransformFragment>();
 			const TConstArrayView<FMassMoveTargetFragment> MoveTargetList = Context.GetFragmentView<FMassMoveTargetFragment>();
 
 			for (int32 EntityIndex = 0; EntityIndex < NumEntities; ++EntityIndex)
diff --git a/Engine/Plugins/AI/MassAI/Source/MassNavigation/Private/MassSmoothOrientationTrait.cpp b/Engine/Plugins/AI/MassAI/Source/MassNavigation/Private/MassSmoothOrientationTrait.cpp
--- a/Engine/Plugins/AI/MassAI/Source/MassNavigation/Private/MassSmoothOrientationTrait.cpp
+++ b/Engine/Plugins/AI/MassAI/Source/MassNavigation/Private/MassSmoothOrientationTrait.cpp
@@ -13,7 +13,7 @@ void UMassSmoothOrientationTrait::BuildTemplate(FMassEntityTemplateBuildContext&
 
 	BuildContext.AddFragment<FMassMoveTargetFragment>();
 	BuildContext.AddFragment<FMassVelocityFragment>();
-	BuildContext.AddFragment<FDataFragment_Transform>();
+	BuildContext.AddFragment<FTransformFragment>();
 
 	const FConstSharedStruct OrientationFragment = EntitySubsystem->GetOrCreateConstSharedFragment(UE::StructUtils::GetStructCrc32(FConstStructView::Make(Orientation)), Orientation);
 	BuildContext.AddConstSharedFragment(OrientationFragment);
diff --git a/Engine/Plugins/AI/MassAI/Source/MassNavigation/Private/MassSteeringProcessors.cpp b/Engine/Plugins/AI/MassAI/Source/MassNavigation/Private/MassSteeringProcessors.cpp
--- a/Engine/Plugins/AI/MassAI/Source/MassNavigation/Private/MassSteeringProcessors.cpp
+++ b/Engine/Plugins/AI/MassAI/Source/MassNavigation/Private/MassSteeringProcessors.cpp
@@ -59,7 +59,7 @@ UMassSteerToMoveTargetProcessor::UMassSteerToMoveTargetProcessor()
 
 void UMassSteerToMoveTargetProcessor::ConfigureQueries()
 {
-	EntityQuery.AddRequirement<FDataFragment_Transform>(EMassFragmentAccess::ReadOnly);
+	EntityQuery.AddRequirement<FTransformFragment>(EMassFragmentAccess::ReadOnly);
 	EntityQuery.AddRequirement<FMassMoveTargetFragment>(EMassFragmentAccess::ReadWrite);
 	EntityQuery.AddRequirement<FMassSteeringFragment>(EMassFragmentAccess::ReadWrite);
 	EntityQuery.AddRequirement<FMassStandingSteeringFragment>(EMassFragmentAccess::ReadWrite);
@@ -79,7 +79,7 @@ void UMassSteerToMoveTargetProcessor::Execute(UMassEntitySubsystem& EntitySubsys
 	EntityQuery.ForEachEntityChunk(EntitySubsystem, Context, [this](FMassExecutionContext& Context)
 	{
 		const int32 NumEntities = Context.GetNumEntities();
-		const TConstArrayView<FDataFragment_Transform> TransformList = Context.GetFragmentView<FDataFragment_Transform>();
+		const TConstArrayView<FTransformFragment> TransformList = Context.GetFragmentView<FTransformFragment>();
 		const TArrayView<FMassMoveTargetFragment> MoveTargetList = Context.GetMutableFragmentView<FMassMoveTargetFragment>();
 		const TArrayView<FMassVelocityFragment> VelocityList = Context.GetMutableFragmentView<FMassVelocityFragment>();
 		const TArrayView<FMassForceFragment> ForceList = Context.GetMutableFragmentView<FMassForceFragment>();
@@ -95,7 +95,7 @@ void UMassSteerToMoveTargetProcessor::Execute(UMassEntitySubsystem& EntitySubsys
 
 		for (int32 EntityIndex = 0; EntityIndex < NumEntities; ++EntityIndex)
 		{
-			const FDataFragment_Transform& TransformFragment = TransformList[EntityIndex];
+			const FTransformFragment& TransformFragment = TransformList[EntityIndex];
 			FMassSteeringFragment& Steering = SteeringList[EntityIndex];
 			FMassStandingSteeringFragment& StandingSteering = StandingSteeringList[EntityIndex];
 			FMassGhostLocationFragment& Ghost = GhostList[EntityIndex];
diff --git a/Engine/Plugins/AI/MassAI/Source/MassNavigation/Private/MassSteeringTrait.cpp b/Engine/Plugins/AI/MassAI/Source/MassNavigation/Private/MassSteeringTrait.cpp
--- a/Engine/Plugins/AI/MassAI/Source/MassNavigation/Private/MassSteeringTrait.cpp
+++ b/Engine/Plugins/AI/MassAI/Source/MassNavigation/Private/MassSteeringTrait.cpp
@@ -13,8 +13,8 @@ void UMassSteeringTrait::BuildTemplate(FMassEntityTemplateBuildContext& BuildCon
 	UMassEntitySubsystem* EntitySubsystem = UWorld::GetSubsystem<UMassEntitySubsystem>(&World);
 	check(EntitySubsystem);
 
-	BuildContext.AddFragment<FDataFragment_AgentRadius>();
-	BuildContext.AddFragment<FDataFragment_Transform>();
+	BuildContext.AddFragment<FAgentRadiusFragment>();
+	BuildContext.AddFragment<FTransformFragment>();
 	BuildContext.AddFragment<FMassVelocityFragment>();
 	BuildContext.AddFragment<FMassForceFragment>();
 	BuildContext.AddFragment<FMassMoveTargetFragment>();
diff --git a/Engine/Plugins/AI/MassAI/Source/MassZoneGraphNavigation/Private/MassZoneGraphNavigationProcessors.cpp b/Engine/Plugins/AI/MassAI/Source/MassZoneGraphNavigation/Private/MassZoneGraphNavigationProcessors.cpp
--- a/Engine/Plugins/AI/MassAI/Source/MassZoneGraphNavigation/Private/MassZoneGraphNavigationProcessors.cpp
+++ b/Engine/Plugins/AI/MassAI/Source/MassZoneGraphNavigation/Private/MassZoneGraphNavigationProcessors.cpp
@@ -45,7 +45,7 @@ UMassZoneGraphLocationInitializer::UMassZoneGraphLocationInitializer()
 void UMassZoneGraphLocationInitializer::ConfigureQueries()
 {
 	EntityQuery.AddRequirement<FMassZoneGraphLaneLocationFragment>(EMassFragmentAccess::ReadWrite);
-	EntityQuery.AddRequirement<FDataFragment_Transform>(EMassFragmentAccess::ReadOnly);
+	EntityQuery.AddRequirement<FTransformFragment>(EMassFragmentAccess::ReadOnly);
 	EntityQuery.AddRequirement<FMassMoveTargetFragment>(EMassFragmentAccess::ReadWrite); // Make optional?
 	EntityQuery.AddConstSharedRequirement<FMassZoneGraphNavigationParameters>(EMassFragmentPresence::All);
 }
@@ -69,12 +69,12 @@ void UMassZoneGraphLocationInitializer::Execute(UMassEntitySubsystem& EntitySubs
 		const int32 NumEntities = Context.GetNumEntities();
 		const TArrayView<FMassZoneGraphLaneLocationFragment> LaneLocationList = Context.GetMutableFragmentView<FMassZoneGraphLaneLocationFragment>();
 		const TArrayView<FMassMoveTargetFragment> MoveTargetList = Context.GetMutableFragmentView<FMassMoveTargetFragment>();
-		const TConstArrayView<FDataFragment_Transform> TransformList = Context.GetFragmentView<FDataFragment_Transform>();
+		const TConstArrayView<FTransformFragment> TransformList = Context.GetFragmentView<FTransformFragment>();
 		const FMassZoneGraphNavigationParameters& NavigationParams = Context.GetConstSharedFragment<FMassZoneGraphNavigationParameters>();
 
 		for (int32 EntityIndex = 0; EntityIndex < NumEntities; ++EntityIndex)
 		{
-			const FDataFragment_Transform& Transform = TransformList[EntityIndex];
+			const FTransformFragment& Transform = TransformList[EntityIndex];
 			const FVector& AgentLocation = Transform.GetTransform().GetLocation();
 			FMassMoveTargetFragment& MoveTarget = MoveTargetList[EntityIndex];
 			FMassZoneGraphLaneLocationFragment& LaneLocation = LaneLocationList[EntityIndex];
diff --git a/Engine/Plugins/AI/MassAI/Source/MassZoneGraphNavigation/Private/MassZoneGraphNavigationTrait.cpp b/Engine/Plugins/AI/MassAI/Source/MassZoneGraphNavigation/Private/MassZoneGraphNavigationTrait.cpp
--- a/Engine/Plugins/AI/MassAI/Source/MassZoneGraphNavigation/Private/MassZoneGraphNavigationTrait.cpp
+++ b/Engine/Plugins/AI/MassAI/Source/MassZoneGraphNavigation/Private/MassZoneGraphNavigationTrait.cpp
@@ -12,8 +12,8 @@ void UMassZoneGraphNavigationTrait::BuildTemplate(FMassEntityTemplateBuildContex
 	UMassEntitySubsystem* EntitySubsystem = UWorld::GetSubsystem<UMassEntitySubsystem>(&World);
 	check(EntitySubsystem);
 
-	BuildContext.AddFragment<FDataFragment_AgentRadius>();
-	BuildContext.AddFragment<FDataFragment_Transform>();
+	BuildContext.AddFragment<FAgentRadiusFragment>();
+	BuildContext.AddFragment<FTransformFragment>();
 	BuildContext.AddFragment<FMassVelocityFragment>();
 	BuildContext.AddFragment<FMassMoveTargetFragment>();
 
diff --git a/Engine/Plugins/AI/MassCrowd/Config/BaseMassCrowd.ini b/Engine/Plugins/AI/MassCrowd/Config/BaseMassCrowd.ini
--- a/Engine/Plugins/AI/MassCrowd/Config/BaseMassCrowd.ini
+++ b/Engine/Plugins/AI/MassCrowd/Config/BaseMassCrowd.ini
@@ -1,3 +1,6 @@
 [CoreRedirects]
 +PackageRedirects=(OldName="/Script/MassCrowdAI5", NewName="/Script/MassCrowd")
-+StructRedirects=(OldName="/Script/MassAIBehavior.MassZoneGraphFindWanderTarget",NewName="/Script/MassCrowd.MassZoneGraphFindWanderTarget")
\ No newline at end of file
++StructRedirects=(OldName="/Script/MassAIBehavior.MassZoneGraphFindWanderTarget",NewName="/Script/MassCrowd.MassZoneGraphFindWanderTarget")
+
++StructRedirects=(OldName="TagFragment_MassCrowd",NewName="MassCrowdTag")
++StructRedirects=(OldName="/Script/MassCrowd.TagFragment_MassCrowd",NewName="/Script/MassCrowd.MassCrowdTag")
diff --git a/Engine/Plugins/AI/MassCrowd/Source/MassCrowd/Private/MassCrowdBubble.cpp b/Engine/Plugins/AI/MassCrowd/Source/MassCrowd/Private/MassCrowdBubble.cpp
--- a/Engine/Plugins/AI/MassCrowd/Source/MassCrowd/Private/MassCrowdBubble.cpp
+++ b/Engine/Plugins/AI/MassCrowd/Source/MassCrowd/Private/MassCrowdBubble.cpp
@@ -24,7 +24,7 @@ namespace UE::Mass::Crowd
 
 		const FMassEntityView EntityView(EntitySystem, Entity);
 
-		const FDataFragment_Transform& TransformFragment = EntityView.GetFragmentData<FDataFragment_Transform>();
+		const FTransformFragment& TransformFragment = EntityView.GetFragmentData<FTransformFragment>();
 		const FMassNetworkIDFragment& NetworkIDFragment = EntityView.GetFragmentData<FMassNetworkIDFragment>();
 
 		const FVector& Pos = TransformFragment.GetTransform().GetLocation();
diff --git a/Engine/Plugins/AI/MassCrowd/Source/MassCrowd/Private/MassCrowdMemberTrait.cpp b/Engine/Plugins/AI/MassCrowd/Source/MassCrowd/Private/MassCrowdMemberTrait.cpp
--- a/Engine/Plugins/AI/MassCrowd/Source/MassCrowd/Private/MassCrowdMemberTrait.cpp
+++ b/Engine/Plugins/AI/MassCrowd/Source/MassCrowd/Private/MassCrowdMemberTrait.cpp
@@ -6,6 +6,6 @@
 
 void UMassCrowdMemberTrait::BuildTemplate(FMassEntityTemplateBuildContext& BuildContext, UWorld& World) const
 {
-	BuildContext.AddTag<FTagFragment_MassCrowd>();
+	BuildContext.AddTag<FMassCrowdTag>();
 	BuildContext.AddFragment<FMassCrowdLaneTrackingFragment>();	
 }
diff --git a/Engine/Plugins/AI/MassCrowd/Source/MassCrowd/Private/MassCrowdNavigationProcessor.cpp b/Engine/Plugins/AI/MassCrowd/Source/MassCrowd/Private/MassCrowdNavigationProcessor.cpp
--- a/Engine/Plugins/AI/MassCrowd/Source/MassCrowd/Private/MassCrowdNavigationProcessor.cpp
+++ b/Engine/Plugins/AI/MassCrowd/Source/MassCrowd/Private/MassCrowdNavigationProcessor.cpp
@@ -24,7 +24,7 @@ UMassCrowdLaneTrackingSignalProcessor::UMassCrowdLaneTrackingSignalProcessor()
 
 void UMassCrowdLaneTrackingSignalProcessor::ConfigureQueries()
 {
-	EntityQuery.AddTagRequirement<FTagFragment_MassCrowd>(EMassFragmentPresence::All);
+	EntityQuery.AddTagRequirement<FMassCrowdTag>(EMassFragmentPresence::All);
 	EntityQuery.AddRequirement<FMassCrowdLaneTrackingFragment>(EMassFragmentAccess::ReadWrite);
 	EntityQuery.AddRequirement<FMassZoneGraphLaneLocationFragment>(EMassFragmentAccess::ReadOnly);
 }
@@ -84,7 +84,7 @@ void UMassCrowdLaneTrackingDestructor::Initialize(UObject& Owner)
 
 void UMassCrowdLaneTrackingDestructor::ConfigureQueries()
 {
-	EntityQuery.AddTagRequirement<FTagFragment_MassCrowd>(EMassFragmentPresence::All);
+	EntityQuery.AddTagRequirement<FMassCrowdTag>(EMassFragmentPresence::All);
 	EntityQuery.AddRequirement<FMassCrowdLaneTrackingFragment>(EMassFragmentAccess::ReadOnly);
 }
 
@@ -131,8 +131,8 @@ void UMassCrowdDynamicObstacleProcessor::Initialize(UObject& Owner)
 
 void UMassCrowdDynamicObstacleProcessor::ConfigureQueries()
 {
-	EntityQuery_Conditional.AddRequirement<FDataFragment_Transform>(EMassFragmentAccess::ReadOnly);
-	EntityQuery_Conditional.AddRequirement<FDataFragment_AgentRadius>(EMassFragmentAccess::ReadOnly);
+	EntityQuery_Conditional.AddRequirement<FTransformFragment>(EMassFragmentAccess::ReadOnly);
+	EntityQuery_Conditional.AddRequirement<FAgentRadiusFragment>(EMassFragmentAccess::ReadOnly);
 	EntityQuery_Conditional.AddRequirement<FMassCrowdObstacleFragment>(EMassFragmentAccess::ReadWrite);
 	EntityQuery_Conditional.AddChunkRequirement<FMassSimulationVariableTickChunkFragment>(EMassFragmentAccess::ReadOnly, EMassFragmentPresence::Optional);
 	EntityQuery_Conditional.SetChunkFilter(&FMassSimulationVariableTickChunkFragment::ShouldTickChunkThisFrame);
@@ -146,8 +146,8 @@ void UMassCrowdDynamicObstacleProcessor::Execute(UMassEntitySubsystem& EntitySub
 	{
 		const int32 NumEntities = Context.GetNumEntities();
 
-		const TConstArrayView<FDataFragment_Transform> LocationList = Context.GetFragmentView<FDataFragment_Transform>();
-		const TConstArrayView<FDataFragment_AgentRadius> RadiusList = Context.GetFragmentView<FDataFragment_AgentRadius>();
+		const TConstArrayView<FTransformFragment> LocationList = Context.GetFragmentView<FTransformFragment>();
+		const TConstArrayView<FAgentRadiusFragment> RadiusList = Context.GetFragmentView<FAgentRadiusFragment>();
 		const TArrayView<FMassCrowdObstacleFragment> ObstacleDataList = Context.GetMutableFragmentView<FMassCrowdObstacleFragment>();
 
 		const float CurrentTime = World->GetTimeSeconds();
diff --git a/Engine/Plugins/AI/MassCrowd/Source/MassCrowd/Private/MassCrowdReplicationProcessor.cpp b/Engine/Plugins/AI/MassCrowd/Source/MassCrowd/Private/MassCrowdReplicationProcessor.cpp
--- a/Engine/Plugins/AI/MassCrowd/Source/MassCrowd/Private/MassCrowdReplicationProcessor.cpp
+++ b/Engine/Plugins/AI/MassCrowd/Source/MassCrowd/Private/MassCrowdReplicationProcessor.cpp
@@ -35,10 +35,10 @@ void UMassCrowdReplicationProcessor::ConfigureQueries()
 	FMassReplicationProcessorPositionYawHandler::AddRequirements(EntityQuery);
 	FMassReplicationProcessorPathHandler::AddRequirements(EntityQuery);
 
-	CollectViewerInfoQuery.AddTagRequirement<FTagFragment_MassCrowd>(EMassFragmentPresence::All);
-	CalculateLODQuery.AddTagRequirement<FTagFragment_MassCrowd>(EMassFragmentPresence::All);
-	AdjustLODDistancesQuery.AddTagRequirement<FTagFragment_MassCrowd>(EMassFragmentPresence::All);
-	EntityQuery.AddTagRequirement<FTagFragment_MassCrowd>(EMassFragmentPresence::All);
+	CollectViewerInfoQuery.AddTagRequirement<FMassCrowdTag>(EMassFragmentPresence::All);
+	CalculateLODQuery.AddTagRequirement<FMassCrowdTag>(EMassFragmentPresence::All);
+	AdjustLODDistancesQuery.AddTagRequirement<FMassCrowdTag>(EMassFragmentPresence::All);
+	EntityQuery.AddTagRequirement<FMassCrowdTag>(EMassFragmentPresence::All);
 }
 
 void UMassCrowdReplicationProcessor::Execute(UMassEntitySubsystem& EntitySubsystem, FMassExecutionContext& Context)
@@ -53,7 +53,7 @@ void UMassCrowdReplicationProcessor::Execute(UMassEntitySubsystem& EntitySubsyst
 	{
 		EntityQuery.ForEachEntityChunk(EntitySubsystem, Context, [this](FMassExecutionContext& Context)
 			{
-				const TConstArrayView<FDataFragment_Transform> TransformList = Context.GetFragmentView<FDataFragment_Transform>();
+				const TConstArrayView<FTransformFragment> TransformList = Context.GetFragmentView<FTransformFragment>();
 				const TConstArrayView<FMassReplicationLODFragment> ViewerLODList = Context.GetFragmentView<FMassReplicationLODFragment>();
 				FMassReplicationSharedFragment& RepSharedFragment = Context.GetMutableSharedFragment<FMassReplicationSharedFragment>();
 				RepSharedFragment.LODCalculator.DebugDisplayLOD(Context, ViewerLODList, TransformList, World);
diff --git a/Engine/Plugins/AI/MassCrowd/Source/MassCrowd/Private/MassCrowdRepresentationActorManagement.cpp b/Engine/Plugins/AI/MassCrowd/Source/MassCrowd/Private/MassCrowdRepresentationActorManagement.cpp
--- a/Engine/Plugins/AI/MassCrowd/Source/MassCrowd/Private/MassCrowdRepresentationActorManagement.cpp
+++ b/Engine/Plugins/AI/MassCrowd/Source/MassCrowd/Private/MassCrowdRepresentationActorManagement.cpp
@@ -64,7 +64,7 @@ void UMassCrowdRepresentationActorManagement::SetActorEnabled(const EMassActorEn
 	}
 }
 
-AActor* UMassCrowdRepresentationActorManagement::GetOrSpawnActor(UMassRepresentationSubsystem& RepresentationSubsystem, UMassEntitySubsystem& EntitySubsystem, const FMassEntityHandle MassAgent, FDataFragment_Actor& ActorInfo, const FTransform& Transform, const int16 TemplateActorIndex, FMassActorSpawnRequestHandle& SpawnRequestHandle, const float Priority) const
+AActor* UMassCrowdRepresentationActorManagement::GetOrSpawnActor(UMassRepresentationSubsystem& RepresentationSubsystem, UMassEntitySubsystem& EntitySubsystem, const FMassEntityHandle MassAgent, FMassActorFragment& ActorInfo, const FTransform& Transform, const int16 TemplateActorIndex, FMassActorSpawnRequestHandle& SpawnRequestHandle, const float Priority) const
 {
 	FTransform RootTransform = Transform;
 	
diff --git a/Engine/Plugins/AI/MassCrowd/Source/MassCrowd/Private/MassCrowdServerRepresentationLODProcessor.cpp b/Engine/Plugins/AI/MassCrowd/Source/MassCrowd/Private/MassCrowdServerRepresentationLODProcessor.cpp
--- a/Engine/Plugins/AI/MassCrowd/Source/MassCrowd/Private/MassCrowdServerRepresentationLODProcessor.cpp
+++ b/Engine/Plugins/AI/MassCrowd/Source/MassCrowd/Private/MassCrowdServerRepresentationLODProcessor.cpp
@@ -30,8 +30,8 @@ UMassCrowdServerRepresentationLODProcessor::UMassCrowdServerRepresentationLODPro
 
 void UMassCrowdServerRepresentationLODProcessor::ConfigureQueries()
 {
-	EntityQuery.AddTagRequirement<FTagFragment_MassCrowd>(EMassFragmentPresence::All);
-	EntityQuery.AddRequirement<FDataFragment_Transform>(EMassFragmentAccess::ReadOnly);
+	EntityQuery.AddTagRequirement<FMassCrowdTag>(EMassFragmentPresence::All);
+	EntityQuery.AddRequirement<FTransformFragment>(EMassFragmentAccess::ReadOnly);
 	EntityQuery.AddRequirement<FMassViewerInfoFragment>(EMassFragmentAccess::ReadOnly);
 	EntityQuery.AddRequirement<FMassRepresentationLODFragment>(EMassFragmentAccess::ReadWrite);
 }
@@ -83,7 +83,7 @@ void UMassCrowdServerRepresentationLODProcessor::Execute(UMassEntitySubsystem& E
 		
 		EntityQuery.ForEachEntityChunk(EntitySubsystem, Context, [this](FMassExecutionContext& Context)
 		{
-			const TConstArrayView<FDataFragment_Transform> LocationList = Context.GetFragmentView<FDataFragment_Transform>();
+			const TConstArrayView<FTransformFragment> LocationList = Context.GetFragmentView<FTransformFragment>();
 			const TConstArrayView<FMassRepresentationLODFragment> RepresentationLODFragments = Context.GetFragmentView<FMassRepresentationLODFragment>();
 			LODCalculator.DebugDisplayLOD(Context, RepresentationLODFragments, LocationList, World);
 		});
diff --git a/Engine/Plugins/AI/MassCrowd/Source/MassCrowd/Private/MassCrowdServerRepresentationProcessor.cpp b/Engine/Plugins/AI/MassCrowd/Source/MassCrowd/Private/MassCrowdServerRepresentationProcessor.cpp
--- a/Engine/Plugins/AI/MassCrowd/Source/MassCrowd/Private/MassCrowdServerRepresentationProcessor.cpp
+++ b/Engine/Plugins/AI/MassCrowd/Source/MassCrowd/Private/MassCrowdServerRepresentationProcessor.cpp
@@ -16,5 +16,5 @@ UMassCrowdServerRepresentationProcessor::UMassCrowdServerRepresentationProcessor
 void UMassCrowdServerRepresentationProcessor::ConfigureQueries()
 {
 	Super::ConfigureQueries();
-	EntityQuery.AddTagRequirement<FTagFragment_MassCrowd>(EMassFragmentPresence::All);
+	EntityQuery.AddTagRequirement<FMassCrowdTag>(EMassFragmentPresence::All);
 }
diff --git a/Engine/Plugins/AI/MassCrowd/Source/MassCrowd/Private/MassCrowdServerRepresentationTrait.cpp b/Engine/Plugins/AI/MassCrowd/Source/MassCrowd/Private/MassCrowdServerRepresentationTrait.cpp
--- a/Engine/Plugins/AI/MassCrowd/Source/MassCrowd/Private/MassCrowdServerRepresentationTrait.cpp
+++ b/Engine/Plugins/AI/MassCrowd/Source/MassCrowd/Private/MassCrowdServerRepresentationTrait.cpp
@@ -34,7 +34,7 @@ void UMassCrowdServerRepresentationTrait::BuildTemplate(FMassEntityTemplateBuild
 	// the following needs to be always there for mesh vis to work. Adding following fragments after already 
 	// adding Config.AdditionalDataFragments to let user configure the fragments first. Calling BuildContext.Add() 
 	// won't override any fragments that are already there
-	BuildContext.AddFragment<FDataFragment_Transform>();
+	BuildContext.AddFragment<FTransformFragment>();
 
 	UMassEntitySubsystem* EntitySubsystem = UWorld::GetSubsystem<UMassEntitySubsystem>(&World);
 	check(EntitySubsystem);
@@ -42,15 +42,15 @@ void UMassCrowdServerRepresentationTrait::BuildTemplate(FMassEntityTemplateBuild
 	UMassCrowdRepresentationSubsystem* RepresentationSubsystem = World.GetSubsystem<UMassCrowdRepresentationSubsystem>();
 	check(RepresentationSubsystem);
 
-	FMassRepresentationSubsystemFragment Subsystem;
+	FMassRepresentationSubsystemSharedFragment Subsystem;
 	Subsystem.RepresentationSubsystem = RepresentationSubsystem;
 	uint32 SubsystemHash = UE::StructUtils::GetStructCrc32(FConstStructView::Make(Subsystem));
-	FSharedStruct SubsystemFragment = EntitySubsystem->GetOrCreateSharedFragment<FMassRepresentationSubsystemFragment>(SubsystemHash, Subsystem);
+	FSharedStruct SubsystemFragment = EntitySubsystem->GetOrCreateSharedFragment<FMassRepresentationSubsystemSharedFragment>(SubsystemHash, Subsystem);
 	BuildContext.AddSharedFragment(SubsystemFragment);
 
 	uint32 ConfigHash = UE::StructUtils::GetStructCrc32(FConstStructView::Make(Config));
-	FConstSharedStruct ConfigFragment = EntitySubsystem->GetOrCreateConstSharedFragment<FMassRepresentationConfig>(ConfigHash, Config);
-	ConfigFragment.Get<FMassRepresentationConfig>().ComputeCachedValues();
+	FConstSharedStruct ConfigFragment = EntitySubsystem->GetOrCreateConstSharedFragment<FMassRepresentationParameters>(ConfigHash, Config);
+	ConfigFragment.Get<FMassRepresentationParameters>().ComputeCachedValues();
 	BuildContext.AddConstSharedFragment(ConfigFragment);
 
 	FMassRepresentationFragment& RepresentationFragment = BuildContext.AddFragment_GetRef<FMassRepresentationFragment>();
diff --git a/Engine/Plugins/AI/MassCrowd/Source/MassCrowd/Private/MassCrowdVisualizationLODProcessor.cpp b/Engine/Plugins/AI/MassCrowd/Source/MassCrowd/Private/MassCrowdVisualizationLODProcessor.cpp
--- a/Engine/Plugins/AI/MassCrowd/Source/MassCrowd/Private/MassCrowdVisualizationLODProcessor.cpp
+++ b/Engine/Plugins/AI/MassCrowd/Source/MassCrowd/Private/MassCrowdVisualizationLODProcessor.cpp
@@ -58,11 +58,11 @@ void UMassCrowdVisualizationLODProcessor::ConfigureQueries()
 	// @todo remove, no need for this anymore since we have the common LOD collector
 	// Do not call super as we do have our own LODInfo fragment, so need to duplicate
 
-	CloseEntityQuery.AddTagRequirement<FTagFragment_MassCrowd>(EMassFragmentPresence::All);
+	CloseEntityQuery.AddTagRequirement<FMassCrowdTag>(EMassFragmentPresence::All);
 	CloseEntityQuery.AddRequirement<FMassViewerInfoFragment>(EMassFragmentAccess::ReadOnly);
 	CloseEntityQuery.AddRequirement<FMassRepresentationFragment>(EMassFragmentAccess::ReadOnly);
 	CloseEntityQuery.AddRequirement<FMassRepresentationLODFragment>(EMassFragmentAccess::ReadWrite);
-	CloseEntityQuery.AddRequirement<FDataFragment_Transform>(EMassFragmentAccess::ReadOnly);
+	CloseEntityQuery.AddRequirement<FTransformFragment>(EMassFragmentAccess::ReadOnly);
 
 	FarEntityQuery = CloseEntityQuery;
 	CloseEntityQuery.AddTagRequirement<FMassVisibilityCulledByDistanceTag>(EMassFragmentPresence::None);
@@ -84,7 +84,7 @@ void UMassCrowdVisualizationLODProcessor::Execute(UMassEntitySubsystem& EntitySu
 
 		auto DebugDisplayLOD = [this](FMassExecutionContext& Context)
 		{
-			const TConstArrayView<FDataFragment_Transform> LocationList = Context.GetFragmentView<FDataFragment_Transform>();
+			const TConstArrayView<FTransformFragment> LocationList = Context.GetFragmentView<FTransformFragment>();
 			const TConstArrayView<FMassRepresentationLODFragment> VisualizationLODList = Context.GetFragmentView<FMassRepresentationLODFragment>();
 			LODCalculator.DebugDisplayLOD(Context, VisualizationLODList, LocationList, World);
 		};
@@ -99,7 +99,7 @@ void UMassCrowdVisualizationLODProcessor::Execute(UMassEntitySubsystem& EntitySu
 
 		auto ShowISMUnderSpecifiedRange = [this](const FMassExecutionContext& Context)
 		{
-			const TConstArrayView<FDataFragment_Transform> LocationList = Context.GetFragmentView<FDataFragment_Transform>();
+			const TConstArrayView<FTransformFragment> LocationList = Context.GetFragmentView<FTransformFragment>();
 			const TConstArrayView<FMassRepresentationFragment> RepresentationFragmentList = Context.GetFragmentView<FMassRepresentationFragment>();
 			const TConstArrayView<FMassViewerInfoFragment> LODInfoFragmentList = Context.GetFragmentView<FMassViewerInfoFragment>();
 			const int32 NumEntities = Context.GetNumEntities();
@@ -110,7 +110,7 @@ void UMassCrowdVisualizationLODProcessor::Execute(UMassEntitySubsystem& EntitySu
 				const FMassViewerInfoFragment& LODInfoFragment = LODInfoFragmentList[EntityIdx];
 				if (RepresentationFragment.CurrentRepresentation == EMassRepresentationType::StaticMeshInstance && SpecifiedRangeSquaredCentimeters > LODInfoFragment.ClosestViewerDistanceSq)
 				{
-					const FDataFragment_Transform& EntityLocation = LocationList[EntityIdx];
+					const FTransformFragment& EntityLocation = LocationList[EntityIdx];
 					DrawDebugSolidBox(World, EntityLocation.GetTransform().GetLocation() + FVector(0.0f, 0.0f, 150.0f), FVector(50.0f), FColor::Red);
 				}
 			}
diff --git a/Engine/Plugins/AI/MassCrowd/Source/MassCrowd/Private/MassCrowdVisualizationProcessor.cpp b/Engine/Plugins/AI/MassCrowd/Source/MassCrowd/Private/MassCrowdVisualizationProcessor.cpp
--- a/Engine/Plugins/AI/MassCrowd/Source/MassCrowd/Private/MassCrowdVisualizationProcessor.cpp
+++ b/Engine/Plugins/AI/MassCrowd/Source/MassCrowd/Private/MassCrowdVisualizationProcessor.cpp
@@ -38,7 +38,7 @@ UMassCrowdVisualizationProcessor::UMassCrowdVisualizationProcessor()
 void UMassCrowdVisualizationProcessor::ConfigureQueries()
 {
 	Super::ConfigureQueries();
-	EntityQuery.AddTagRequirement<FTagFragment_MassCrowd>(EMassFragmentPresence::All);
+	EntityQuery.AddTagRequirement<FMassCrowdTag>(EMassFragmentPresence::All);
 }
 
 //----------------------------------------------------------------------//
@@ -56,11 +56,11 @@ UMassDebugCrowdVisualizationProcessor::UMassDebugCrowdVisualizationProcessor()
 
 void UMassDebugCrowdVisualizationProcessor::ConfigureQueries()
 {
-	EntityQuery.AddTagRequirement<FTagFragment_MassCrowd>(EMassFragmentPresence::All);
+	EntityQuery.AddTagRequirement<FMassCrowdTag>(EMassFragmentPresence::All);
 
-	EntityQuery.AddRequirement<FDataFragment_Transform>(EMassFragmentAccess::ReadOnly);
+	EntityQuery.AddRequirement<FTransformFragment>(EMassFragmentAccess::ReadOnly);
 	EntityQuery.AddRequirement<FMassRepresentationFragment>(EMassFragmentAccess::ReadWrite);
-	EntityQuery.AddRequirement<FDataFragment_Actor>(EMassFragmentAccess::ReadWrite);
+	EntityQuery.AddRequirement<FMassActorFragment>(EMassFragmentAccess::ReadWrite);
 }
 
 void UMassDebugCrowdVisualizationProcessor::Initialize(UObject& Owner)
@@ -82,15 +82,15 @@ void UMassDebugCrowdVisualizationProcessor::Execute(UMassEntitySubsystem& Entity
 		EntityQuery.ForEachEntityChunk(EntitySubsystem, Context, [this](const FMassExecutionContext& Context)
 		{
 			const TConstArrayView<FMassRepresentationFragment> VisualizationList = Context.GetFragmentView<FMassRepresentationFragment>();
-			const TConstArrayView<FDataFragment_Actor> ActorList = Context.GetFragmentView<FDataFragment_Actor>();
-			const TConstArrayView<FDataFragment_Transform> EntityLocationList = Context.GetFragmentView<FDataFragment_Transform>();
+			const TConstArrayView<FMassActorFragment> ActorList = Context.GetFragmentView<FMassActorFragment>();
+			const TConstArrayView<FTransformFragment> EntityLocationList = Context.GetFragmentView<FTransformFragment>();
 
 			const int32 NumEntities = Context.GetNumEntities();
 			for (int EntityIdx = 0; EntityIdx < NumEntities; EntityIdx++)
 			{
-				const FDataFragment_Transform& EntityLocation = EntityLocationList[EntityIdx];
+				const FTransformFragment& EntityLocation = EntityLocationList[EntityIdx];
 				const FMassRepresentationFragment& Visualization = VisualizationList[EntityIdx];
-				const FDataFragment_Actor& ActorInfo = ActorList[EntityIdx];
+				const FMassActorFragment& ActorInfo = ActorList[EntityIdx];
 				const int32 RepresentationTypeIdx = (int32)Visualization.CurrentRepresentation;
 				// Show replicated actors
 				if (ActorInfo.IsValid() && !ActorInfo.IsOwnedByMass())
diff --git a/Engine/Plugins/AI/MassCrowd/Source/MassCrowd/Public/MassCrowdFragments.h b/Engine/Plugins/AI/MassCrowd/Source/MassCrowd/Public/MassCrowdFragments.h
--- a/Engine/Plugins/AI/MassCrowd/Source/MassCrowd/Public/MassCrowdFragments.h
+++ b/Engine/Plugins/AI/MassCrowd/Source/MassCrowd/Public/MassCrowdFragments.h
@@ -12,7 +12,7 @@
  * Should not contain any data, this is purely a tag
  */
 USTRUCT()
-struct MASSCROWD_API FTagFragment_MassCrowd : public FMassTag
+struct MASSCROWD_API FMassCrowdTag : public FMassTag
 {
 	GENERATED_BODY()
 };
diff --git a/Engine/Plugins/AI/MassCrowd/Source/MassCrowd/Public/MassCrowdRepresentationActorManagement.h b/Engine/Plugins/AI/MassCrowd/Source/MassCrowd/Public/MassCrowdRepresentationActorManagement.h
--- a/Engine/Plugins/AI/MassCrowd/Source/MassCrowd/Public/MassCrowdRepresentationActorManagement.h
+++ b/Engine/Plugins/AI/MassCrowd/Source/MassCrowd/Public/MassCrowdRepresentationActorManagement.h
@@ -38,7 +38,7 @@ protected:
 	 * @param Priority of this spawn request in comparison with the others, lower value means higher priority
 	 * @return the actor spawned
 	 */
-	virtual AActor* GetOrSpawnActor(UMassRepresentationSubsystem& RepresentationSubsystem, UMassEntitySubsystem& EntitySubsystem, const FMassEntityHandle MassAgent, FDataFragment_Actor& ActorInfo, const FTransform& Transform, const int16 TemplateActorIndex, FMassActorSpawnRequestHandle& SpawnRequestHandle, const float Priority) const override;
+	virtual AActor* GetOrSpawnActor(UMassRepresentationSubsystem& RepresentationSubsystem, UMassEntitySubsystem& EntitySubsystem, const FMassEntityHandle MassAgent, FMassActorFragment& ActorInfo, const FTransform& Transform, const int16 TemplateActorIndex, FMassActorSpawnRequestHandle& SpawnRequestHandle, const float Priority) const override;
 
 	/**
 	 * Teleports the actor at the specified transform by preserving its velocity and without collision.
diff --git a/Engine/Plugins/AI/MassCrowd/Source/MassCrowd/Public/MassCrowdServerRepresentationTrait.h b/Engine/Plugins/AI/MassCrowd/Source/MassCrowd/Public/MassCrowdServerRepresentationTrait.h
--- a/Engine/Plugins/AI/MassCrowd/Source/MassCrowd/Public/MassCrowdServerRepresentationTrait.h
+++ b/Engine/Plugins/AI/MassCrowd/Source/MassCrowd/Public/MassCrowdServerRepresentationTrait.h
@@ -27,5 +27,5 @@ protected:
 
 	/** Configuration parameters for the representation processor */
 	UPROPERTY(EditAnywhere, Category = "Mass|Visual")
-	FMassRepresentationConfig Config;
+	FMassRepresentationParameters Config;
 };
diff --git a/Engine/Plugins/Runtime/MassEntity/Source/MassEntity/Public/MassEntitySubsystem.h b/Engine/Plugins/Runtime/MassEntity/Source/MassEntity/Public/MassEntitySubsystem.h
--- a/Engine/Plugins/Runtime/MassEntity/Source/MassEntity/Public/MassEntitySubsystem.h
+++ b/Engine/Plugins/Runtime/MassEntity/Source/MassEntity/Public/MassEntitySubsystem.h
@@ -603,7 +603,7 @@ public:
 		return TConstArrayView<TFragment>((const TFragment*)View->FragmentView.GetData(), View->FragmentView.Num());
 	}
 
-	TConstArrayView<FMassFragment> GetFragmentFragmentView(const UScriptStruct* FragmentType) const
+	TConstArrayView<FMassFragment> GetFragmentView(const UScriptStruct* FragmentType) const
 	{
 		const FFragmentView* View = FragmentViews.FindByPredicate([FragmentType](const FFragmentView& Element) { return Element.Requirement.StructType == FragmentType; });
 		checkSlow(View);
diff --git a/Engine/Plugins/Runtime/MassEntity/Source/MassEntityTestSuite/Private/LWComponentTestFarmPlot.cpp b/Engine/Plugins/Runtime/MassEntity/Source/MassEntityTestSuite/Private/LWComponentTestFarmPlot.cpp
--- a/Engine/Plugins/Runtime/MassEntity/Source/MassEntityTestSuite/Private/LWComponentTestFarmPlot.cpp
+++ b/Engine/Plugins/Runtime/MassEntity/Source/MassEntityTestSuite/Private/LWComponentTestFarmPlot.cpp
@@ -40,8 +40,8 @@ void UFarmHarvestTimerSetIcon::Execute(UMassEntitySubsystem& EntitySubsystem, FM
 			VisualList[i].HarvestIconIndex = HarvestIconISMC->AddInstance(IconTransform);
 
 			FMassEntityHandle ThisEntity = Context.GetEntity(i);
-			Context.Defer().RemoveTag<FFarmJustBecameReadyToHarvest>(ThisEntity);
-			Context.Defer().AddTag<FFarmReadyToHarvest>(ThisEntity);
+			Context.Defer().RemoveTag<FFarmJustBecameReadyToHarvestTag>(ThisEntity);
+			Context.Defer().AddTag<FFarmReadyToHarvestTag>(ThisEntity);
 		}
 	});
 }
diff --git a/Engine/Plugins/Runtime/MassEntity/Source/MassEntityTestSuite/Public/LWComponentTestFarmPlot.h b/Engine/Plugins/Runtime/MassEntity/Source/MassEntityTestSuite/Public/LWComponentTestFarmPlot.h
--- a/Engine/Plugins/Runtime/MassEntity/Source/MassEntityTestSuite/Public/LWComponentTestFarmPlot.h
+++ b/Engine/Plugins/Runtime/MassEntity/Source/MassEntityTestSuite/Public/LWComponentTestFarmPlot.h
@@ -29,13 +29,13 @@ struct FFarmVisualDataRow : public FTableRowBase
 };
 
 USTRUCT()
-struct FFarmJustBecameReadyToHarvest : public FMassTag
+struct FFarmJustBecameReadyToHarvestTag : public FMassTag
 {
 	GENERATED_BODY()
 };
 
 USTRUCT()
-struct FFarmReadyToHarvest : public FMassTag
+struct FFarmReadyToHarvestTag : public FMassTag
 {
 	GENERATED_BODY()
 };
@@ -223,8 +223,8 @@ public:
 	virtual void ConfigureQueries() override
 	{
 		EntityQuery.AddRequirement<FHarvestTimerFragment>(EMassFragmentAccess::ReadOnly);
-		EntityQuery.AddTagRequirement<FFarmJustBecameReadyToHarvest>(EMassFragmentPresence::None);
-		EntityQuery.AddTagRequirement<FFarmReadyToHarvest>(EMassFragmentPresence::None);
+		EntityQuery.AddTagRequirement<FFarmJustBecameReadyToHarvestTag>(EMassFragmentPresence::None);
+		EntityQuery.AddTagRequirement<FFarmReadyToHarvestTag>(EMassFragmentPresence::None);
 	}
 
 	virtual void Execute(UMassEntitySubsystem& EntitySubsystem, FMassExecutionContext& Context) override
@@ -239,7 +239,7 @@ public:
 			{
 				if (TimerList[i].NumSecondsLeft == 0)
 				{
-					Context.Defer().AddTag<FFarmJustBecameReadyToHarvest>(Context.GetEntity(i));
+					Context.Defer().AddTag<FFarmJustBecameReadyToHarvestTag>(Context.GetEntity(i));
 				}
 			}
 		});
@@ -266,7 +266,7 @@ public:
 	{
 		EntityQuery.AddRequirement<FFarmGridCellData>(EMassFragmentAccess::ReadOnly);
 		EntityQuery.AddRequirement<FFarmVisualFragment>(EMassFragmentAccess::ReadWrite);
-		EntityQuery.AddTagRequirement<FFarmJustBecameReadyToHarvest>(EMassFragmentPresence::All);
+		EntityQuery.AddTagRequirement<FFarmJustBecameReadyToHarvestTag>(EMassFragmentPresence::All);
 	}
 
 	virtual void Execute(UMassEntitySubsystem& EntitySubsystem, FMassExecutionContext& Context) override;
diff --git a/Engine/Plugins/Runtime/MassGameplay/Config/BaseMassGameplay.ini b/Engine/Plugins/Runtime/MassGameplay/Config/BaseMassGameplay.ini
--- a/Engine/Plugins/Runtime/MassGameplay/Config/BaseMassGameplay.ini
+++ b/Engine/Plugins/Runtime/MassGameplay/Config/BaseMassGameplay.ini
@@ -6,3 +6,43 @@
 +ClassRedirects=(OldName="/Script/MassAIMovement.MassMovementSettings", NewName="/Script/MassMovement.MassMovementSettings")
 
 +PropertyRedirects=(OldName="/Script/MassSpawner.SpawnPointsGenerators", NewName="SpawnDataGenerators")
+
+;MassCommon
++StructRedirects=(OldName="DataFragment_AgentRadius", NewName="AgentRadiusFragment")
++StructRedirects=(OldName="DataFragment_ObjectWrapper", NewName="ObjectWrapperFragment")
++StructRedirects=(OldName="DataFragment_Transform", NewName="TransformFragment")
++PropertyRedirects=(OldName="/Script/MassCommon.DataFragment_AgentRadius", NewName="/Script/MassCommon.AgentRadiusFragment")
++PropertyRedirects=(OldName="/Script/MassCommon.DataFragment_ObjectWrapper", NewName="/Script/MassCommon.ObjectWrapperFragment")
++PropertyRedirects=(OldName="/Script/MassCommon.DataFragment_Transform", NewName="/Script/MassCommon.TransformFragment")
+
+;MassActors
++StructRedirects=(OldName="DataFragment_Actor", NewName="MassActorFragment")
++StructRedirects=(OldName="DataFragment_CapsuleComponentWrapper", NewName="CapsuleComponentWrapperFragment")
++StructRedirects=(OldName="DataFragment_CharacterMovementComponentWrapper", NewName="CharacterMovementComponentWrapperFragment")
++StructRedirects=(OldName="/Script/MassActors.DataFragment_Actor", NewName="/Script/MassActors.MassActorFragment")
++StructRedirects=(OldName="/Script/MassActors.DataFragment_CapsuleComponentWrapper", NewName="/Script/MassActors.CapsuleComponentWrapperFragment")
++StructRedirects=(OldName="/Script/MassActors.DataFragment_CharacterMovementComponentWrapper", NewName="/Script/MassActors.CharacterMovementComponentWrapperFragment")
+
+;MassGameplayDebug
++StructRedirects=(OldName="SimDebugVisComponent", NewName="SimDebugVisFragment")
++StructRedirects=(OldName="/Script/MassGameplayDebug.SimDebugVisComponent", NewName="/Script/MassGameplayDebug.SimDebugVisFragment")
+
+;MassSpawner
++StructRedirects=(OldName="DataFragment_ReplicationTemplateID", NewName="ReplicationTemplateIDFragment")
++StructRedirects=(OldName="/Script/MassSpawner.DataFragment_ReplicationTemplateID", NewName="/Script/MassSpawner.ReplicationTemplateIDFragment")
+
+;MassRepresentation
++StructRedirects=(OldName="MassRepresentationConfig", NewName="MassRepresentationParameters")
++StructRedirects=(OldName="MassRepresentationSubsystemFragment", NewName="MassRepresentationSubsystemSharedFragment")
++StructRedirects=(OldName="/Script/MassRepresentation.MassRepresentationConfig", NewName="/Script/MassRepresentation.MassRepresentationParameters")
++StructRedirects=(OldName="/Script/MassRepresentation.MassRepresentationSubsystemFragment", NewName="/Script/MassRepresentation.MassRepresentationSubsystemSharedFragment")
+
+;MassLOD
++StructRedirects=(OldName="MassSimulationVariableTickConfig", NewName="MassSimulationVariableTickParameters")
++StructRedirects=(OldName="MassSimulationLODConfig", NewName="MassSimulationLODParameters")
++StructRedirects=(OldName="/Script/MassLOD.MassSimulationVariableTickConfig", NewName="/Script/MassLOD.MassSimulationVariableTickParameters")
++StructRedirects=(OldName="/Script/MassLOD.MassSimulationLODConfig", NewName="/Script/MassLOD.MassSimulationLODParameters")
+
+;MassReplication
++ClassRedirects=(OldName="/Script/MassReplication.MassReplicationManager", NewName="/Script/MassReplication.MassReplicationSubsystem")
++ClassRedirects=(OldName="MassReplicationManager", NewName="MassReplicationSubsystem")
diff --git a/Engine/Plugins/Runtime/MassGameplay/Source/MassActors/Private/MassActorSubsystem.cpp b/Engine/Plugins/Runtime/MassGameplay/Source/MassActors/Private/MassActorSubsystem.cpp
--- a/Engine/Plugins/Runtime/MassGameplay/Source/MassActors/Private/MassActorSubsystem.cpp
+++ b/Engine/Plugins/Runtime/MassGameplay/Source/MassActors/Private/MassActorSubsystem.cpp
@@ -8,10 +8,10 @@
 #include "VisualLogger/VisualLogger.h"
 
 //----------------------------------------------------------------------//
-//  FDataFragment_Actor 
+//  FMassActorFragment 
 //----------------------------------------------------------------------//
 
-void FDataFragment_Actor::SetAndUpdateHandleMap(const FMassEntityHandle MassAgent, AActor* InActor, const bool bInIsOwnedByMass)
+void FMassActorFragment::SetAndUpdateHandleMap(const FMassEntityHandle MassAgent, AActor* InActor, const bool bInIsOwnedByMass)
 {
 	SetNoHandleMapUpdate(MassAgent, InActor, bInIsOwnedByMass);
 
@@ -23,7 +23,7 @@ void FDataFragment_Actor::SetAndUpdateHandleMap(const FMassEntityHandle MassAgen
 	}
 }
 
-void FDataFragment_Actor::ResetAndUpdateHandleMap()
+void FMassActorFragment::ResetAndUpdateHandleMap()
 {
 	if (AActor* ActorPtr = Cast<AActor>(Actor.Get()))
 	{
@@ -38,7 +38,7 @@ void FDataFragment_Actor::ResetAndUpdateHandleMap()
 	ResetNoHandleMapUpdate();
 }
 
-void FDataFragment_Actor::SetNoHandleMapUpdate(const FMassEntityHandle MassAgent, AActor* InActor, const bool bInIsOwnedByMass)
+void FMassActorFragment::SetNoHandleMapUpdate(const FMassEntityHandle MassAgent, AActor* InActor, const bool bInIsOwnedByMass)
 {
 	check(InActor);
 	check(!Actor.IsValid());
@@ -47,7 +47,7 @@ void FDataFragment_Actor::SetNoHandleMapUpdate(const FMassEntityHandle MassAgent
 	bIsOwnedByMass = bInIsOwnedByMass;
 }
 
-void FDataFragment_Actor::ResetNoHandleMapUpdate()
+void FMassActorFragment::ResetNoHandleMapUpdate()
 {
 	Actor.Reset();
 	bIsOwnedByMass = false;
@@ -81,7 +81,7 @@ FMassEntityHandle UMassActorSubsystem::GetEntityHandleFromActor(const TObjectKey
 AActor* UMassActorSubsystem::GetActorFromHandle(const FMassEntityHandle Handle) const
 {
 	check(EntitySystem);
-	FDataFragment_Actor* Data = EntitySystem->GetFragmentDataPtr<FDataFragment_Actor>(Handle);
+	FMassActorFragment* Data = EntitySystem->GetFragmentDataPtr<FMassActorFragment>(Handle);
 	return Data != nullptr ? Data->GetMutable() : nullptr;
 }
 
@@ -119,7 +119,7 @@ void UMassActorSubsystem::DisconnectActor(const TObjectKey<const AActor> Actor,
 	if (FoundEntity == Handle)
 	{
 		check(EntitySystem);
-		if (FDataFragment_Actor* Data = EntitySystem->GetFragmentDataPtr<FDataFragment_Actor>(Handle))
+		if (FMassActorFragment* Data = EntitySystem->GetFragmentDataPtr<FMassActorFragment>(Handle))
 		{
 			Data->ResetAndUpdateHandleMap();
 		}
diff --git a/Engine/Plugins/Runtime/MassGameplay/Source/MassActors/Private/MassAgentComponent.cpp b/Engine/Plugins/Runtime/MassGameplay/Source/MassActors/Private/MassAgentComponent.cpp
--- a/Engine/Plugins/Runtime/MassGameplay/Source/MassActors/Private/MassAgentComponent.cpp
+++ b/Engine/Plugins/Runtime/MassGameplay/Source/MassActors/Private/MassAgentComponent.cpp
@@ -194,14 +194,14 @@ void UMassAgentComponent::SetEntityHandleInternal(const FMassEntityHandle NewHan
 			// @todo Find a way to add these initialization into either translator initializer or adding new fragments
 			// Make sure to fetch the fragment after any release, as that action can move the entity around into new archetype and 
 			// by the same fact change the references to the fragments.
-			if (FDataFragment_Actor* ActorInfo = EntityView.GetFragmentDataPtr<FDataFragment_Actor>())
+			if (FMassActorFragment* ActorInfo = EntityView.GetFragmentDataPtr<FMassActorFragment>())
 			{
 				checkf(!ActorInfo->IsValid(), TEXT("Expecting ActorInfo fragment to be null"));
 				ActorInfo->SetAndUpdateHandleMap(AgentHandle, GetOwner(), !IsNetSimulating()/*bIsOwnedByMass*/);
 			}
 
 			// Initialize location of the replicated actor to match the mass replicated one
-			if (const FDataFragment_Transform* TransformFragment = EntityView.GetFragmentDataPtr<FDataFragment_Transform>())
+			if (const FTransformFragment* TransformFragment = EntityView.GetFragmentDataPtr<FTransformFragment>())
 			{
 				GetOwner()->SetActorTransform(TransformFragment->GetTransform(), /*bSweep*/false, /*OutSweepHitResult*/nullptr, ETeleportType::TeleportPhysics);
 			}
@@ -279,7 +279,7 @@ void UMassAgentComponent::ClearEntityHandleInternal()
 		if (IsNetSimulating())
 		{
 			const FMassEntityView EntityView(*EntitySubsystem, AgentHandle);
-			if (FDataFragment_Actor* ActorInfo = EntityView.GetFragmentDataPtr<FDataFragment_Actor>())
+			if (FMassActorFragment* ActorInfo = EntityView.GetFragmentDataPtr<FMassActorFragment>())
 			{
 				checkf(!ActorInfo->IsValid() || ActorInfo->Get() == GetOwner(), TEXT("Expecting actor pointer to be the Component\'s owner"));
 				ActorInfo->ResetAndUpdateHandleMap();
@@ -357,7 +357,7 @@ void UMassAgentComponent::DebugCheckStateConsistency()
 							if (bIsBuiltEntity)
 							{
 								AActor* Owner = GetOwner();
-								const AActor* Actor = EntitySubsystem->GetFragmentDataChecked<FDataFragment_Actor>(AgentHandle).Get();
+								const AActor* Actor = EntitySubsystem->GetFragmentDataChecked<FMassActorFragment>(AgentHandle).Get();
 								MASSAGENT_CHECK(Actor == nullptr || Actor == Owner, TEXT("Mass Actor and Owner mismatched in state %s"), *UEnum::GetValueAsString(State));
 							}
 						}
diff --git a/Engine/Plugins/Runtime/MassGameplay/Source/MassActors/Private/MassAgentTraits.cpp b/Engine/Plugins/Runtime/MassGameplay/Source/MassActors/Private/MassAgentTraits.cpp
--- a/Engine/Plugins/Runtime/MassGameplay/Source/MassActors/Private/MassAgentTraits.cpp
+++ b/Engine/Plugins/Runtime/MassGameplay/Source/MassActors/Private/MassAgentTraits.cpp
@@ -41,26 +41,26 @@ namespace FMassAgentTraitsHelper
 //----------------------------------------------------------------------//
 void UMassAgentCapsuleCollisionSyncTrait::BuildTemplate(FMassEntityTemplateBuildContext& BuildContext, UWorld& World) const
 {
-	BuildContext.AddFragment<FDataFragment_CapsuleComponentWrapper>();
-	BuildContext.AddFragment<FDataFragment_AgentRadius>();
+	BuildContext.AddFragment<FCapsuleComponentWrapperFragment>();
+	BuildContext.AddFragment<FAgentRadiusFragment>();
 	if (bSyncTransform)
 	{
-		BuildContext.AddFragment<FDataFragment_Transform>();
+		BuildContext.AddFragment<FTransformFragment>();
 	}
 	
 	BuildContext.GetMutableObjectFragmentInitializers().Add([=](UObject& Owner, FMassEntityView& EntityView, const EMassTranslationDirection CurrentDirection)
 		{
 			if (UCapsuleComponent* CapsuleComponent = FMassAgentTraitsHelper::AsComponent<UCapsuleComponent>(Owner))
 			{
-				FDataFragment_CapsuleComponentWrapper& CapsuleFragment = EntityView.GetFragmentData<FDataFragment_CapsuleComponentWrapper>();
+				FCapsuleComponentWrapperFragment& CapsuleFragment = EntityView.GetFragmentData<FCapsuleComponentWrapperFragment>();
 				CapsuleFragment.Component = CapsuleComponent;
 
-				FDataFragment_AgentRadius& RadiusFragment = EntityView.GetFragmentData<FDataFragment_AgentRadius>();
+				FAgentRadiusFragment& RadiusFragment = EntityView.GetFragmentData<FAgentRadiusFragment>();
 				RadiusFragment.Radius = CapsuleComponent->GetScaledCapsuleRadius();
 
 				if (bSyncTransform)
 				{
-					FDataFragment_Transform& TransformFragment = EntityView.GetFragmentData<FDataFragment_Transform>();
+					FTransformFragment& TransformFragment = EntityView.GetFragmentData<FTransformFragment>();
 					TransformFragment.GetMutableTransform() = CapsuleComponent->GetComponentTransform();
 				}
 			}
@@ -85,14 +85,14 @@ void UMassAgentCapsuleCollisionSyncTrait::BuildTemplate(FMassEntityTemplateBuild
 //----------------------------------------------------------------------//
 void UMassAgentMovementSyncTrait::BuildTemplate(FMassEntityTemplateBuildContext& BuildContext, UWorld& World) const
 {
-	BuildContext.AddFragment<FDataFragment_CharacterMovementComponentWrapper>();
+	BuildContext.AddFragment<FCharacterMovementComponentWrapperFragment>();
 	BuildContext.AddFragment<FMassVelocityFragment>();
 	
 	BuildContext.GetMutableObjectFragmentInitializers().Add([=](UObject& Owner, FMassEntityView& EntityView, const EMassTranslationDirection CurrentDirection)
 		{
 			if (UCharacterMovementComponent* MovementComp = FMassAgentTraitsHelper::AsComponent<UCharacterMovementComponent>(Owner))
 			{
-				FDataFragment_CharacterMovementComponentWrapper& ComponentFragment = EntityView.GetFragmentData<FDataFragment_CharacterMovementComponentWrapper>();
+				FCharacterMovementComponentWrapperFragment& ComponentFragment = EntityView.GetFragmentData<FCharacterMovementComponentWrapperFragment>();
 				ComponentFragment.Component = MovementComp;
 
 				FMassVelocityFragment& VelocityFragment = EntityView.GetFragmentData<FMassVelocityFragment>();
@@ -128,7 +128,7 @@ void UMassAgentMovementSyncTrait::BuildTemplate(FMassEntityTemplateBuildContext&
 //----------------------------------------------------------------------//
 void UMassAgentOrientationSyncTrait::BuildTemplate(FMassEntityTemplateBuildContext& BuildContext, UWorld& World) const
 {
-	BuildContext.AddFragment<FDataFragment_CharacterMovementComponentWrapper>();
+	BuildContext.AddFragment<FCharacterMovementComponentWrapperFragment>();
 	// @todo: Figure out how we can share init with UMassAgentMovementSyncTrait, or make this this trait to depend on UMassAgentMovementSyncTrait.
 
 	if (EnumHasAnyFlags(SyncDirection, EMassTranslationDirection::ActorToMass))
@@ -148,7 +148,7 @@ void UMassAgentOrientationSyncTrait::BuildTemplate(FMassEntityTemplateBuildConte
 void UMassAgentFeetLocationSyncTrait::BuildTemplate(FMassEntityTemplateBuildContext& BuildContext, UWorld& World) const
 {
 	BuildContext.AddFragment<FMassSceneComponentWrapperFragment>();
-	BuildContext.AddFragment<FDataFragment_Transform>();
+	BuildContext.AddFragment<FTransformFragment>();
 
 	BuildContext.GetMutableObjectFragmentInitializers().Add([=](UObject& Owner, FMassEntityView& EntityView, const EMassTranslationDirection CurrentDirection)
 		{
@@ -159,7 +159,7 @@ void UMassAgentFeetLocationSyncTrait::BuildTemplate(FMassEntityTemplateBuildCont
 				FMassSceneComponentWrapperFragment& ComponentFragment = EntityView.GetFragmentData<FMassSceneComponentWrapperFragment>();
 				ComponentFragment.Component = Component;
 
-				FDataFragment_Transform& TransformFragment = EntityView.GetFragmentData<FDataFragment_Transform>();
+				FTransformFragment& TransformFragment = EntityView.GetFragmentData<FTransformFragment>();
 
 				REDIRECT_OBJECT_TO_VLOG(Component, &Owner);
 				UE_VLOG_LOCATION(&Owner, LogMass, Log, Component->GetComponentLocation(), 30, FColor::Yellow, TEXT("Initial component location"));
diff --git a/Engine/Plugins/Runtime/MassGameplay/Source/MassActors/Private/Translators/MassCapsuleComponentTranslators.cpp b/Engine/Plugins/Runtime/MassGameplay/Source/MassActors/Private/Translators/MassCapsuleComponentTranslators.cpp
--- a/Engine/Plugins/Runtime/MassGameplay/Source/MassActors/Private/Translators/MassCapsuleComponentTranslators.cpp
+++ b/Engine/Plugins/Runtime/MassGameplay/Source/MassActors/Private/Translators/MassCapsuleComponentTranslators.cpp
@@ -19,16 +19,16 @@ UMassCapsuleTransformToMassTranslator::UMassCapsuleTransformToMassTranslator()
 void UMassCapsuleTransformToMassTranslator::ConfigureQueries()
 {
 	AddRequiredTagsToQuery(EntityQuery);
-	EntityQuery.AddRequirement<FDataFragment_CapsuleComponentWrapper>(EMassFragmentAccess::ReadOnly);
-	EntityQuery.AddRequirement<FDataFragment_Transform>(EMassFragmentAccess::ReadWrite);
+	EntityQuery.AddRequirement<FCapsuleComponentWrapperFragment>(EMassFragmentAccess::ReadOnly);
+	EntityQuery.AddRequirement<FTransformFragment>(EMassFragmentAccess::ReadWrite);
 }
 
 void UMassCapsuleTransformToMassTranslator::Execute(UMassEntitySubsystem& EntitySubsystem, FMassExecutionContext& Context)
 {
 	EntityQuery.ForEachEntityChunk(EntitySubsystem, Context, [this](FMassExecutionContext& Context)
 		{
-			const TConstArrayView<FDataFragment_CapsuleComponentWrapper> CapsuleComponentList = Context.GetFragmentView<FDataFragment_CapsuleComponentWrapper>();
-			const TArrayView<FDataFragment_Transform> LocationList = Context.GetMutableFragmentView<FDataFragment_Transform>();
+			const TConstArrayView<FCapsuleComponentWrapperFragment> CapsuleComponentList = Context.GetFragmentView<FCapsuleComponentWrapperFragment>();
+			const TArrayView<FTransformFragment> LocationList = Context.GetMutableFragmentView<FTransformFragment>();
 			for (int i = 0; i < CapsuleComponentList.Num(); ++i)
 			{
 				if (const UCapsuleComponent* CapsuleComp = CapsuleComponentList[i].Component.Get())
@@ -53,16 +53,16 @@ UMassTransformToActorCapsuleTranslator::UMassTransformToActorCapsuleTranslator()
 void UMassTransformToActorCapsuleTranslator::ConfigureQueries()
 {
 	AddRequiredTagsToQuery(EntityQuery);
-	EntityQuery.AddRequirement<FDataFragment_CapsuleComponentWrapper>(EMassFragmentAccess::ReadWrite);
-	EntityQuery.AddRequirement<FDataFragment_Transform>(EMassFragmentAccess::ReadOnly);
+	EntityQuery.AddRequirement<FCapsuleComponentWrapperFragment>(EMassFragmentAccess::ReadWrite);
+	EntityQuery.AddRequirement<FTransformFragment>(EMassFragmentAccess::ReadOnly);
 }
 
 void UMassTransformToActorCapsuleTranslator::Execute(UMassEntitySubsystem& EntitySubsystem, FMassExecutionContext& Context)
 {
 	EntityQuery.ForEachEntityChunk(EntitySubsystem, Context, [this](FMassExecutionContext& Context)
 		{
-			const TArrayView<FDataFragment_CapsuleComponentWrapper> CapsuleComponentList = Context.GetMutableFragmentView<FDataFragment_CapsuleComponentWrapper>();
-			const TConstArrayView<FDataFragment_Transform> LocationList = Context.GetFragmentView<FDataFragment_Transform>();
+			const TArrayView<FCapsuleComponentWrapperFragment> CapsuleComponentList = Context.GetMutableFragmentView<FCapsuleComponentWrapperFragment>();
+			const TConstArrayView<FTransformFragment> LocationList = Context.GetFragmentView<FTransformFragment>();
 			for (int i = 0; i < CapsuleComponentList.Num(); ++i)
 			{
 				if (UCapsuleComponent* CapsuleComp = CapsuleComponentList[i].Component.Get())
diff --git a/Engine/Plugins/Runtime/MassGameplay/Source/MassActors/Private/Translators/MassCharacterMovementTranslators.cpp b/Engine/Plugins/Runtime/MassGameplay/Source/MassActors/Private/Translators/MassCharacterMovementTranslators.cpp
--- a/Engine/Plugins/Runtime/MassGameplay/Source/MassActors/Private/Translators/MassCharacterMovementTranslators.cpp
+++ b/Engine/Plugins/Runtime/MassGameplay/Source/MassActors/Private/Translators/MassCharacterMovementTranslators.cpp
@@ -22,8 +22,8 @@ UMassCharacterMovementToMassTranslator::UMassCharacterMovementToMassTranslator()
 void UMassCharacterMovementToMassTranslator::ConfigureQueries()
 {
 	AddRequiredTagsToQuery(EntityQuery);
-	EntityQuery.AddRequirement<FDataFragment_CharacterMovementComponentWrapper>(EMassFragmentAccess::ReadOnly);
-	EntityQuery.AddRequirement<FDataFragment_Transform>(EMassFragmentAccess::ReadWrite);
+	EntityQuery.AddRequirement<FCharacterMovementComponentWrapperFragment>(EMassFragmentAccess::ReadOnly);
+	EntityQuery.AddRequirement<FTransformFragment>(EMassFragmentAccess::ReadWrite);
 	EntityQuery.AddRequirement<FMassVelocityFragment>(EMassFragmentAccess::ReadWrite);
 }
 
@@ -31,8 +31,8 @@ void UMassCharacterMovementToMassTranslator::Execute(UMassEntitySubsystem& Entit
 {
 	EntityQuery.ForEachEntityChunk(EntitySubsystem, Context, [this](FMassExecutionContext& Context)
 	{
-		const TConstArrayView<FDataFragment_CharacterMovementComponentWrapper> ComponentList = Context.GetFragmentView<FDataFragment_CharacterMovementComponentWrapper>();
-		const TArrayView<FDataFragment_Transform> LocationList = Context.GetMutableFragmentView<FDataFragment_Transform>();
+		const TConstArrayView<FCharacterMovementComponentWrapperFragment> ComponentList = Context.GetFragmentView<FCharacterMovementComponentWrapperFragment>();
+		const TArrayView<FTransformFragment> LocationList = Context.GetMutableFragmentView<FTransformFragment>();
 		const TArrayView<FMassVelocityFragment> VelocityList = Context.GetMutableFragmentView<FMassVelocityFragment>();
 
 		const int32 NumEntities = Context.GetNumEntities();
@@ -63,7 +63,7 @@ UMassCharacterMovementToActorTranslator::UMassCharacterMovementToActorTranslator
 void UMassCharacterMovementToActorTranslator::ConfigureQueries()
 {
 	AddRequiredTagsToQuery(EntityQuery);
-	EntityQuery.AddRequirement<FDataFragment_CharacterMovementComponentWrapper>(EMassFragmentAccess::ReadWrite);
+	EntityQuery.AddRequirement<FCharacterMovementComponentWrapperFragment>(EMassFragmentAccess::ReadWrite);
 	EntityQuery.AddRequirement<FMassVelocityFragment>(EMassFragmentAccess::ReadOnly);
 }
 
@@ -71,7 +71,7 @@ void UMassCharacterMovementToActorTranslator::Execute(UMassEntitySubsystem& Enti
 {
 	EntityQuery.ForEachEntityChunk(EntitySubsystem, Context, [this](FMassExecutionContext& Context)
 	{
-		const TArrayView<FDataFragment_CharacterMovementComponentWrapper> ComponentList = Context.GetMutableFragmentView<FDataFragment_CharacterMovementComponentWrapper>();
+		const TArrayView<FCharacterMovementComponentWrapperFragment> ComponentList = Context.GetMutableFragmentView<FCharacterMovementComponentWrapperFragment>();
 		const TConstArrayView<FMassVelocityFragment> VelocityList = Context.GetFragmentView<FMassVelocityFragment>();
 		
 		const int32 NumEntities = Context.GetNumEntities();
@@ -99,16 +99,16 @@ UMassCharacterOrientationToMassTranslator::UMassCharacterOrientationToMassTransl
 void UMassCharacterOrientationToMassTranslator::ConfigureQueries()
 {
 	AddRequiredTagsToQuery(EntityQuery);
-	EntityQuery.AddRequirement<FDataFragment_CharacterMovementComponentWrapper>(EMassFragmentAccess::ReadOnly);
-	EntityQuery.AddRequirement<FDataFragment_Transform>(EMassFragmentAccess::ReadWrite);
+	EntityQuery.AddRequirement<FCharacterMovementComponentWrapperFragment>(EMassFragmentAccess::ReadOnly);
+	EntityQuery.AddRequirement<FTransformFragment>(EMassFragmentAccess::ReadWrite);
 }
 
 void UMassCharacterOrientationToMassTranslator::Execute(UMassEntitySubsystem& EntitySubsystem, FMassExecutionContext& Context)
 {
 	EntityQuery.ForEachEntityChunk(EntitySubsystem, Context, [this](FMassExecutionContext& Context)
 	{
-		const TConstArrayView<FDataFragment_CharacterMovementComponentWrapper> ComponentList = Context.GetFragmentView<FDataFragment_CharacterMovementComponentWrapper>();
-		const TArrayView<FDataFragment_Transform> LocationList = Context.GetMutableFragmentView<FDataFragment_Transform>();
+		const TConstArrayView<FCharacterMovementComponentWrapperFragment> ComponentList = Context.GetFragmentView<FCharacterMovementComponentWrapperFragment>();
+		const TArrayView<FTransformFragment> LocationList = Context.GetMutableFragmentView<FTransformFragment>();
 
 		const int32 NumEntities = Context.GetNumEntities();
 		
@@ -139,16 +139,16 @@ UMassCharacterOrientationToActorTranslator::UMassCharacterOrientationToActorTran
 void UMassCharacterOrientationToActorTranslator::ConfigureQueries()
 {
 	AddRequiredTagsToQuery(EntityQuery);
-	EntityQuery.AddRequirement<FDataFragment_CharacterMovementComponentWrapper>(EMassFragmentAccess::ReadWrite);
-	EntityQuery.AddRequirement<FDataFragment_Transform>(EMassFragmentAccess::ReadOnly);
+	EntityQuery.AddRequirement<FCharacterMovementComponentWrapperFragment>(EMassFragmentAccess::ReadWrite);
+	EntityQuery.AddRequirement<FTransformFragment>(EMassFragmentAccess::ReadOnly);
 }
 
 void UMassCharacterOrientationToActorTranslator::Execute(UMassEntitySubsystem& EntitySubsystem, FMassExecutionContext& Context)
 {
 	EntityQuery.ForEachEntityChunk(EntitySubsystem, Context, [this](FMassExecutionContext& Context)
 	{
-		const TArrayView<FDataFragment_CharacterMovementComponentWrapper> ComponentList = Context.GetMutableFragmentView<FDataFragment_CharacterMovementComponentWrapper>();
-		const TConstArrayView<FDataFragment_Transform> TransformList = Context.GetFragmentView<FDataFragment_Transform>();
+		const TArrayView<FCharacterMovementComponentWrapperFragment> ComponentList = Context.GetMutableFragmentView<FCharacterMovementComponentWrapperFragment>();
+		const TConstArrayView<FTransformFragment> TransformList = Context.GetFragmentView<FTransformFragment>();
 		
 		const int32 NumEntities = Context.GetNumEntities();
 
@@ -158,7 +158,7 @@ void UMassCharacterOrientationToActorTranslator::Execute(UMassEntitySubsystem& E
 			{
 				if (AsMovementComponent->UpdatedComponent != nullptr)
 				{
-					const FDataFragment_Transform& Transform = TransformList[i];
+					const FTransformFragment& Transform = TransformList[i];
 					AsMovementComponent->bOrientRotationToMovement = false;
 					AsMovementComponent->UpdatedComponent->SetWorldRotation(Transform.GetTransform().GetRotation());
 				}
@@ -166,30 +166,3 @@ void UMassCharacterOrientationToActorTranslator::Execute(UMassEntitySubsystem& E
 		}
 	});
 }
-
-//----------------------------------------------------------------------//
-// UMassFragmentInitializer_NavLocation 
-//----------------------------------------------------------------------//
-UMassFragmentInitializer_NavLocation::UMassFragmentInitializer_NavLocation()
-{
-	ObservedType = FDataFragment_NavLocation::StaticStruct();
-	Operation = EMassObservedOperation::Add;
-}
-
-void UMassFragmentInitializer_NavLocation::ConfigureQueries() 
-{
-	EntityQuery.AddRequirement<FDataFragment_NavLocation>(EMassFragmentAccess::ReadWrite);
-}
-
-void UMassFragmentInitializer_NavLocation::Execute(UMassEntitySubsystem& EntitySubsystem, FMassExecutionContext& Context)
-{
-	EntityQuery.ForEachEntityChunk(EntitySubsystem, Context, [this](FMassExecutionContext& Context)
-		{
-			const TArrayView<FDataFragment_NavLocation> NavLocationList = Context.GetMutableFragmentView<FDataFragment_NavLocation>();
-			const int32 NumEntities = Context.GetNumEntities();
-			for (int32 i = 0; i < NumEntities; ++i)
-			{
-				NavLocationList[i].NodeRef = INVALID_NAVNODEREF;
-			}
-		});
-}
diff --git a/Engine/Plugins/Runtime/MassGameplay/Source/MassActors/Private/Translators/MassSceneComponentLocationTranslator.cpp b/Engine/Plugins/Runtime/MassGameplay/Source/MassActors/Private/Translators/MassSceneComponentLocationTranslator.cpp
--- a/Engine/Plugins/Runtime/MassGameplay/Source/MassActors/Private/Translators/MassSceneComponentLocationTranslator.cpp
+++ b/Engine/Plugins/Runtime/MassGameplay/Source/MassActors/Private/Translators/MassSceneComponentLocationTranslator.cpp
@@ -20,7 +20,7 @@ void UMassSceneComponentLocationToMassTranslator::ConfigureQueries()
 {
 	AddRequiredTagsToQuery(EntityQuery);
 	EntityQuery.AddRequirement<FMassSceneComponentWrapperFragment>(EMassFragmentAccess::ReadOnly);
-	EntityQuery.AddRequirement<FDataFragment_Transform>(EMassFragmentAccess::ReadWrite);
+	EntityQuery.AddRequirement<FTransformFragment>(EMassFragmentAccess::ReadWrite);
 }
 
 void UMassSceneComponentLocationToMassTranslator::Execute(UMassEntitySubsystem& EntitySubsystem, FMassExecutionContext& Context)
@@ -28,7 +28,7 @@ void UMassSceneComponentLocationToMassTranslator::Execute(UMassEntitySubsystem&
 	EntityQuery.ForEachEntityChunk(EntitySubsystem, Context, [this](FMassExecutionContext& Context)
 	{
 		const TConstArrayView<FMassSceneComponentWrapperFragment> ComponentList = Context.GetFragmentView<FMassSceneComponentWrapperFragment>();
-		const TArrayView<FDataFragment_Transform> LocationList = Context.GetMutableFragmentView<FDataFragment_Transform>();
+		const TArrayView<FTransformFragment> LocationList = Context.GetMutableFragmentView<FTransformFragment>();
 
 		const int32 NumEntities = Context.GetNumEntities();
 		for (int32 i = 0; i < NumEntities; ++i)
@@ -56,7 +56,7 @@ void UMassSceneComponentLocationToActorTranslator::ConfigureQueries()
 {
 	AddRequiredTagsToQuery(EntityQuery);
 	EntityQuery.AddRequirement<FMassSceneComponentWrapperFragment>(EMassFragmentAccess::ReadOnly);
-	EntityQuery.AddRequirement<FDataFragment_Transform>(EMassFragmentAccess::ReadWrite);
+	EntityQuery.AddRequirement<FTransformFragment>(EMassFragmentAccess::ReadWrite);
 }
 
 void UMassSceneComponentLocationToActorTranslator::Execute(UMassEntitySubsystem& EntitySubsystem, FMassExecutionContext& Context)
@@ -64,7 +64,7 @@ void UMassSceneComponentLocationToActorTranslator::Execute(UMassEntitySubsystem&
 	EntityQuery.ForEachEntityChunk(EntitySubsystem, Context, [this](FMassExecutionContext& Context)
 		{
 			const TConstArrayView<FMassSceneComponentWrapperFragment> ComponentList = Context.GetFragmentView<FMassSceneComponentWrapperFragment>();
-			const TArrayView<FDataFragment_Transform> LocationList = Context.GetMutableFragmentView<FDataFragment_Transform>();
+			const TArrayView<FTransformFragment> LocationList = Context.GetMutableFragmentView<FTransformFragment>();
 
 			const int32 NumEntities = Context.GetNumEntities();
 			for (int32 i = 0; i < NumEntities; ++i)
diff --git a/Engine/Plugins/Runtime/MassGameplay/Source/MassActors/Public/MassActorSubsystem.h b/Engine/Plugins/Runtime/MassGameplay/Source/MassActors/Public/MassActorSubsystem.h
--- a/Engine/Plugins/Runtime/MassGameplay/Source/MassActors/Public/MassActorSubsystem.h
+++ b/Engine/Plugins/Runtime/MassGameplay/Source/MassActors/Public/MassActorSubsystem.h
@@ -18,7 +18,7 @@ class UMassEntitySubsystem;
  * Fragment to save the actor pointer of a mass entity if it exist
  */
 USTRUCT()
-struct MASSACTORS_API FDataFragment_Actor : public FDataFragment_ObjectWrapper
+struct MASSACTORS_API FMassActorFragment : public FObjectWrapperFragment
 {
 	GENERATED_BODY()
 
diff --git a/Engine/Plugins/Runtime/MassGameplay/Source/MassActors/Public/Translators/MassCapsuleComponentTranslators.h b/Engine/Plugins/Runtime/MassGameplay/Source/MassActors/Public/Translators/MassCapsuleComponentTranslators.h
--- a/Engine/Plugins/Runtime/MassGameplay/Source/MassActors/Public/Translators/MassCapsuleComponentTranslators.h
+++ b/Engine/Plugins/Runtime/MassGameplay/Source/MassActors/Public/Translators/MassCapsuleComponentTranslators.h
@@ -8,10 +8,10 @@
 
 
 class UCapsuleComponent;
-struct FDataFragment_AgentRadius;
+struct FAgentRadiusFragment;
 
 USTRUCT()
-struct FDataFragment_CapsuleComponentWrapper : public FDataFragment_ObjectWrapper
+struct FCapsuleComponentWrapperFragment : public FObjectWrapperFragment
 {
 	GENERATED_BODY()
 	TWeakObjectPtr<UCapsuleComponent> Component;
diff --git a/Engine/Plugins/Runtime/MassGameplay/Source/MassActors/Public/Translators/MassCharacterMovementTranslators.h b/Engine/Plugins/Runtime/MassGameplay/Source/MassActors/Public/Translators/MassCharacterMovementTranslators.h
--- a/Engine/Plugins/Runtime/MassGameplay/Source/MassActors/Public/Translators/MassCharacterMovementTranslators.h
+++ b/Engine/Plugins/Runtime/MassGameplay/Source/MassActors/Public/Translators/MassCharacterMovementTranslators.h
@@ -10,7 +10,7 @@
 class UCharacterMovementComponent;
 
 USTRUCT()
-struct FDataFragment_CharacterMovementComponentWrapper : public FDataFragment_ObjectWrapper
+struct FCharacterMovementComponentWrapperFragment : public FObjectWrapperFragment
 {
 	GENERATED_BODY()
 	TWeakObjectPtr<UCharacterMovementComponent> Component;
@@ -102,17 +102,3 @@ protected:
 
 	FMassEntityQuery EntityQuery;
 };
-
-UCLASS()
-class MASSACTORS_API UMassFragmentInitializer_NavLocation : public UMassObserverProcessor
-{
-	GENERATED_BODY()
-public:
-	UMassFragmentInitializer_NavLocation();
-
-protected:
-	virtual void ConfigureQueries() override;
-	virtual void Execute(UMassEntitySubsystem& EntitySubsystem, FMassExecutionContext& Context) override;
-
-	FMassEntityQuery EntityQuery;
-};
diff --git a/Engine/Plugins/Runtime/MassGameplay/Source/MassActors/Public/Translators/MassSceneComponentLocationTranslator.h b/Engine/Plugins/Runtime/MassGameplay/Source/MassActors/Public/Translators/MassSceneComponentLocationTranslator.h
--- a/Engine/Plugins/Runtime/MassGameplay/Source/MassActors/Public/Translators/MassSceneComponentLocationTranslator.h
+++ b/Engine/Plugins/Runtime/MassGameplay/Source/MassActors/Public/Translators/MassSceneComponentLocationTranslator.h
@@ -8,7 +8,7 @@
 
 
 USTRUCT()
-struct FMassSceneComponentWrapperFragment : public FDataFragment_ObjectWrapper
+struct FMassSceneComponentWrapperFragment : public FObjectWrapperFragment
 {
 	GENERATED_BODY()
 	TWeakObjectPtr<USceneComponent> Component;
diff --git a/Engine/Plugins/Runtime/MassGameplay/Source/MassActors/Public/Translators/MassTranslators_BehaviorTree.h b/Engine/Plugins/Runtime/MassGameplay/Source/MassActors/Public/Translators/MassTranslators_BehaviorTree.h
--- a/Engine/Plugins/Runtime/MassGameplay/Source/MassActors/Public/Translators/MassTranslators_BehaviorTree.h
+++ b/Engine/Plugins/Runtime/MassGameplay/Source/MassActors/Public/Translators/MassTranslators_BehaviorTree.h
@@ -10,7 +10,7 @@
 class UBehaviorTreeComponent;
 
 USTRUCT()
-struct FDataFragment_BehaviorTreeComponentWrapper : public FDataFragment_ObjectWrapper
+struct FDataFragment_BehaviorTreeComponentWrapper : public FObjectWrapperFragment
 {
 	GENERATED_BODY()
 	TWeakObjectPtr<UBehaviorTreeComponent> Component;
diff --git a/Engine/Plugins/Runtime/MassGameplay/Source/MassCommon/Public/MassCommonFragments.h b/Engine/Plugins/Runtime/MassGameplay/Source/MassCommon/Public/MassCommonFragments.h
--- a/Engine/Plugins/Runtime/MassGameplay/Source/MassCommon/Public/MassCommonFragments.h
+++ b/Engine/Plugins/Runtime/MassGameplay/Source/MassCommon/Public/MassCommonFragments.h
@@ -8,7 +8,7 @@
 
 
 USTRUCT()
-struct MASSCOMMON_API FDataFragment_Transform : public FMassFragment
+struct MASSCOMMON_API FTransformFragment : public FMassFragment
 {
 	GENERATED_BODY()
 
@@ -22,25 +22,7 @@ protected:
 };
 
 USTRUCT()
-struct MASSCOMMON_API FDataFragment_NavLocation : public FMassFragment
-{
-	GENERATED_BODY()
-	NavNodeRef NodeRef = INVALID_NAVNODEREF;
-};
-
-/** this fragment is a kind of an experiment. I've made FMassVelocityFragment and FMassAvoidanceFragment extend it and 
- *  use URandomizeVectorProcessor to initialize it to random variables at spawn time.
- */
-USTRUCT()
-struct MASSCOMMON_API FVectorFragment : public FMassFragment
-{
-	GENERATED_BODY()
-	UPROPERTY(EditAnywhere, Category = "")
-	FVector Value = FVector::ZeroVector;
-};
-
-USTRUCT()
-struct MASSCOMMON_API FDataFragment_AgentRadius : public FMassFragment
+struct MASSCOMMON_API FAgentRadiusFragment : public FMassFragment
 {
 	GENERATED_BODY()
 	UPROPERTY(EditAnywhere, Category = "")
@@ -51,7 +33,7 @@ struct MASSCOMMON_API FDataFragment_AgentRadius : public FMassFragment
  *	Mass simulation..
  */
 USTRUCT()
-struct FDataFragment_ObjectWrapper : public FMassFragment
+struct FObjectWrapperFragment : public FMassFragment
 {
 	GENERATED_BODY()
 };
diff --git a/Engine/Plugins/Runtime/MassGameplay/Source/MassGameplayDebug/Private/AssignDebugVisProcessor.cpp b/Engine/Plugins/Runtime/MassGameplay/Source/MassGameplayDebug/Private/AssignDebugVisProcessor.cpp
--- a/Engine/Plugins/Runtime/MassGameplay/Source/MassGameplayDebug/Private/AssignDebugVisProcessor.cpp
+++ b/Engine/Plugins/Runtime/MassGameplay/Source/MassGameplayDebug/Private/AssignDebugVisProcessor.cpp
@@ -12,13 +12,13 @@
 UAssignDebugVisProcessor::UAssignDebugVisProcessor()
 {
 	bAutoRegisterWithProcessingPhases = false;
-	ObservedType = FSimDebugVisComponent::StaticStruct();
+	ObservedType = FSimDebugVisFragment::StaticStruct();
 	Operation = EMassObservedOperation::Add;
 }
 
 void UAssignDebugVisProcessor::ConfigureQueries()
 {
-	EntityQuery.AddRequirement<FSimDebugVisComponent>(EMassFragmentAccess::ReadWrite);
+	EntityQuery.AddRequirement<FSimDebugVisFragment>(EMassFragmentAccess::ReadWrite);
 }
 
 void UAssignDebugVisProcessor::Execute(UMassEntitySubsystem& EntitySubsystem, FMassExecutionContext& Context)
@@ -37,8 +37,8 @@ void UAssignDebugVisProcessor::Execute(UMassEntitySubsystem& EntitySubsystem, FM
 
 	EntityQuery.ForEachEntityChunk(EntitySubsystem, Context, [this, Visualizer](FMassExecutionContext& Context)
 	{
-		const TArrayView<FSimDebugVisComponent> DebugVisList = Context.GetMutableFragmentView<FSimDebugVisComponent>();
-		for (FSimDebugVisComponent& VisualComp : DebugVisList)
+		const TArrayView<FSimDebugVisFragment> DebugVisList = Context.GetMutableFragmentView<FSimDebugVisFragment>();
+		for (FSimDebugVisFragment& VisualComp : DebugVisList)
 		{
 			// VisualComp.VisualType needs to be assigned by now. Should be performed as part of spawning, copied from the AgentTemplate
 			if (ensure(VisualComp.VisualType != INDEX_NONE))
diff --git a/Engine/Plugins/Runtime/MassGameplay/Source/MassGameplayDebug/Private/DebugVisLocationProcessor.cpp b/Engine/Plugins/Runtime/MassGameplay/Source/MassGameplayDebug/Private/DebugVisLocationProcessor.cpp
--- a/Engine/Plugins/Runtime/MassGameplay/Source/MassGameplayDebug/Private/DebugVisLocationProcessor.cpp
+++ b/Engine/Plugins/Runtime/MassGameplay/Source/MassGameplayDebug/Private/DebugVisLocationProcessor.cpp
@@ -18,8 +18,8 @@ UDebugVisLocationProcessor::UDebugVisLocationProcessor()
 
 void UDebugVisLocationProcessor::ConfigureQueries()
 {
-	EntityQuery.AddRequirement<FDataFragment_Transform>(EMassFragmentAccess::ReadOnly);
-	EntityQuery.AddRequirement<FSimDebugVisComponent>(EMassFragmentAccess::ReadOnly);
+	EntityQuery.AddRequirement<FTransformFragment>(EMassFragmentAccess::ReadOnly);
+	EntityQuery.AddRequirement<FSimDebugVisFragment>(EMassFragmentAccess::ReadOnly);
 	EntityQuery.AddTagRequirement<FMassDebuggableTag>(EMassFragmentPresence::All);
 }
 
@@ -40,12 +40,12 @@ void UDebugVisLocationProcessor::Execute(UMassEntitySubsystem& EntitySubsystem,
 		EntityQuery.ForEachEntityChunk(EntitySubsystem, Context, [this, &VisualDataISMCs](const FMassExecutionContext& Context)
 		{
 			const int32 NumEntities = Context.GetNumEntities();
-			const TConstArrayView<FDataFragment_Transform> LocationList = Context.GetFragmentView<FDataFragment_Transform>();
-			const TConstArrayView<FSimDebugVisComponent> DebugVisList = Context.GetFragmentView<FSimDebugVisComponent>();
+			const TConstArrayView<FTransformFragment> LocationList = Context.GetFragmentView<FTransformFragment>();
+			const TConstArrayView<FSimDebugVisFragment> DebugVisList = Context.GetFragmentView<FSimDebugVisFragment>();
 
 			for (int32 i = 0; i < NumEntities; ++i)
 			{
-				const FSimDebugVisComponent& VisualComp = DebugVisList[i];
+				const FSimDebugVisFragment& VisualComp = DebugVisList[i];
 
 				// @todo: remove this code once the asset is exported with correct alignment SM_Mannequin.uasset
 				FTransform SMTransform = LocationList[i].GetTransform();
@@ -90,9 +90,9 @@ UMassProcessor_UpdateDebugVis::UMassProcessor_UpdateDebugVis()
 void UMassProcessor_UpdateDebugVis::ConfigureQueries() 
 {
 	// @todo only FDataFragment_DebugVis should be mandatory, rest optional 
-	EntityQuery.AddRequirement<FDataFragment_Transform>(EMassFragmentAccess::ReadOnly);
+	EntityQuery.AddRequirement<FTransformFragment>(EMassFragmentAccess::ReadOnly);
 	EntityQuery.AddRequirement<FDataFragment_DebugVis>(EMassFragmentAccess::ReadWrite);
-	EntityQuery.AddRequirement<FDataFragment_AgentRadius>(EMassFragmentAccess::ReadWrite);
+	EntityQuery.AddRequirement<FAgentRadiusFragment>(EMassFragmentAccess::ReadWrite);
 	EntityQuery.AddTagRequirement<FMassDebuggableTag>(EMassFragmentPresence::All);
 }
 
@@ -109,9 +109,9 @@ void UMassProcessor_UpdateDebugVis::Execute(UMassEntitySubsystem& EntitySubsyste
 	EntityQuery.ForEachEntityChunk(EntitySubsystem, Context, [this, Debugger](FMassExecutionContext& Context)
 		{
 			const int32 NumEntities = Context.GetNumEntities();
-			const TConstArrayView<FDataFragment_Transform> LocationList = Context.GetFragmentView<FDataFragment_Transform>();
+			const TConstArrayView<FTransformFragment> LocationList = Context.GetFragmentView<FTransformFragment>();
 			const TArrayView<FDataFragment_DebugVis> DebugVisList = Context.GetMutableFragmentView<FDataFragment_DebugVis>();
-			const TArrayView<FDataFragment_AgentRadius> RadiiList = Context.GetMutableFragmentView<FDataFragment_AgentRadius>();
+			const TArrayView<FAgentRadiusFragment> RadiiList = Context.GetMutableFragmentView<FAgentRadiusFragment>();
 
 			for (int32 i = 0; i < NumEntities; ++i)
 			{
diff --git a/Engine/Plugins/Runtime/MassGameplay/Source/MassGameplayDebug/Private/MassDebugVisualizationTrait.cpp b/Engine/Plugins/Runtime/MassGameplay/Source/MassGameplayDebug/Private/MassDebugVisualizationTrait.cpp
--- a/Engine/Plugins/Runtime/MassGameplay/Source/MassGameplayDebug/Private/MassDebugVisualizationTrait.cpp
+++ b/Engine/Plugins/Runtime/MassGameplay/Source/MassGameplayDebug/Private/MassDebugVisualizationTrait.cpp
@@ -18,7 +18,7 @@ void UMassDebugVisualizationTrait::BuildTemplate(FMassEntityTemplateBuildContext
 	if (DebugMesh)
 	{
 #if WITH_EDITORONLY_DATA
-		FSimDebugVisComponent& DebugVisFragment = BuildContext.AddFragment_GetRef<FSimDebugVisComponent>();
+		FSimDebugVisFragment& DebugVisFragment = BuildContext.AddFragment_GetRef<FSimDebugVisFragment>();
 		UMassDebuggerSubsystem* Debugger = World.GetSubsystem<UMassDebuggerSubsystem>();
 		if (ensure(Debugger))
 		{
@@ -40,9 +40,9 @@ void UMassDebugVisualizationTrait::BuildTemplate(FMassEntityTemplateBuildContext
 	// DebugShape unavailable, will used default instead
 	BuildContext.AddFragment<FDataFragment_DebugVis>();
 #endif // WITH_EDITORONLY_DATA
-	BuildContext.AddFragment<FDataFragment_AgentRadius>();
+	BuildContext.AddFragment<FAgentRadiusFragment>();
 
-	BuildContext.AddFragment<FDataFragment_Transform>();
+	BuildContext.AddFragment<FTransformFragment>();
 #endif // if !(UE_BUILD_SHIPPING || UE_BUILD_TEST)
 
 }
\ No newline at end of file
diff --git a/Engine/Plugins/Runtime/MassGameplay/Source/MassGameplayDebug/Private/MassGameplayDebugTypes.cpp b/Engine/Plugins/Runtime/MassGameplay/Source/MassGameplayDebug/Private/MassGameplayDebugTypes.cpp
--- a/Engine/Plugins/Runtime/MassGameplay/Source/MassGameplayDebug/Private/MassGameplayDebugTypes.cpp
+++ b/Engine/Plugins/Runtime/MassGameplay/Source/MassGameplayDebug/Private/MassGameplayDebugTypes.cpp
@@ -27,7 +27,7 @@ void GetDebugEntitiesAndLocations(const UMassEntitySubsystem& EntitySystem, TArr
 		const FMassEntityHandle EntityHandle = ConvertEntityIndexToHandle(EntitySystem, i);
 		if (EntityHandle.IsSet())
 		{
-			if (const FDataFragment_Transform* TransformFragment = EntitySystem.GetFragmentDataPtr<FDataFragment_Transform>(EntityHandle))
+			if (const FTransformFragment* TransformFragment = EntitySystem.GetFragmentDataPtr<FTransformFragment>(EntityHandle))
 			{
 				OutEntities.Add(EntityHandle);
 				OutLocations.Add(TransformFragment->GetTransform().GetLocation());
diff --git a/Engine/Plugins/Runtime/MassGameplay/Source/MassGameplayDebug/Public/AssignDebugVisProcessor.h b/Engine/Plugins/Runtime/MassGameplay/Source/MassGameplayDebug/Public/AssignDebugVisProcessor.h
--- a/Engine/Plugins/Runtime/MassGameplay/Source/MassGameplayDebug/Public/AssignDebugVisProcessor.h
+++ b/Engine/Plugins/Runtime/MassGameplay/Source/MassGameplayDebug/Public/AssignDebugVisProcessor.h
@@ -7,7 +7,7 @@
 
 
 class UMassDebugVisualizationComponent;
-struct FSimDebugVisComponent;
+struct FSimDebugVisFragment;
 
 UCLASS()
 class MASSGAMEPLAYDEBUG_API UAssignDebugVisProcessor : public UMassObserverProcessor
diff --git a/Engine/Plugins/Runtime/MassGameplay/Source/MassGameplayDebug/Public/DebugVisLocationProcessor.h b/Engine/Plugins/Runtime/MassGameplay/Source/MassGameplayDebug/Public/DebugVisLocationProcessor.h
--- a/Engine/Plugins/Runtime/MassGameplay/Source/MassGameplayDebug/Public/DebugVisLocationProcessor.h
+++ b/Engine/Plugins/Runtime/MassGameplay/Source/MassGameplayDebug/Public/DebugVisLocationProcessor.h
@@ -6,7 +6,7 @@
 #include "DebugVisLocationProcessor.generated.h"
 
 class UMassDebugVisualizationComponent;
-struct FSimDebugVisComponent;
+struct FSimDebugVisFragment;
 
 UCLASS()
 class MASSGAMEPLAYDEBUG_API UDebugVisLocationProcessor : public UMassProcessor
diff --git a/Engine/Plugins/Runtime/MassGameplay/Source/MassGameplayDebug/Public/MassGameplayDebugTypes.h b/Engine/Plugins/Runtime/MassGameplay/Source/MassGameplayDebug/Public/MassGameplayDebugTypes.h
--- a/Engine/Plugins/Runtime/MassGameplay/Source/MassGameplayDebug/Public/MassGameplayDebugTypes.h
+++ b/Engine/Plugins/Runtime/MassGameplay/Source/MassGameplayDebug/Public/MassGameplayDebugTypes.h
@@ -47,7 +47,7 @@ struct FSimDebugDataRow : public FTableRowBase
 };
 
 USTRUCT()
-struct FSimDebugVisComponent : public FMassFragment
+struct FSimDebugVisFragment : public FMassFragment
 {
 	GENERATED_BODY()
 	int32 InstanceIndex = INDEX_NONE;
diff --git a/Engine/Plugins/Runtime/MassGameplay/Source/MassLOD/Private/MassLODCollectorProcessor.cpp b/Engine/Plugins/Runtime/MassGameplay/Source/MassLOD/Private/MassLODCollectorProcessor.cpp
--- a/Engine/Plugins/Runtime/MassGameplay/Source/MassLOD/Private/MassLODCollectorProcessor.cpp
+++ b/Engine/Plugins/Runtime/MassGameplay/Source/MassLOD/Private/MassLODCollectorProcessor.cpp
@@ -18,7 +18,7 @@ void UMassLODCollectorProcessor::ConfigureQueries()
 {
 	FMassEntityQuery BaseQuery;
 	BaseQuery.AddTagRequirement<FMassCollectLODViewerInfoTag>(EMassFragmentPresence::All);
-	BaseQuery.AddRequirement<FDataFragment_Transform>(EMassFragmentAccess::ReadOnly);
+	BaseQuery.AddRequirement<FTransformFragment>(EMassFragmentAccess::ReadOnly);
 	BaseQuery.AddRequirement<FMassViewerInfoFragment>(EMassFragmentAccess::ReadWrite);
 	BaseQuery.AddChunkRequirement<FMassSimulationVariableTickChunkFragment>(EMassFragmentAccess::ReadOnly, EMassFragmentPresence::Optional);
 	BaseQuery.AddChunkRequirement<FMassVisualizationChunkFragment>(EMassFragmentAccess::ReadOnly, EMassFragmentPresence::Optional);
@@ -48,10 +48,10 @@ void UMassLODCollectorProcessor::ConfigureQueries()
 template <bool bLocalViewersOnly>
 void UMassLODCollectorProcessor::CollectLODForChunk(FMassExecutionContext& Context)
 {
-	TConstArrayView<FDataFragment_Transform> LocationList = Context.GetFragmentView<FDataFragment_Transform>();
+	TConstArrayView<FTransformFragment> LocationList = Context.GetFragmentView<FTransformFragment>();
 	TArrayView<FMassViewerInfoFragment> ViewerInfoList = Context.GetMutableFragmentView<FMassViewerInfoFragment>();
 
-	Collector.CollectLODInfo<FDataFragment_Transform, FMassViewerInfoFragment, bLocalViewersOnly, true/*bCollectDistanceToViewer*/>(Context, LocationList, ViewerInfoList);
+	Collector.CollectLODInfo<FTransformFragment, FMassViewerInfoFragment, bLocalViewersOnly, true/*bCollectDistanceToViewer*/>(Context, LocationList, ViewerInfoList);
 }
 
 template <bool bLocalViewersOnly>
diff --git a/Engine/Plugins/Runtime/MassGameplay/Source/MassLOD/Private/MassLODTrait.cpp b/Engine/Plugins/Runtime/MassGameplay/Source/MassLOD/Private/MassLODTrait.cpp
--- a/Engine/Plugins/Runtime/MassGameplay/Source/MassLOD/Private/MassLODTrait.cpp
+++ b/Engine/Plugins/Runtime/MassGameplay/Source/MassLOD/Private/MassLODTrait.cpp
@@ -16,7 +16,7 @@ void UMassLODCollectorTrait::BuildTemplate(FMassEntityTemplateBuildContext& Buil
 void UMassSimulationLODTrait::BuildTemplate(FMassEntityTemplateBuildContext& BuildContext, UWorld& World) const
 {
 	BuildContext.AddTag<FMassCollectLODViewerInfoTag>(); // Depends on FMassViewerInfoFragment
-	BuildContext.AddFragment<FDataFragment_Transform>();
+	BuildContext.AddFragment<FTransformFragment>();
 
 	FMassSimulationLODFragment& LODFragment = BuildContext.AddFragment_GetRef<FMassSimulationLODFragment>();
 
diff --git a/Engine/Plugins/Runtime/MassGameplay/Source/MassLOD/Private/MassRepresentationLODCollectorProcessor.cpp b/Engine/Plugins/Runtime/MassGameplay/Source/MassLOD/Private/MassRepresentationLODCollectorProcessor.cpp
--- a/Engine/Plugins/Runtime/MassGameplay/Source/MassLOD/Private/MassRepresentationLODCollectorProcessor.cpp
+++ b/Engine/Plugins/Runtime/MassGameplay/Source/MassLOD/Private/MassRepresentationLODCollectorProcessor.cpp
@@ -19,7 +19,7 @@ void UMassRepresentationLODCollectorProcessor::ConfigureQueries()
 		{
 			BaseQuery.AddTagRequirement(*LODConfig.TagFilter.GetScriptStruct(), EMassFragmentPresence::All);
 		}
-		BaseQuery.AddRequirement<FDataFragment_Transform>(EMassFragmentAccess::ReadOnly);
+		BaseQuery.AddRequirement<FTransformFragment>(EMassFragmentAccess::ReadOnly);
 		BaseQuery.AddRequirement<FMassViewerInfoFragment>(EMassFragmentAccess::ReadWrite);
 		BaseQuery.AddChunkRequirement<FMassVisualizationChunkFragment>(EMassFragmentAccess::ReadOnly);
 
@@ -38,7 +38,7 @@ void CollectLODInfo(UMassEntitySubsystem& EntitySubsystem, FMassExecutionContext
 
 	auto InternalCollectLODInfo = [&Collector](FMassExecutionContext& Context)
 	{
-		TConstArrayView<FDataFragment_Transform> LocationList = Context.GetFragmentView<FDataFragment_Transform>();
+		TConstArrayView<FTransformFragment> LocationList = Context.GetFragmentView<FTransformFragment>();
 		TArrayView<FMassViewerInfoFragment> ViewerInfoList = Context.GetMutableFragmentView<FMassViewerInfoFragment>();
 
 		Collector.CollectLODInfo(Context, LocationList, ViewerInfoList);
diff --git a/Engine/Plugins/Runtime/MassGameplay/Source/MassLOD/Private/MassSimulationLOD.cpp b/Engine/Plugins/Runtime/MassGameplay/Source/MassLOD/Private/MassSimulationLOD.cpp
--- a/Engine/Plugins/Runtime/MassGameplay/Source/MassLOD/Private/MassSimulationLOD.cpp
+++ b/Engine/Plugins/Runtime/MassGameplay/Source/MassLOD/Private/MassSimulationLOD.cpp
@@ -9,9 +9,9 @@
 #include "MassLODUtils.h"
 
 //-----------------------------------------------------------------------------
-// FMassSimulationLODConfig
+// FMassSimulationLODParameters
 //-----------------------------------------------------------------------------
-FMassSimulationLODConfig::FMassSimulationLODConfig()
+FMassSimulationLODParameters::FMassSimulationLODParameters()
 {
 	LODDistance[EMassLOD::High] = 0.0f;
 	LODDistance[EMassLOD::Medium] = 5000.0f;
@@ -25,9 +25,9 @@ FMassSimulationLODConfig::FMassSimulationLODConfig()
 }
 
 //-----------------------------------------------------------------------------
-// FMassSimulationVariableTickConfig
+// FMassSimulationVariableTickParameters
 //-----------------------------------------------------------------------------
-FMassSimulationVariableTickConfig::FMassSimulationVariableTickConfig()
+FMassSimulationVariableTickParameters::FMassSimulationVariableTickParameters()
 {
 	TickRates[EMassLOD::High] = 0.0f;
 	TickRates[EMassLOD::Medium] = 0.5f;
@@ -38,7 +38,7 @@ FMassSimulationVariableTickConfig::FMassSimulationVariableTickConfig()
 //-----------------------------------------------------------------------------
 // FMassSimulationLODSharedFragment
 //-----------------------------------------------------------------------------
-FMassSimulationLODSharedFragment::FMassSimulationLODSharedFragment(const FMassSimulationLODConfig& LODConfig)
+FMassSimulationLODSharedFragment::FMassSimulationLODSharedFragment(const FMassSimulationLODParameters& LODConfig)
 {
 	LODCalculator.Initialize(LODConfig.LODDistance, LODConfig.BufferHysteresisOnDistancePercentage / 100.0f, LODConfig.LODMaxCount);
 }
@@ -46,7 +46,7 @@ FMassSimulationLODSharedFragment::FMassSimulationLODSharedFragment(const FMassSi
 //-----------------------------------------------------------------------------
 // FMassSimulationLODSharedFragment
 //-----------------------------------------------------------------------------
-FMassSimulationVariableTickSharedFragment::FMassSimulationVariableTickSharedFragment(const FMassSimulationVariableTickConfig& TickRateConfig)
+FMassSimulationVariableTickSharedFragment::FMassSimulationVariableTickSharedFragment(const FMassSimulationVariableTickParameters& TickRateConfig)
 {
 	LODTickRateController.Initialize(TickRateConfig.TickRates, TickRateConfig.bSpreadFirstSimulationUpdate);
 }
@@ -70,10 +70,10 @@ UMassSimulationLODProcessor::UMassSimulationLODProcessor()
 
 void UMassSimulationLODProcessor::ConfigureQueries()
 {
-	EntityQuery.AddRequirement<FDataFragment_Transform>(EMassFragmentAccess::ReadOnly);
+	EntityQuery.AddRequirement<FTransformFragment>(EMassFragmentAccess::ReadOnly);
 	EntityQuery.AddRequirement<FMassViewerInfoFragment>(EMassFragmentAccess::ReadOnly);
 	EntityQuery.AddRequirement<FMassSimulationLODFragment>(EMassFragmentAccess::ReadWrite);
-	EntityQuery.AddConstSharedRequirement<FMassSimulationLODConfig>();
+	EntityQuery.AddConstSharedRequirement<FMassSimulationLODParameters>();
 	EntityQuery.AddSharedRequirement<FMassSimulationLODSharedFragment>(EMassFragmentAccess::ReadWrite);
 	EntityQuery.AddChunkRequirement<FMassSimulationVariableTickChunkFragment>(EMassFragmentAccess::ReadOnly, EMassFragmentPresence::Optional);
 	EntityQuery.AddSharedRequirement<FMassSimulationVariableTickSharedFragment>(EMassFragmentAccess::ReadOnly, EMassFragmentPresence::Optional);
@@ -91,17 +91,17 @@ void UMassSimulationLODProcessor::ConfigureQueries()
 
 	EntityQueryVariableTick.AddRequirement<FMassSimulationLODFragment>(EMassFragmentAccess::ReadOnly);
 	EntityQueryVariableTick.AddRequirement<FMassSimulationVariableTickFragment>(EMassFragmentAccess::ReadWrite);
-	EntityQueryVariableTick.AddConstSharedRequirement<FMassSimulationVariableTickConfig>();
+	EntityQueryVariableTick.AddConstSharedRequirement<FMassSimulationVariableTickParameters>();
 	EntityQueryVariableTick.AddChunkRequirement<FMassSimulationVariableTickChunkFragment>(EMassFragmentAccess::ReadWrite);
 	EntityQueryVariableTick.AddSharedRequirement<FMassSimulationVariableTickSharedFragment>(EMassFragmentAccess::ReadOnly);
 
 	// In case where the variableTick isn't enabled, we might need to set LOD tags as if the users still wants them
 	EntityQuerySetLODTag.AddRequirement<FMassSimulationLODFragment>(EMassFragmentAccess::ReadOnly);
 	EntityQuerySetLODTag.AddRequirement<FMassSimulationVariableTickFragment>(EMassFragmentAccess::ReadWrite, EMassFragmentPresence::None);
-	EntityQuerySetLODTag.AddConstSharedRequirement<FMassSimulationLODConfig>();
+	EntityQuerySetLODTag.AddConstSharedRequirement<FMassSimulationLODParameters>();
 	EntityQuerySetLODTag.SetArchetypeFilter([](const FMassExecutionContext& Context)
 	{
-		const FMassSimulationLODConfig& LODConfig = Context.GetConstSharedFragment<FMassSimulationLODConfig>();
+		const FMassSimulationLODParameters& LODConfig = Context.GetConstSharedFragment<FMassSimulationLODParameters>();
 		return LODConfig.bSetLODTags;
 	});
 }
@@ -188,7 +188,7 @@ void UMassSimulationLODProcessor::Execute(UMassEntitySubsystem& EntitySubsystem,
 		EntityQuery.ForEachEntityChunk(EntitySubsystem, Context, [this](FMassExecutionContext& Context)
 		{
 			FMassSimulationLODSharedFragment& LODSharedFragment = Context.GetMutableSharedFragment<FMassSimulationLODSharedFragment>();
-			TConstArrayView<FDataFragment_Transform> LocationList = Context.GetFragmentView<FDataFragment_Transform>();
+			TConstArrayView<FTransformFragment> LocationList = Context.GetFragmentView<FTransformFragment>();
 			TConstArrayView<FMassSimulationLODFragment> SimulationLODList = Context.GetFragmentView<FMassSimulationLODFragment>();
 			LODSharedFragment.LODCalculator.DebugDisplayLOD(Context, SimulationLODList, LocationList, World);
 		});
diff --git a/Engine/Plugins/Runtime/MassGameplay/Source/MassLOD/Public/MassLODTrait.h b/Engine/Plugins/Runtime/MassGameplay/Source/MassLOD/Public/MassLODTrait.h
--- a/Engine/Plugins/Runtime/MassGameplay/Source/MassLOD/Public/MassLODTrait.h
+++ b/Engine/Plugins/Runtime/MassGameplay/Source/MassLOD/Public/MassLODTrait.h
@@ -22,13 +22,13 @@ class MASSLOD_API UMassSimulationLODTrait : public UMassEntityTraitBase
 	GENERATED_BODY()
 
 	UPROPERTY(Category = "Config", EditAnywhere)
-	FMassSimulationLODConfig Config;
+	FMassSimulationLODParameters Config;
 
 	UPROPERTY(Category = "Config", EditAnywhere)
 	bool bEnableVariableTicking = false;
 
 	UPROPERTY(Category = "Config", EditAnywhere, meta = (EditCondition = "bEnableVariableTicking", EditConditionHides))
-	FMassSimulationVariableTickConfig VariableTickConfig;
+	FMassSimulationVariableTickParameters VariableTickConfig;
 
 protected:
 	virtual void BuildTemplate(FMassEntityTemplateBuildContext& BuildContext, UWorld& World) const override;
diff --git a/Engine/Plugins/Runtime/MassGameplay/Source/MassLOD/Public/MassSimulationLOD.h b/Engine/Plugins/Runtime/MassGameplay/Source/MassLOD/Public/MassSimulationLOD.h
--- a/Engine/Plugins/Runtime/MassGameplay/Source/MassLOD/Public/MassSimulationLOD.h
+++ b/Engine/Plugins/Runtime/MassGameplay/Source/MassLOD/Public/MassSimulationLOD.h
@@ -61,11 +61,11 @@ struct MASSLOD_API FMassSimulationVariableTickChunkFragment : public FMassVariab
 };
 
 USTRUCT()
-struct MASSLOD_API FMassSimulationLODConfig : public FMassSharedFragment
+struct MASSLOD_API FMassSimulationLODParameters : public FMassSharedFragment
 {
 	GENERATED_BODY()
 
-	FMassSimulationLODConfig();
+	FMassSimulationLODParameters();
 
 	/** Distance where each LOD becomes relevant */
 	UPROPERTY(EditAnywhere, Category = "LOD", config)
@@ -85,11 +85,11 @@ struct MASSLOD_API FMassSimulationLODConfig : public FMassSharedFragment
 };
 
 USTRUCT()
-struct MASSLOD_API FMassSimulationVariableTickConfig : public FMassSharedFragment
+struct MASSLOD_API FMassSimulationVariableTickParameters : public FMassSharedFragment
 {
 	GENERATED_BODY()
 
-	FMassSimulationVariableTickConfig();
+	FMassSimulationVariableTickParameters();
 
 	/** Rate in seconds at which entities should update when in this LOD */
 	UPROPERTY(EditAnywhere, Category = "VariableTick", config)
@@ -106,7 +106,7 @@ struct MASSLOD_API FMassSimulationLODSharedFragment : public FMassSharedFragment
 	GENERATED_BODY()
 
 	FMassSimulationLODSharedFragment() = default;
-	FMassSimulationLODSharedFragment(const FMassSimulationLODConfig& Config);
+	FMassSimulationLODSharedFragment(const FMassSimulationLODParameters& Config);
 
 	/** Runtime data for matching the simulation LOD config */
 	TMassLODCalculator<FMassSimulationLODLogic> LODCalculator;
@@ -119,7 +119,7 @@ struct MASSLOD_API FMassSimulationVariableTickSharedFragment : public FMassShare
 	GENERATED_BODY()
 
 	FMassSimulationVariableTickSharedFragment() = default;
-	FMassSimulationVariableTickSharedFragment(const FMassSimulationVariableTickConfig& Config);
+	FMassSimulationVariableTickSharedFragment(const FMassSimulationVariableTickParameters& Config);
 
 	/** Runtime data for matching the simulation tick rate config */
 	TMassLODTickRateController<FMassSimulationVariableTickChunkFragment, FMassSimulationLODLogic> LODTickRateController;
@@ -156,7 +156,7 @@ protected:
 	virtual void ConfigureQueries() override;
 	virtual void Execute(UMassEntitySubsystem& EntitySubsystem, FMassExecutionContext& Context) override;
 
-	void CalculateLODForConfig(UMassEntitySubsystem& EntitySubsystem, FMassExecutionContext& Context, FMassSimulationLODConfig& LODConfig);
+	void CalculateLODForConfig(UMassEntitySubsystem& EntitySubsystem, FMassExecutionContext& Context, FMassSimulationLODParameters& LODConfig);
 
 	FMassEntityQuery EntityQuery;
 	FMassEntityQuery EntityQueryCalculateLOD;
diff --git a/Engine/Plugins/Runtime/MassGameplay/Source/MassMovement/Private/MassMovementProcessors.cpp b/Engine/Plugins/Runtime/MassGameplay/Source/MassMovement/Private/MassMovementProcessors.cpp
--- a/Engine/Plugins/Runtime/MassGameplay/Source/MassMovement/Private/MassMovementProcessors.cpp
+++ b/Engine/Plugins/Runtime/MassGameplay/Source/MassMovement/Private/MassMovementProcessors.cpp
@@ -25,7 +25,7 @@ UMassApplyMovementProcessor::UMassApplyMovementProcessor()
 void UMassApplyMovementProcessor::ConfigureQueries()
 {
 	EntityQuery.AddRequirement<FMassVelocityFragment>(EMassFragmentAccess::ReadWrite);
-	EntityQuery.AddRequirement<FDataFragment_Transform>(EMassFragmentAccess::ReadWrite);
+	EntityQuery.AddRequirement<FTransformFragment>(EMassFragmentAccess::ReadWrite);
 	EntityQuery.AddRequirement<FMassForceFragment>(EMassFragmentAccess::ReadWrite);
 	EntityQuery.AddTagRequirement<FMassOffLODTag>(EMassFragmentPresence::None);
 	EntityQuery.AddConstSharedRequirement<FMassMovementParameters>(EMassFragmentPresence::All);
@@ -45,7 +45,7 @@ void UMassApplyMovementProcessor::Execute(UMassEntitySubsystem& EntitySubsystem,
 
 		const FMassMovementParameters& MovementParams = Context.GetConstSharedFragment<FMassMovementParameters>();
 
-		const TArrayView<FDataFragment_Transform> LocationList = Context.GetMutableFragmentView<FDataFragment_Transform>();
+		const TArrayView<FTransformFragment> LocationList = Context.GetMutableFragmentView<FTransformFragment>();
 		const TArrayView<FMassForceFragment> ForceList = Context.GetMutableFragmentView<FMassForceFragment>();
 		const TArrayView<FMassVelocityFragment> VelocityList = Context.GetMutableFragmentView<FMassVelocityFragment>();
 
diff --git a/Engine/Plugins/Runtime/MassGameplay/Source/MassMovement/Private/MassMovementTrait.cpp b/Engine/Plugins/Runtime/MassGameplay/Source/MassMovement/Private/MassMovementTrait.cpp
--- a/Engine/Plugins/Runtime/MassGameplay/Source/MassMovement/Private/MassMovementTrait.cpp
+++ b/Engine/Plugins/Runtime/MassGameplay/Source/MassMovement/Private/MassMovementTrait.cpp
@@ -11,8 +11,8 @@ void UMassMovementTrait::BuildTemplate(FMassEntityTemplateBuildContext& BuildCon
 	UMassEntitySubsystem* EntitySubsystem = UWorld::GetSubsystem<UMassEntitySubsystem>(&World);
 	check(EntitySubsystem);
 
-	BuildContext.AddFragment<FDataFragment_AgentRadius>();
-	BuildContext.AddFragment<FDataFragment_Transform>();
+	BuildContext.AddFragment<FAgentRadiusFragment>();
+	BuildContext.AddFragment<FTransformFragment>();
 
 	BuildContext.AddFragment<FMassVelocityFragment>();
 	BuildContext.AddFragment<FMassForceFragment>();
diff --git a/Engine/Plugins/Runtime/MassGameplay/Source/MassMovement/Private/MassSimpleMovementTrait.cpp b/Engine/Plugins/Runtime/MassGameplay/Source/MassMovement/Private/MassSimpleMovementTrait.cpp
--- a/Engine/Plugins/Runtime/MassGameplay/Source/MassMovement/Private/MassSimpleMovementTrait.cpp
+++ b/Engine/Plugins/Runtime/MassGameplay/Source/MassMovement/Private/MassSimpleMovementTrait.cpp
@@ -13,7 +13,7 @@
 //----------------------------------------------------------------------//
 void UMassSimpleMovementTrait::BuildTemplate(FMassEntityTemplateBuildContext& BuildContext, UWorld& World) const
 {
-	BuildContext.AddFragment<FDataFragment_Transform>();
+	BuildContext.AddFragment<FTransformFragment>();
 	BuildContext.AddFragment<FMassVelocityFragment>();
 	BuildContext.AddTag<FMassSimpleMovementTag>();	
 }
@@ -31,7 +31,7 @@ UMassSimpleMovementProcessor::UMassSimpleMovementProcessor()
 void UMassSimpleMovementProcessor::ConfigureQueries()
 {
 	EntityQuery.AddRequirement<FMassVelocityFragment>(EMassFragmentAccess::ReadOnly);
-	EntityQuery.AddRequirement<FDataFragment_Transform>(EMassFragmentAccess::ReadWrite);
+	EntityQuery.AddRequirement<FTransformFragment>(EMassFragmentAccess::ReadWrite);
 	EntityQuery.AddTagRequirement<FMassSimpleMovementTag>(EMassFragmentPresence::All);
 
 	EntityQuery.AddRequirement<FMassSimulationVariableTickFragment>(EMassFragmentAccess::ReadOnly, EMassFragmentPresence::Optional);
@@ -44,7 +44,7 @@ void UMassSimpleMovementProcessor::Execute(UMassEntitySubsystem& EntitySubsystem
 	EntityQuery.ForEachEntityChunk(EntitySubsystem, Context, ([this](FMassExecutionContext& Context)
 		{
 			const TConstArrayView<FMassVelocityFragment> VelocitiesList = Context.GetFragmentView<FMassVelocityFragment>();
-			const TArrayView<FDataFragment_Transform> TransformsList = Context.GetMutableFragmentView<FDataFragment_Transform>();
+			const TArrayView<FTransformFragment> TransformsList = Context.GetMutableFragmentView<FTransformFragment>();
 			const TConstArrayView<FMassSimulationVariableTickFragment> SimVariableTickList = Context.GetFragmentView<FMassSimulationVariableTickFragment>();
 			const bool bHasVariableTick = (SimVariableTickList.Num() > 0);
 			const float WorldDeltaTime = Context.GetDeltaTimeSeconds();
diff --git a/Engine/Plugins/Runtime/MassGameplay/Source/MassReplication/Private/MassReplicationProcessorBase.cpp b/Engine/Plugins/Runtime/MassGameplay/Source/MassReplication/Private/MassReplicationProcessorBase.cpp
--- a/Engine/Plugins/Runtime/MassGameplay/Source/MassReplication/Private/MassReplicationProcessorBase.cpp
+++ b/Engine/Plugins/Runtime/MassGameplay/Source/MassReplication/Private/MassReplicationProcessorBase.cpp
@@ -16,7 +16,7 @@ UMassReplicationProcessorBase::UMassReplicationProcessorBase()
 
 void UMassReplicationProcessorBase::ConfigureQueries()
 {
-	CollectViewerInfoQuery.AddRequirement<FDataFragment_Transform>(EMassFragmentAccess::ReadOnly);
+	CollectViewerInfoQuery.AddRequirement<FTransformFragment>(EMassFragmentAccess::ReadOnly);
 	CollectViewerInfoQuery.AddRequirement<FMassReplicationViewerInfoFragment>(EMassFragmentAccess::ReadOnly);
 	CollectViewerInfoQuery.AddSharedRequirement<FMassReplicationSharedFragment>(EMassFragmentAccess::ReadWrite);
 
@@ -35,7 +35,7 @@ void UMassReplicationProcessorBase::ConfigureQueries()
 	});
 
 	EntityQuery.AddRequirement<FMassNetworkIDFragment>(EMassFragmentAccess::ReadOnly);
-	EntityQuery.AddRequirement<FDataFragment_ReplicationTemplateID>(EMassFragmentAccess::ReadOnly);
+	EntityQuery.AddRequirement<FReplicationTemplateIDFragment>(EMassFragmentAccess::ReadOnly);
 	EntityQuery.AddRequirement<FMassReplicationLODFragment>(EMassFragmentAccess::ReadWrite);
 	EntityQuery.AddRequirement<FMassReplicatedAgentFragment>(EMassFragmentAccess::ReadWrite);
 	EntityQuery.AddConstSharedRequirement<FMassReplicationParameters>();
@@ -137,7 +137,7 @@ void UMassReplicationProcessorBase::Execute(UMassEntitySubsystem& EntitySubsyste
 
 	CollectViewerInfoQuery.ForEachEntityChunk(EntitySubsystem, Context, [this](FMassExecutionContext& Context)
 	{
-		const TConstArrayView<FDataFragment_Transform> LocationList = Context.GetFragmentView<FDataFragment_Transform>();
+		const TConstArrayView<FTransformFragment> LocationList = Context.GetFragmentView<FTransformFragment>();
 		const TArrayView<FMassReplicationViewerInfoFragment> ViewersInfoList = Context.GetMutableFragmentView<FMassReplicationViewerInfoFragment>();
 		FMassReplicationSharedFragment& RepSharedFragment = Context.GetMutableSharedFragment<FMassReplicationSharedFragment>();
 		RepSharedFragment.LODCollector.CollectLODInfo(Context, LocationList, ViewersInfoList, ViewersInfoList);
diff --git a/Engine/Plugins/Runtime/MassGameplay/Source/MassReplication/Private/MassReplicationTrait.cpp b/Engine/Plugins/Runtime/MassGameplay/Source/MassReplication/Private/MassReplicationTrait.cpp
--- a/Engine/Plugins/Runtime/MassGameplay/Source/MassReplication/Private/MassReplicationTrait.cpp
+++ b/Engine/Plugins/Runtime/MassGameplay/Source/MassReplication/Private/MassReplicationTrait.cpp
@@ -16,7 +16,7 @@ void UMassReplicationTrait::BuildTemplate(FMassEntityTemplateBuildContext& Build
 		return;
 	}
 
-	FDataFragment_ReplicationTemplateID& TemplateIDFragment = BuildContext.AddFragment_GetRef<FDataFragment_ReplicationTemplateID>();
+	FReplicationTemplateIDFragment& TemplateIDFragment = BuildContext.AddFragment_GetRef<FReplicationTemplateIDFragment>();
 	TemplateIDFragment.ID = BuildContext.GetTemplateID();
 
 	BuildContext.AddFragment<FMassNetworkIDFragment>();
diff --git a/Engine/Plugins/Runtime/MassGameplay/Source/MassReplication/Private/MassReplicationTransformHandlers.cpp b/Engine/Plugins/Runtime/MassGameplay/Source/MassReplication/Private/MassReplicationTransformHandlers.cpp
--- a/Engine/Plugins/Runtime/MassGameplay/Source/MassReplication/Private/MassReplicationTransformHandlers.cpp
+++ b/Engine/Plugins/Runtime/MassGameplay/Source/MassReplication/Private/MassReplicationTransformHandlers.cpp
@@ -7,17 +7,17 @@
 
 void FMassReplicationProcessorTransformHandlerBase::AddRequirements(FMassEntityQuery& InQuery)
 {
-	InQuery.AddRequirement<FDataFragment_Transform>(EMassFragmentAccess::ReadWrite);
+	InQuery.AddRequirement<FTransformFragment>(EMassFragmentAccess::ReadWrite);
 }
 
 void FMassReplicationProcessorTransformHandlerBase::CacheFragmentViews(FMassExecutionContext& ExecContext)
 {
-	TransformList = ExecContext.GetMutableFragmentView<FDataFragment_Transform>();
+	TransformList = ExecContext.GetMutableFragmentView<FTransformFragment>();
 }
 
 void FMassReplicationProcessorPositionYawHandler::AddEntity(const int32 EntityIdx, FReplicatedAgentPositionYawData& InOutReplicatedPositionYawData) const
 {
-	const FDataFragment_Transform& TransformFragment = TransformList[EntityIdx];
+	const FTransformFragment& TransformFragment = TransformList[EntityIdx];
 	InOutReplicatedPositionYawData.SetPosition(TransformFragment.GetTransform().GetLocation());
 
 	if (const TOptional<float> Yaw = UE::AI::GetYawFromQuaternion(TransformFragment.GetTransform().GetRotation()))
diff --git a/Engine/Plugins/Runtime/MassGameplay/Source/MassReplication/Public/MassClientBubbleHandler.h b/Engine/Plugins/Runtime/MassGameplay/Source/MassReplication/Public/MassClientBubbleHandler.h
--- a/Engine/Plugins/Runtime/MassGameplay/Source/MassReplication/Public/MassClientBubbleHandler.h
+++ b/Engine/Plugins/Runtime/MassGameplay/Source/MassReplication/Public/MassClientBubbleHandler.h
@@ -744,7 +744,7 @@ void TClientBubbleHandlerBase<AgentArrayItem>::DebugValidateBubbleOnServer()
 		checkf(FragmentNetID.NetID == OuterItem.Agent.GetNetID(), TEXT("Fragment and Agent NetID do not match!"));
 		checkf(LookupData.NetID == OuterItem.Agent.GetNetID(), TEXT("LookupData and Agent NetID do not match!"));
 
-		const FDataFragment_ReplicationTemplateID& FragmentTemplateID = EntitySystem->GetFragmentDataChecked<FDataFragment_ReplicationTemplateID>(LookupData.Entity);
+		const FReplicationTemplateIDFragment& FragmentTemplateID = EntitySystem->GetFragmentDataChecked<FReplicationTemplateIDFragment>(LookupData.Entity);
 		checkf(FragmentTemplateID.ID == OuterItem.Agent.GetTemplateID(), TEXT("Agent TemplateID different to Fragment!"));
 	}
 
diff --git a/Engine/Plugins/Runtime/MassGameplay/Source/MassReplication/Public/MassReplicationProcessorBase.h b/Engine/Plugins/Runtime/MassGameplay/Source/MassReplication/Public/MassReplicationProcessorBase.h
--- a/Engine/Plugins/Runtime/MassGameplay/Source/MassReplication/Public/MassReplicationProcessorBase.h
+++ b/Engine/Plugins/Runtime/MassGameplay/Source/MassReplication/Public/MassReplicationProcessorBase.h
@@ -70,7 +70,7 @@ void UMassReplicationProcessorBase::CalculateClientReplication(FMassExecutionCon
 	TConstArrayView<FMassNetworkIDFragment> NetworkIDList = Context.GetFragmentView<FMassNetworkIDFragment>();
 	TArrayView<FMassReplicationLODFragment> ViewerLODList = Context.GetMutableFragmentView<FMassReplicationLODFragment>();
 	TArrayView<FMassReplicatedAgentFragment> ReplicatedAgentList = Context.GetMutableFragmentView<FMassReplicatedAgentFragment>();
-	TConstArrayView<FDataFragment_ReplicationTemplateID> TemplateIDList = Context.GetFragmentView<FDataFragment_ReplicationTemplateID>();
+	TConstArrayView<FReplicationTemplateIDFragment> TemplateIDList = Context.GetFragmentView<FReplicationTemplateIDFragment>();
 	FMassReplicationSharedFragment& RepSharedFragment = Context.GetMutableSharedFragment<FMassReplicationSharedFragment>();
 
 	CacheViews(Context);
@@ -139,7 +139,7 @@ void UMassReplicationProcessorBase::CalculateClientReplication(FMassExecutionCon
 						typename AgentArrayItem::FReplicatedAgentType ReplicatedAgent;
 
 						const FMassNetworkIDFragment& NetIDFragment = NetworkIDList[EntityIdx];
-						const FDataFragment_ReplicationTemplateID& TemplateIDFragment = TemplateIDList[EntityIdx];
+						const FReplicationTemplateIDFragment& TemplateIDFragment = TemplateIDList[EntityIdx];
 
 						ReplicatedAgent.SetNetID(NetIDFragment.NetID);
 						ReplicatedAgent.SetTemplateID(TemplateIDFragment.ID);
diff --git a/Engine/Plugins/Runtime/MassGameplay/Source/MassReplication/Public/MassReplicationTransformHandlers.h b/Engine/Plugins/Runtime/MassGameplay/Source/MassReplication/Public/MassReplicationTransformHandlers.h
--- a/Engine/Plugins/Runtime/MassGameplay/Source/MassReplication/Public/MassReplicationTransformHandlers.h
+++ b/Engine/Plugins/Runtime/MassGameplay/Source/MassReplication/Public/MassReplicationTransformHandlers.h
@@ -89,11 +89,11 @@ public:
 
 protected:
 #if UE_REPLICATION_COMPILE_CLIENT_CODE
-	static void SetEntityData(FDataFragment_Transform& TransformFragment, const FReplicatedAgentPositionYawData& ReplicatedPositionYawData);
+	static void SetEntityData(FTransformFragment& TransformFragment, const FReplicatedAgentPositionYawData& ReplicatedPositionYawData);
 #endif // UE_REPLICATION_COMPILE_CLIENT_CODE
 
 protected:
-	TArrayView<FDataFragment_Transform> TransformList;
+	TArrayView<FTransformFragment> TransformList;
 
 	TClientBubbleHandlerBase<AgentArrayItem>& OwnerHandler;
 };
@@ -144,7 +144,7 @@ void TMassClientBubbleTransformHandler<AgentArrayItem>::SetBubblePositionYawFrom
 template<typename AgentArrayItem>
 void TMassClientBubbleTransformHandler<AgentArrayItem>::AddRequirementsForSpawnQuery(FMassEntityQuery& InQuery)
 {
-	InQuery.AddRequirement<FDataFragment_Transform>(EMassFragmentAccess::ReadWrite);
+	InQuery.AddRequirement<FTransformFragment>(EMassFragmentAccess::ReadWrite);
 }
 #endif // UE_REPLICATION_COMPILE_CLIENT_CODE
 
@@ -152,7 +152,7 @@ void TMassClientBubbleTransformHandler<AgentArrayItem>::AddRequirementsForSpawnQ
 template<typename AgentArrayItem>
 void TMassClientBubbleTransformHandler<AgentArrayItem>::CacheFragmentViewsForSpawnQuery(FMassExecutionContext& InExecContext)
 {
-	TransformList = InExecContext.GetMutableFragmentView<FDataFragment_Transform>();
+	TransformList = InExecContext.GetMutableFragmentView<FTransformFragment>();
 }
 #endif // UE_REPLICATION_COMPILE_CLIENT_CODE
 
@@ -160,7 +160,7 @@ void TMassClientBubbleTransformHandler<AgentArrayItem>::CacheFragmentViewsForSpa
 template<typename AgentArrayItem>
 void TMassClientBubbleTransformHandler<AgentArrayItem>::ClearFragmentViewsForSpawnQuery()
 {
-	TransformList = TArrayView<FDataFragment_Transform>();
+	TransformList = TArrayView<FTransformFragment>();
 }
 #endif // UE_REPLICATION_COMPILE_CLIENT_CODE
 
@@ -168,7 +168,7 @@ void TMassClientBubbleTransformHandler<AgentArrayItem>::ClearFragmentViewsForSpa
 template<typename AgentArrayItem>
 void TMassClientBubbleTransformHandler<AgentArrayItem>::SetSpawnedEntityData(const int32 EntityIdx, const FReplicatedAgentPositionYawData& ReplicatedPositionYawData) const
 {
-	FDataFragment_Transform& TransformFragment = TransformList[EntityIdx];
+	FTransformFragment& TransformFragment = TransformList[EntityIdx];
 
 	SetEntityData(TransformFragment, ReplicatedPositionYawData);
 }
@@ -178,7 +178,7 @@ void TMassClientBubbleTransformHandler<AgentArrayItem>::SetSpawnedEntityData(con
 template<typename AgentArrayItem>
 void TMassClientBubbleTransformHandler<AgentArrayItem>::SetModifiedEntityData(const FMassEntityView& EntityView, const FReplicatedAgentPositionYawData& ReplicatedPositionYawData)
 {
-	FDataFragment_Transform& TransformFragment = EntityView.GetFragmentData<FDataFragment_Transform>();
+	FTransformFragment& TransformFragment = EntityView.GetFragmentData<FTransformFragment>();
 
 	SetEntityData(TransformFragment, ReplicatedPositionYawData);
 }
@@ -186,7 +186,7 @@ void TMassClientBubbleTransformHandler<AgentArrayItem>::SetModifiedEntityData(co
 
 #if UE_REPLICATION_COMPILE_CLIENT_CODE
 template<typename AgentArrayItem>
-void TMassClientBubbleTransformHandler<AgentArrayItem>::SetEntityData(FDataFragment_Transform& TransformFragment, const FReplicatedAgentPositionYawData& ReplicatedPositionYawData)
+void TMassClientBubbleTransformHandler<AgentArrayItem>::SetEntityData(FTransformFragment& TransformFragment, const FReplicatedAgentPositionYawData& ReplicatedPositionYawData)
 {
 	TransformFragment.GetMutableTransform().SetLocation(ReplicatedPositionYawData.GetPosition());
 	TransformFragment.GetMutableTransform().SetRotation(FQuat(FVector::UpVector, ReplicatedPositionYawData.GetYaw()));
@@ -202,7 +202,7 @@ public:
 	void CacheFragmentViews(FMassExecutionContext& ExecContext);
 
 protected:
-	TArrayView<FDataFragment_Transform> TransformList;
+	TArrayView<FTransformFragment> TransformList;
 };
 
 //////////////////////////////////////////////////////////////////////////// FMassReplicationProcessorPositionYawHandler ////////////////////////////////////////////////////////////////////////////
@@ -222,7 +222,7 @@ public:
 template<typename AgentArrayItem>
 void FMassReplicationProcessorPositionYawHandler::ModifyEntity(const FMassReplicatedAgentHandle Handle, const int32 EntityIdx, TMassClientBubbleTransformHandler<AgentArrayItem>& BubbleTransformHandler)
 {
-	const FDataFragment_Transform& TransformFragment = TransformList[EntityIdx];
+	const FTransformFragment& TransformFragment = TransformList[EntityIdx];
 
 	BubbleTransformHandler.SetBubblePositionYawFromTransform(Handle, TransformFragment.GetTransform());
 }
\ No newline at end of file
diff --git a/Engine/Plugins/Runtime/MassGameplay/Source/MassRepresentation/Private/MassRepresentationActorManagement.cpp b/Engine/Plugins/Runtime/MassGameplay/Source/MassRepresentation/Private/MassRepresentationActorManagement.cpp
--- a/Engine/Plugins/Runtime/MassGameplay/Source/MassRepresentation/Private/MassRepresentationActorManagement.cpp
+++ b/Engine/Plugins/Runtime/MassGameplay/Source/MassRepresentation/Private/MassRepresentationActorManagement.cpp
@@ -13,7 +13,7 @@ float UMassRepresentationActorManagement::GetSpawnPriority(const FMassRepresenta
 	return Representation.LODSignificance - (Representation.Visibility == EMassVisibility::CanBeSeen ? 1.0f : 0.0f);
 }
 
-AActor* UMassRepresentationActorManagement::GetOrSpawnActor(UMassRepresentationSubsystem& RepresentationSubsystem, UMassEntitySubsystem& EntitySubsystem, const FMassEntityHandle MassAgent, FDataFragment_Actor& ActorInfo, const FTransform& Transform, const int16 TemplateActorIndex, FMassActorSpawnRequestHandle& SpawnRequestHandle, const float Priority) const
+AActor* UMassRepresentationActorManagement::GetOrSpawnActor(UMassRepresentationSubsystem& RepresentationSubsystem, UMassEntitySubsystem& EntitySubsystem, const FMassEntityHandle MassAgent, FMassActorFragment& ActorInfo, const FTransform& Transform, const int16 TemplateActorIndex, FMassActorSpawnRequestHandle& SpawnRequestHandle, const float Priority) const
 {
 	return RepresentationSubsystem.GetOrSpawnActorFromTemplate(MassAgent, Transform, TemplateActorIndex, SpawnRequestHandle, Priority,
 		FMassActorPreSpawnDelegate::CreateUObject(this, &UMassRepresentationActorManagement::OnPreActorSpawn, &EntitySubsystem),
@@ -55,9 +55,9 @@ void UMassRepresentationActorManagement::OnPreActorSpawn(const FMassActorSpawnRe
 
 	const FMassActorSpawnRequest& MassActorSpawnRequest = SpawnRequest.Get<FMassActorSpawnRequest>();
 	const FMassEntityView EntityView(*EntitySubsystem, MassActorSpawnRequest.MassAgent);
-	FDataFragment_Actor& ActorInfo = EntityView.GetFragmentData<FDataFragment_Actor>();
+	FMassActorFragment& ActorInfo = EntityView.GetFragmentData<FMassActorFragment>();
 	FMassRepresentationFragment& Representation = EntityView.GetFragmentData<FMassRepresentationFragment>();
-	UMassRepresentationSubsystem* RepresentationSubsystem = EntityView.GetSharedFragmentData<FMassRepresentationSubsystemFragment>().RepresentationSubsystem;
+	UMassRepresentationSubsystem* RepresentationSubsystem = EntityView.GetSharedFragmentData<FMassRepresentationSubsystemSharedFragment>().RepresentationSubsystem;
 	check(RepresentationSubsystem);
 
 	// Release any existing actor
@@ -88,7 +88,7 @@ EMassActorSpawnRequestAction UMassRepresentationActorManagement::OnPostActorSpaw
 	checkf(MassActorSpawnRequest.SpawnedActor, TEXT("Expecting valid spawned actor"));
 
 	// Might be already done if the actor has a MassAgentComponent via the callback OnMassAgentComponentEntityAssociated on the MassRepresentationSubsystem
-	FDataFragment_Actor& ActorInfo = EntitySubsystem->GetFragmentDataChecked<FDataFragment_Actor>(MassActorSpawnRequest.MassAgent);
+	FMassActorFragment& ActorInfo = EntitySubsystem->GetFragmentDataChecked<FMassActorFragment>(MassActorSpawnRequest.MassAgent);
 	if (ActorInfo.IsValid())
 	{
 		// If already set, make sure it is pointing to the same actor.
@@ -105,14 +105,14 @@ EMassActorSpawnRequestAction UMassRepresentationActorManagement::OnPostActorSpaw
 void UMassRepresentationActorManagement::ReleaseAnyActorOrCancelAnySpawning(UMassEntitySubsystem& EntitySubsystem, const FMassEntityHandle MassAgent)
 {
 	FMassEntityView EntityView(EntitySubsystem, MassAgent);
-	FDataFragment_Actor& ActorInfo = EntityView.GetFragmentData<FDataFragment_Actor>();
+	FMassActorFragment& ActorInfo = EntityView.GetFragmentData<FMassActorFragment>();
 	FMassRepresentationFragment& Representation = EntityView.GetFragmentData<FMassRepresentationFragment>();
-	UMassRepresentationSubsystem* RepresentationSubsystem = EntityView.GetSharedFragmentData<FMassRepresentationSubsystemFragment>().RepresentationSubsystem;
+	UMassRepresentationSubsystem* RepresentationSubsystem = EntityView.GetSharedFragmentData<FMassRepresentationSubsystemSharedFragment>().RepresentationSubsystem;
 	check(RepresentationSubsystem);
 	ReleaseAnyActorOrCancelAnySpawning(*RepresentationSubsystem, MassAgent, ActorInfo, Representation);
 }
 
-void UMassRepresentationActorManagement::ReleaseAnyActorOrCancelAnySpawning(UMassRepresentationSubsystem& RepresentationSubsystem, const FMassEntityHandle MassAgent, FDataFragment_Actor& ActorInfo, FMassRepresentationFragment& Representation)
+void UMassRepresentationActorManagement::ReleaseAnyActorOrCancelAnySpawning(UMassRepresentationSubsystem& RepresentationSubsystem, const FMassEntityHandle MassAgent, FMassActorFragment& ActorInfo, FMassRepresentationFragment& Representation)
 {
 	// This method can only release owned by mass actors
 	AActor* Actor = ActorInfo.GetOwnedByMassMutable();
diff --git a/Engine/Plugins/Runtime/MassGameplay/Source/MassRepresentation/Private/MassRepresentationProcessor.cpp b/Engine/Plugins/Runtime/MassGameplay/Source/MassRepresentation/Private/MassRepresentationProcessor.cpp
--- a/Engine/Plugins/Runtime/MassGameplay/Source/MassRepresentation/Private/MassRepresentationProcessor.cpp
+++ b/Engine/Plugins/Runtime/MassGameplay/Source/MassRepresentation/Private/MassRepresentationProcessor.cpp
@@ -29,12 +29,12 @@ UMassRepresentationProcessor::UMassRepresentationProcessor()
 
 void UMassRepresentationProcessor::ConfigureQueries()
 {
-	EntityQuery.AddRequirement<FDataFragment_Transform>(EMassFragmentAccess::ReadOnly);
+	EntityQuery.AddRequirement<FTransformFragment>(EMassFragmentAccess::ReadOnly);
 	EntityQuery.AddRequirement<FMassRepresentationFragment>(EMassFragmentAccess::ReadWrite);
 	EntityQuery.AddRequirement<FMassRepresentationLODFragment>(EMassFragmentAccess::ReadOnly);
-	EntityQuery.AddRequirement<FDataFragment_Actor>(EMassFragmentAccess::ReadWrite);
-	EntityQuery.AddConstSharedRequirement<FMassRepresentationConfig>();
-	EntityQuery.AddSharedRequirement<FMassRepresentationSubsystemFragment>(EMassFragmentAccess::ReadWrite);
+	EntityQuery.AddRequirement<FMassActorFragment>(EMassFragmentAccess::ReadWrite);
+	EntityQuery.AddConstSharedRequirement<FMassRepresentationParameters>();
+	EntityQuery.AddSharedRequirement<FMassRepresentationSubsystemSharedFragment>(EMassFragmentAccess::ReadWrite);
 }
 
 void UMassRepresentationProcessor::Initialize(UObject& Owner)
@@ -47,16 +47,16 @@ void UMassRepresentationProcessor::Initialize(UObject& Owner)
 
 void UMassRepresentationProcessor::UpdateRepresentation(FMassExecutionContext& Context)
 {
-	UMassRepresentationSubsystem* RepresentationSubsystem = Context.GetMutableSharedFragment<FMassRepresentationSubsystemFragment>().RepresentationSubsystem;
+	UMassRepresentationSubsystem* RepresentationSubsystem = Context.GetMutableSharedFragment<FMassRepresentationSubsystemSharedFragment>().RepresentationSubsystem;
 	check(RepresentationSubsystem);
-	const FMassRepresentationConfig& RepresentationConfig = Context.GetConstSharedFragment<FMassRepresentationConfig>();
+	const FMassRepresentationParameters& RepresentationConfig = Context.GetConstSharedFragment<FMassRepresentationParameters>();
 	UMassRepresentationActorManagement* RepresentationActorManagement = RepresentationConfig.CachedRepresentationActorManagement;
 	check(RepresentationActorManagement);
 
-	const TConstArrayView<FDataFragment_Transform> TransformList = Context.GetFragmentView<FDataFragment_Transform>();
+	const TConstArrayView<FTransformFragment> TransformList = Context.GetFragmentView<FTransformFragment>();
 	const TArrayView<FMassRepresentationFragment> RepresentationList = Context.GetMutableFragmentView<FMassRepresentationFragment>();
 	const TConstArrayView<FMassRepresentationLODFragment> RepresentationLODList = Context.GetFragmentView<FMassRepresentationLODFragment>();
-	const TArrayView<FDataFragment_Actor> ActorList = Context.GetMutableFragmentView<FDataFragment_Actor>();
+	const TArrayView<FMassActorFragment> ActorList = Context.GetMutableFragmentView<FMassActorFragment>();
 
 	const bool bDoKeepActorExtraFrame = UE::MassRepresentation::bAllowKeepActorExtraFrame ? RepresentationConfig.bKeepLowResActors : false;
 
@@ -64,10 +64,10 @@ void UMassRepresentationProcessor::UpdateRepresentation(FMassExecutionContext& C
 	for (int32 EntityIdx = 0; EntityIdx < NumEntities; EntityIdx++)
 	{
 		const FMassEntityHandle MassAgent = Context.GetEntity(EntityIdx);
-		const FDataFragment_Transform& TransformFragment = TransformList[EntityIdx];
+		const FTransformFragment& TransformFragment = TransformList[EntityIdx];
 		const FMassRepresentationLODFragment& RepresentationLOD = RepresentationLODList[EntityIdx];
 		FMassRepresentationFragment& Representation = RepresentationList[EntityIdx];
-		FDataFragment_Actor& ActorInfo = ActorList[EntityIdx];
+		FMassActorFragment& ActorInfo = ActorList[EntityIdx];
 
 		// Keeping a copy of the that last calculated previous representation
 		const EMassRepresentationType PrevRepresentationCopy = Representation.PrevRepresentation;
@@ -227,7 +227,7 @@ void UMassRepresentationProcessor::Execute(UMassEntitySubsystem& InEntitySubsyst
 	});
 }
 
-bool UMassRepresentationProcessor::ReleaseActorOrCancelSpawning(UMassRepresentationSubsystem& RepresentationSubsystem, const FMassEntityHandle MassAgent, FDataFragment_Actor& ActorInfo, const int16 TemplateActorIndex, FMassActorSpawnRequestHandle& SpawnRequestHandle, FMassCommandBuffer& CommandBuffer, bool bCancelSpawningOnly /*= false*/)
+bool UMassRepresentationProcessor::ReleaseActorOrCancelSpawning(UMassRepresentationSubsystem& RepresentationSubsystem, const FMassEntityHandle MassAgent, FMassActorFragment& ActorInfo, const int16 TemplateActorIndex, FMassActorSpawnRequestHandle& SpawnRequestHandle, FMassCommandBuffer& CommandBuffer, bool bCancelSpawningOnly /*= false*/)
 {
 	if (TemplateActorIndex == INDEX_NONE)
 	{
@@ -297,7 +297,7 @@ void UMassVisualizationProcessor::ConfigureQueries()
 
 FMassVisualizationChunkFragment& UMassVisualizationProcessor::UpdateChunkVisibility(FMassExecutionContext& Context) const
 {
-	const FMassRepresentationConfig& RepresentationConfig = Context.GetConstSharedFragment<FMassRepresentationConfig>();
+	const FMassRepresentationParameters& RepresentationConfig = Context.GetConstSharedFragment<FMassRepresentationParameters>();
 	bool bFirstUpdate = false;
 
 	// Setup chunk fragment data about visibility
@@ -373,26 +373,26 @@ UMassRepresentationFragmentDestructor::UMassRepresentationFragmentDestructor()
 void UMassRepresentationFragmentDestructor::ConfigureQueries()
 {
 	EntityQuery.AddRequirement<FMassRepresentationFragment>(EMassFragmentAccess::ReadOnly);
-	EntityQuery.AddRequirement<FDataFragment_Actor>(EMassFragmentAccess::ReadWrite);
-	EntityQuery.AddConstSharedRequirement<FMassRepresentationConfig>();
-	EntityQuery.AddSharedRequirement<FMassRepresentationSubsystemFragment>(EMassFragmentAccess::ReadWrite);
+	EntityQuery.AddRequirement<FMassActorFragment>(EMassFragmentAccess::ReadWrite);
+	EntityQuery.AddConstSharedRequirement<FMassRepresentationParameters>();
+	EntityQuery.AddSharedRequirement<FMassRepresentationSubsystemSharedFragment>(EMassFragmentAccess::ReadWrite);
 }
 
 void UMassRepresentationFragmentDestructor::Execute(UMassEntitySubsystem& EntitySubsystem, FMassExecutionContext& Context)
 {
 	EntityQuery.ForEachEntityChunk(EntitySubsystem, Context, [this](FMassExecutionContext& Context)
 	{
-		UMassRepresentationSubsystem* RepresentationSubsystem = Context.GetMutableSharedFragment<FMassRepresentationSubsystemFragment>().RepresentationSubsystem;
+		UMassRepresentationSubsystem* RepresentationSubsystem = Context.GetMutableSharedFragment<FMassRepresentationSubsystemSharedFragment>().RepresentationSubsystem;
 		check(RepresentationSubsystem);
 
 		const TArrayView<FMassRepresentationFragment> RepresentationList = Context.GetMutableFragmentView<FMassRepresentationFragment>();
-		const TArrayView<FDataFragment_Actor> ActorList = Context.GetMutableFragmentView<FDataFragment_Actor>();
+		const TArrayView<FMassActorFragment> ActorList = Context.GetMutableFragmentView<FMassActorFragment>();
 
 		const int32 NumEntities = Context.GetNumEntities();
 		for (int32 i = 0; i < NumEntities; ++i)
 		{
 			FMassRepresentationFragment& Representation = RepresentationList[i];
-			FDataFragment_Actor& ActorInfo = ActorList[i];
+			FMassActorFragment& ActorInfo = ActorList[i];
 
 			const FMassEntityHandle MassAgent = Context.GetEntity(i);
 
diff --git a/Engine/Plugins/Runtime/MassGameplay/Source/MassRepresentation/Private/MassUpdateISMProcessor.cpp b/Engine/Plugins/Runtime/MassGameplay/Source/MassRepresentation/Private/MassUpdateISMProcessor.cpp
--- a/Engine/Plugins/Runtime/MassGameplay/Source/MassRepresentation/Private/MassUpdateISMProcessor.cpp
+++ b/Engine/Plugins/Runtime/MassGameplay/Source/MassRepresentation/Private/MassUpdateISMProcessor.cpp
@@ -18,30 +18,30 @@ UMassUpdateISMProcessor::UMassUpdateISMProcessor()
 
 void UMassUpdateISMProcessor::ConfigureQueries()
 {
-	EntityQuery.AddRequirement<FDataFragment_Transform>(EMassFragmentAccess::ReadOnly);
+	EntityQuery.AddRequirement<FTransformFragment>(EMassFragmentAccess::ReadOnly);
 	EntityQuery.AddRequirement<FMassRepresentationFragment>(EMassFragmentAccess::ReadWrite);
 	EntityQuery.AddRequirement<FMassRepresentationLODFragment>(EMassFragmentAccess::ReadWrite);
 	EntityQuery.AddChunkRequirement<FMassVisualizationChunkFragment>(EMassFragmentAccess::ReadWrite);
 	EntityQuery.SetChunkFilter(&FMassVisualizationChunkFragment::AreAnyEntitiesVisibleInChunk);
-	EntityQuery.AddSharedRequirement<FMassRepresentationSubsystemFragment>(EMassFragmentAccess::ReadWrite);
+	EntityQuery.AddSharedRequirement<FMassRepresentationSubsystemSharedFragment>(EMassFragmentAccess::ReadWrite);
 }
 
 void UMassUpdateISMProcessor::Execute(UMassEntitySubsystem& EntitySubsystem, FMassExecutionContext& Context)
 {
 	EntityQuery.ForEachEntityChunk(EntitySubsystem, Context, [](FMassExecutionContext& Context)
 	{
-		UMassRepresentationSubsystem* RepresentationSubsystem = Context.GetSharedFragment<FMassRepresentationSubsystemFragment>().RepresentationSubsystem;
+		UMassRepresentationSubsystem* RepresentationSubsystem = Context.GetSharedFragment<FMassRepresentationSubsystemSharedFragment>().RepresentationSubsystem;
 		check(RepresentationSubsystem);
 		FMassInstancedStaticMeshInfoArrayView ISMInfo = RepresentationSubsystem->GetMutableInstancedStaticMeshInfos();
 
-		const TConstArrayView<FDataFragment_Transform> TransformList = Context.GetFragmentView<FDataFragment_Transform>();
+		const TConstArrayView<FTransformFragment> TransformList = Context.GetFragmentView<FTransformFragment>();
 		const TArrayView<FMassRepresentationFragment> RepresentationList = Context.GetMutableFragmentView<FMassRepresentationFragment>();
 		const TConstArrayView<FMassRepresentationLODFragment> RepresentationLODList = Context.GetFragmentView<FMassRepresentationLODFragment>();
 
 		const int32 NumEntities = Context.GetNumEntities();
 		for (int32 EntityIdx = 0; EntityIdx < NumEntities; EntityIdx++)
 		{
-			const FDataFragment_Transform& TransformFragment = TransformList[EntityIdx];
+			const FTransformFragment& TransformFragment = TransformList[EntityIdx];
 			const FMassRepresentationLODFragment& RepresentationLOD = RepresentationLODList[EntityIdx];
 			FMassRepresentationFragment& Representation = RepresentationList[EntityIdx];
 
diff --git a/Engine/Plugins/Runtime/MassGameplay/Source/MassRepresentation/Private/MassVisualizationLODProcessor.cpp b/Engine/Plugins/Runtime/MassGameplay/Source/MassRepresentation/Private/MassVisualizationLODProcessor.cpp
--- a/Engine/Plugins/Runtime/MassGameplay/Source/MassRepresentation/Private/MassVisualizationLODProcessor.cpp
+++ b/Engine/Plugins/Runtime/MassGameplay/Source/MassRepresentation/Private/MassVisualizationLODProcessor.cpp
@@ -35,7 +35,7 @@ void UMassVisualizationLODProcessor::ConfigureQueries()
 {
 	CloseEntityQuery.AddRequirement<FMassViewerInfoFragment>(EMassFragmentAccess::ReadOnly);
 	CloseEntityQuery.AddRequirement<FMassRepresentationLODFragment>(EMassFragmentAccess::ReadWrite);
-	CloseEntityQuery.AddRequirement<FDataFragment_Transform>(EMassFragmentAccess::ReadOnly);
+	CloseEntityQuery.AddRequirement<FTransformFragment>(EMassFragmentAccess::ReadOnly);
 
 	FarEntityQuery = CloseEntityQuery;
 
diff --git a/Engine/Plugins/Runtime/MassGameplay/Source/MassRepresentation/Private/MassVisualizationTrait.cpp b/Engine/Plugins/Runtime/MassGameplay/Source/MassRepresentation/Private/MassVisualizationTrait.cpp
--- a/Engine/Plugins/Runtime/MassGameplay/Source/MassRepresentation/Private/MassVisualizationTrait.cpp
+++ b/Engine/Plugins/Runtime/MassGameplay/Source/MassRepresentation/Private/MassVisualizationTrait.cpp
@@ -32,7 +32,7 @@ void UMassVisualizationTrait::BuildTemplate(FMassEntityTemplateBuildContext& Bui
 	// adding Config.AdditionalDataFragments to let user configure the fragments first. Calling BuildContext.Add() 
 	// won't override any fragments that are already there
 	BuildContext.AddTag<FMassCollectLODViewerInfoTag>(); // Depends on FMassViewerInfoFragment
-	BuildContext.AddFragment<FDataFragment_Transform>();
+	BuildContext.AddFragment<FTransformFragment>();
 
 	UMassEntitySubsystem* EntitySubsystem = UWorld::GetSubsystem<UMassEntitySubsystem>(&World);
 	check(EntitySubsystem);
@@ -45,10 +45,10 @@ void UMassVisualizationTrait::BuildTemplate(FMassEntityTemplateBuildContext& Bui
 		check(RepresentationSubsystem);
 	}
 
-	FMassRepresentationSubsystemFragment Subsystem;
+	FMassRepresentationSubsystemSharedFragment Subsystem;
 	Subsystem.RepresentationSubsystem = RepresentationSubsystem;
 	uint32 SubsystemHash = UE::StructUtils::GetStructCrc32(FConstStructView::Make(Subsystem));
-	FSharedStruct SubsystemFragment = EntitySubsystem->GetOrCreateSharedFragment<FMassRepresentationSubsystemFragment>(SubsystemHash, Subsystem);
+	FSharedStruct SubsystemFragment = EntitySubsystem->GetOrCreateSharedFragment<FMassRepresentationSubsystemSharedFragment>(SubsystemHash, Subsystem);
 	BuildContext.AddSharedFragment(SubsystemFragment);
 
 	if (!Config.RepresentationActorManagementClass)
@@ -56,8 +56,8 @@ void UMassVisualizationTrait::BuildTemplate(FMassEntityTemplateBuildContext& Bui
 		UE_LOG(LogMassRepresentation, Error, TEXT("Expecting a valid class for the representation actor management"));
 	}
 	uint32 ConfigHash = UE::StructUtils::GetStructCrc32(FConstStructView::Make(Config));
-	FConstSharedStruct ConfigFragment = EntitySubsystem->GetOrCreateConstSharedFragment<FMassRepresentationConfig>(ConfigHash, Config);
-	ConfigFragment.Get<FMassRepresentationConfig>().ComputeCachedValues();
+	FConstSharedStruct ConfigFragment = EntitySubsystem->GetOrCreateConstSharedFragment<FMassRepresentationParameters>(ConfigHash, Config);
+	ConfigFragment.Get<FMassRepresentationParameters>().ComputeCachedValues();
 	BuildContext.AddConstSharedFragment(ConfigFragment);
 
 	FMassRepresentationFragment& RepresentationFragment = BuildContext.AddFragment_GetRef<FMassRepresentationFragment>();
diff --git a/Engine/Plugins/Runtime/MassGameplay/Source/MassRepresentation/Public/MassRepresentationActorManagement.h b/Engine/Plugins/Runtime/MassGameplay/Source/MassRepresentation/Public/MassRepresentationActorManagement.h
--- a/Engine/Plugins/Runtime/MassGameplay/Source/MassRepresentation/Public/MassRepresentationActorManagement.h
+++ b/Engine/Plugins/Runtime/MassGameplay/Source/MassRepresentation/Public/MassRepresentationActorManagement.h
@@ -12,7 +12,7 @@ enum class EMassActorEnabledType : uint8;
 struct FMassActorSpawnRequestHandle;
 struct FConstStructView;
 struct FMassEntityHandle;
-struct FDataFragment_Actor;
+struct FMassActorFragment;
 struct FMassRepresentationLODFragment;
 struct FMassRepresentationFragment;
 class UMassRepresentationSubsystem;
@@ -42,7 +42,7 @@ public:
 	 * @param Priority of this spawn request in comparison with the others, lower value means higher priority
 	 * @return the actor spawned
 	 */
-	virtual AActor* GetOrSpawnActor(UMassRepresentationSubsystem& RepresentationSubsystem, UMassEntitySubsystem& EntitySubsystem, const FMassEntityHandle MassAgent, FDataFragment_Actor& ActorInfo, const FTransform& Transform, const int16 TemplateActorIndex, FMassActorSpawnRequestHandle& SpawnRequestHandle, const float Priority) const;
+	virtual AActor* GetOrSpawnActor(UMassRepresentationSubsystem& RepresentationSubsystem, UMassEntitySubsystem& EntitySubsystem, const FMassEntityHandle MassAgent, FMassActorFragment& ActorInfo, const FTransform& Transform, const int16 TemplateActorIndex, FMassActorSpawnRequestHandle& SpawnRequestHandle, const float Priority) const;
 
 	/**
 	 * Enable/disable a spawned actor for a mass entity
@@ -99,5 +99,5 @@ public:
 	 * @param ActorInfo is the fragment where we are going to store the actor pointer
 	 * @param Representation fragment containing the current and previous visual state
 	 */
-	static void ReleaseAnyActorOrCancelAnySpawning(UMassRepresentationSubsystem& RepresentationSubsystem, const FMassEntityHandle MassAgent, FDataFragment_Actor& ActorInfo, FMassRepresentationFragment& Representation);
+	static void ReleaseAnyActorOrCancelAnySpawning(UMassRepresentationSubsystem& RepresentationSubsystem, const FMassEntityHandle MassAgent, FMassActorFragment& ActorInfo, FMassRepresentationFragment& Representation);
 };
\ No newline at end of file
diff --git a/Engine/Plugins/Runtime/MassGameplay/Source/MassRepresentation/Public/MassRepresentationFragments.h b/Engine/Plugins/Runtime/MassGameplay/Source/MassRepresentation/Public/MassRepresentationFragments.h
--- a/Engine/Plugins/Runtime/MassGameplay/Source/MassRepresentation/Public/MassRepresentationFragments.h
+++ b/Engine/Plugins/Runtime/MassGameplay/Source/MassRepresentation/Public/MassRepresentationFragments.h
@@ -55,7 +55,7 @@ struct MASSREPRESENTATION_API FMassRepresentationFragment : public FMassFragment
 };
 
 USTRUCT()
-struct FMassRepresentationSubsystemFragment : public FMassSharedFragment
+struct FMassRepresentationSubsystemSharedFragment : public FMassSharedFragment
 {
 	GENERATED_BODY()
 
@@ -64,7 +64,7 @@ struct FMassRepresentationSubsystemFragment : public FMassSharedFragment
 };
 
 USTRUCT()
-struct FMassRepresentationConfig : public FMassSharedFragment
+struct FMassRepresentationParameters : public FMassSharedFragment
 {
 	GENERATED_BODY()
 
@@ -107,7 +107,7 @@ struct FMassRepresentationConfig : public FMassSharedFragment
 	mutable UMassRepresentationActorManagement* CachedRepresentationActorManagement = nullptr;
 };
 
-inline void FMassRepresentationConfig::ComputeCachedValues() const
+inline void FMassRepresentationParameters::ComputeCachedValues() const
 {
 	// Calculate the default representation when actor isn't spawned yet.
 	for (int32 LOD = EMassLOD::High; LOD < EMassLOD::Max; LOD++)
diff --git a/Engine/Plugins/Runtime/MassGameplay/Source/MassRepresentation/Public/MassRepresentationProcessor.h b/Engine/Plugins/Runtime/MassGameplay/Source/MassRepresentation/Public/MassRepresentationProcessor.h
--- a/Engine/Plugins/Runtime/MassGameplay/Source/MassRepresentation/Public/MassRepresentationProcessor.h
+++ b/Engine/Plugins/Runtime/MassGameplay/Source/MassRepresentation/Public/MassRepresentationProcessor.h
@@ -10,7 +10,7 @@
 #include "MassRepresentationProcessor.generated.h"
 
 class UMassRepresentationSubsystem;
-struct FDataFragment_Actor;
+struct FMassActorFragment;
 
 UCLASS()
 class MASSREPRESENTATION_API UMassRepresentationProcessor : public UMassProcessor
@@ -49,7 +49,7 @@ protected:
 	 * @param bCancelSpawningOnly tell to only cancel the existing spawning request and to not release the associated actor it any.
 	 * @return if the actor was release or the spawning was canceled.
 	 */
-	bool ReleaseActorOrCancelSpawning(UMassRepresentationSubsystem& RepresentationSubsystem, const FMassEntityHandle MassAgent, FDataFragment_Actor& ActorInfo, const int16 TemplateActorIndex, FMassActorSpawnRequestHandle& SpawnRequestHandle, FMassCommandBuffer& CommandBuffer, const bool bCancelSpawningOnly = false);
+	bool ReleaseActorOrCancelSpawning(UMassRepresentationSubsystem& RepresentationSubsystem, const FMassEntityHandle MassAgent, FMassActorFragment& ActorInfo, const int16 TemplateActorIndex, FMassActorSpawnRequestHandle& SpawnRequestHandle, FMassCommandBuffer& CommandBuffer, const bool bCancelSpawningOnly = false);
 
 	/*
 	 * Update representation type for each entity, must be called within a ForEachEntityChunk
diff --git a/Engine/Plugins/Runtime/MassGameplay/Source/MassRepresentation/Public/MassVisualizationLODProcessor.h b/Engine/Plugins/Runtime/MassGameplay/Source/MassRepresentation/Public/MassVisualizationLODProcessor.h
--- a/Engine/Plugins/Runtime/MassGameplay/Source/MassRepresentation/Public/MassVisualizationLODProcessor.h
+++ b/Engine/Plugins/Runtime/MassGameplay/Source/MassRepresentation/Public/MassVisualizationLODProcessor.h
@@ -143,7 +143,7 @@ void UMassVisualizationLODProcessor::ExecuteInternal(UMassEntitySubsystem& Entit
 		auto DebugDisplayLOD = [this](FMassExecutionContext& Context)
 		{
 			TConstArrayView<FMassRepresentationLODFragment> RepresentationLODList = Context.GetFragmentView<FMassRepresentationLODFragment>();
-			TConstArrayView<FDataFragment_Transform> TransformList = Context.GetFragmentView<FDataFragment_Transform>();
+			TConstArrayView<FTransformFragment> TransformList = Context.GetFragmentView<FTransformFragment>();
 			LODCalculator.DebugDisplayLOD(Context, RepresentationLODList, TransformList, World);
 		};
 
diff --git a/Engine/Plugins/Runtime/MassGameplay/Source/MassRepresentation/Public/MassVisualizationTrait.h b/Engine/Plugins/Runtime/MassGameplay/Source/MassRepresentation/Public/MassVisualizationTrait.h
--- a/Engine/Plugins/Runtime/MassGameplay/Source/MassRepresentation/Public/MassVisualizationTrait.h
+++ b/Engine/Plugins/Runtime/MassGameplay/Source/MassRepresentation/Public/MassVisualizationTrait.h
@@ -40,5 +40,5 @@ public:
 
 	/** Configuration parameters for the representaion processor */
 	UPROPERTY(EditAnywhere, Category = "Mass|Visual")
-	FMassRepresentationConfig Config;
+	FMassRepresentationParameters Config;
 };
diff --git a/Engine/Plugins/Runtime/MassGameplay/Source/MassSmartObjects/Private/MassSmartObjectHandler.cpp b/Engine/Plugins/Runtime/MassGameplay/Source/MassSmartObjects/Private/MassSmartObjectHandler.cpp
--- a/Engine/Plugins/Runtime/MassGameplay/Source/MassSmartObjects/Private/MassSmartObjectHandler.cpp
+++ b/Engine/Plugins/Runtime/MassGameplay/Source/MassSmartObjects/Private/MassSmartObjectHandler.cpp
@@ -221,7 +221,7 @@ bool FMassSmartObjectHandler::ClaimSmartObject(const FMassEntityHandle Entity, F
 bool FMassSmartObjectHandler::UseSmartObject(
 	const FMassEntityHandle Entity,
 	FMassSmartObjectUserFragment& User,
-	const FDataFragment_Transform& Transform) const
+	const FTransformFragment& Transform) const
 {
 #if WITH_MASSGAMEPLAY_DEBUG
 	UE_CVLOG(UE::Mass::Debug::IsDebuggingEntity(Entity),
diff --git a/Engine/Plugins/Runtime/MassGameplay/Source/MassSmartObjects/Public/MassSmartObjectBehaviorDefinition.h b/Engine/Plugins/Runtime/MassGameplay/Source/MassSmartObjects/Public/MassSmartObjectBehaviorDefinition.h
--- a/Engine/Plugins/Runtime/MassGameplay/Source/MassSmartObjects/Public/MassSmartObjectBehaviorDefinition.h
+++ b/Engine/Plugins/Runtime/MassGameplay/Source/MassSmartObjects/Public/MassSmartObjectBehaviorDefinition.h
@@ -10,7 +10,7 @@
 class UMassEntitySubsystem;
 class USmartObjectSubsystem;
 struct FMassExecutionContext;
-struct FDataFragment_Transform;
+struct FTransformFragment;
 struct FMassSmartObjectUserFragment;
 
 /**
diff --git a/Engine/Plugins/Runtime/MassGameplay/Source/MassSmartObjects/Public/MassSmartObjectHandler.h b/Engine/Plugins/Runtime/MassGameplay/Source/MassSmartObjects/Public/MassSmartObjectHandler.h
--- a/Engine/Plugins/Runtime/MassGameplay/Source/MassSmartObjects/Public/MassSmartObjectHandler.h
+++ b/Engine/Plugins/Runtime/MassGameplay/Source/MassSmartObjects/Public/MassSmartObjectHandler.h
@@ -11,7 +11,7 @@ class USmartObjectSubsystem;
 struct FMassExecutionContext;
 struct FMassEntityHandle;
 struct FMassSmartObjectUserFragment;
-struct FDataFragment_Transform;
+struct FTransformFragment;
 struct FSmartObjectClaimHandle;
 struct FSmartObjectHandle;
 struct FZoneGraphCompactLaneLocation;
@@ -110,7 +110,7 @@ struct MASSSMARTOBJECTS_API FMassSmartObjectHandler
 	 * @param Transform Fragment holding the transform of the user claiming
 	 * @return Whether the slot has been successfully claimed or not
 	 */
-	bool UseSmartObject(const FMassEntityHandle Entity, FMassSmartObjectUserFragment& User, const FDataFragment_Transform& Transform) const;
+	bool UseSmartObject(const FMassEntityHandle Entity, FMassSmartObjectUserFragment& User, const FTransformFragment& Transform) const;
 
 	/**
 	 * Releases a claimed/in-use smart object and update user fragment.
diff --git a/Engine/Plugins/Runtime/MassGameplay/Source/MassSpawner/Private/MassSpawnLocationProcessor.cpp b/Engine/Plugins/Runtime/MassGameplay/Source/MassSpawner/Private/MassSpawnLocationProcessor.cpp
--- a/Engine/Plugins/Runtime/MassGameplay/Source/MassSpawner/Private/MassSpawnLocationProcessor.cpp
+++ b/Engine/Plugins/Runtime/MassGameplay/Source/MassSpawner/Private/MassSpawnLocationProcessor.cpp
@@ -17,7 +17,7 @@ UMassSpawnLocationProcessor::UMassSpawnLocationProcessor()
 
 void UMassSpawnLocationProcessor::ConfigureQueries()
 {
-	EntityQuery.AddRequirement<FDataFragment_Transform>(EMassFragmentAccess::ReadWrite);
+	EntityQuery.AddRequirement<FTransformFragment>(EMassFragmentAccess::ReadWrite);
 }
 
 void UMassSpawnLocationProcessor::Execute(UMassEntitySubsystem& EntitySubsystem, FMassExecutionContext& Context)
@@ -68,7 +68,7 @@ void UMassSpawnLocationProcessor::Execute(UMassEntitySubsystem& EntitySubsystem,
 
 		EntityQuery.ForEachEntityChunk(EntitySubsystem, Context, [&Transforms, this](FMassExecutionContext& Context)
 			{
-				const TArrayView<FDataFragment_Transform> LocationList = Context.GetMutableFragmentView<FDataFragment_Transform>();
+				const TArrayView<FTransformFragment> LocationList = Context.GetMutableFragmentView<FTransformFragment>();
 				const int32 NumEntities = Context.GetNumEntities();
 				for (int32 i = 0; i < NumEntities; ++i)
 				{
diff --git a/Engine/Plugins/Runtime/MassGameplay/Source/MassSpawner/Public/MassSpawnerTypes.h b/Engine/Plugins/Runtime/MassGameplay/Source/MassSpawner/Public/MassSpawnerTypes.h
--- a/Engine/Plugins/Runtime/MassGameplay/Source/MassSpawner/Public/MassSpawnerTypes.h
+++ b/Engine/Plugins/Runtime/MassGameplay/Source/MassSpawner/Public/MassSpawnerTypes.h
@@ -75,7 +75,7 @@ struct FMassSpawnDataGenerator
 };
 
 USTRUCT()
-struct MASSSPAWNER_API FDataFragment_ReplicationTemplateID : public FMassFragment
+struct MASSSPAWNER_API FReplicationTemplateIDFragment : public FMassFragment
 {
 	GENERATED_BODY()
 
