Mass Crowd: Fixed check it path follow current position validation.

#rnx

diff --git a/Engine/Plugins/AI/MassAI/Source/MassZoneGraphNavigation/Private/MassZoneGraphNavigationFragments.cpp b/Engine/Plugins/AI/MassAI/Source/MassZoneGraphNavigation/Private/MassZoneGraphNavigationFragments.cpp
--- a/Engine/Plugins/AI/MassAI/Source/MassZoneGraphNavigation/Private/MassZoneGraphNavigationFragments.cpp
+++ b/Engine/Plugins/AI/MassAI/Source/MassZoneGraphNavigation/Private/MassZoneGraphNavigationFragments.cpp
@@ -155,7 +155,7 @@ void FMassZoneGraphCachedLaneFragment::CacheLaneData(const FZoneGraphStorage& Zo
 	LaneRightSpace = FMassInt16Real(AdjacentRightWidth);
 }
 
-bool FMassZoneGraphShortPathFragment::RequestPath(const FMassZoneGraphCachedLaneFragment& CachedLane, const FZoneGraphShortPathRequest& Request, const float CurrentDistanceAlongLane, const float AgentRadius)
+bool FMassZoneGraphShortPathFragment::RequestPath(const FMassZoneGraphCachedLaneFragment& CachedLane, const FZoneGraphShortPathRequest& Request, const float InCurrentDistanceAlongLane, const float AgentRadius)
 {
 	Reset();
 
@@ -164,8 +164,15 @@ bool FMassZoneGraphShortPathFragment::RequestPath(const FMassZoneGraphCachedLane
 		return false;
 	}
 
+
+	// The current distance can come from a quantized lane distance. Check against quantized bounds, but clamp it to the actual path length when calculating the path.
+	static_assert(TAreTypesEqual<decltype(FMassZoneGraphPathPoint::Distance), FMassInt16Real>::Value, "Assuming FMassZoneGraphPathPoint::Distance is quantized to 10 units.");
+	const float LaneLengthQuantized = FMath::CeilToFloat(CachedLane.LaneLength / 10.0f) * 10.0f;
+
 	static constexpr float Epsilon = 0.1f;
-	check(CurrentDistanceAlongLane >= -Epsilon && CurrentDistanceAlongLane <= (CachedLane.LaneLength + Epsilon));
+	ensureMsgf(InCurrentDistanceAlongLane >= -Epsilon && InCurrentDistanceAlongLane <= (LaneLengthQuantized + Epsilon), TEXT("Current distance %f should be within the lane bounds 0.0 - %f"), InCurrentDistanceAlongLane, LaneLengthQuantized);
+
+	const float CurrentDistanceAlongLane = FMath::Min(InCurrentDistanceAlongLane, CachedLane.LaneLength);
 	
 	// Set common lane parameters
 #if WITH_MASSGAMEPLAY_DEBUG
