Integration of cl 19122140

Fix replication when more than one client is connected

#rnx

diff --git a/Engine/Plugins/AI/MassAI/Source/MassAIReplication/Public/MassReplicationPathHandlers.h b/Engine/Plugins/AI/MassAI/Source/MassAIReplication/Public/MassReplicationPathHandlers.h
--- a/Engine/Plugins/AI/MassAI/Source/MassAIReplication/Public/MassReplicationPathHandlers.h
+++ b/Engine/Plugins/AI/MassAI/Source/MassAIReplication/Public/MassReplicationPathHandlers.h
@@ -214,9 +214,10 @@ public:
 	 * @param Handle to the agent in the TMassClientBubbleHandler (that TMassClientBubblePathHandler is a member variable of).
 	 * @param EntityIdx the index of the entity in fragment views that have been cached.
 	 * @param BubblePathHandler handler to actually set the data in the client bubble
+	 * @param bLastClient means it safe to reset any dirtiness
 	 */
 	template<typename AgentArrayItem>
-	void ModifyEntity(const FMassReplicatedAgentHandle Handle, const int32 EntityIdx, TMassClientBubblePathHandler<AgentArrayItem>& BubblePathHandler);
+	void ModifyEntity(const FMassReplicatedAgentHandle Handle, const int32 EntityIdx, TMassClientBubblePathHandler<AgentArrayItem>& BubblePathHandler, bool bLastClient);
 
 	TArrayView<FMassZoneGraphPathRequestFragment> PathRequestList;
 	TArrayView<FMassMoveTargetFragment> MoveTargetList;
@@ -224,7 +225,7 @@ public:
 };
 
 template<typename AgentArrayItem>
-void FMassReplicationProcessorPathHandler::ModifyEntity(const FMassReplicatedAgentHandle Handle, const int32 EntityIdx, TMassClientBubblePathHandler<AgentArrayItem>& BubblePathHandler)
+void FMassReplicationProcessorPathHandler::ModifyEntity(const FMassReplicatedAgentHandle Handle, const int32 EntityIdx, TMassClientBubblePathHandler<AgentArrayItem>& BubblePathHandler, bool bLastClient)
 {
 	const FMassZoneGraphPathRequestFragment& PathRequest = PathRequestList[EntityIdx];
 	FMassMoveTargetFragment& MoveTargetFragment = MoveTargetList[EntityIdx];
@@ -233,6 +234,9 @@ void FMassReplicationProcessorPathHandler::ModifyEntity(const FMassReplicatedAge
 	if (MoveTargetFragment.GetNetDirty())
 	{
 		BubblePathHandler.SetBubblePathData(Handle, PathRequest, MoveTargetFragment, LaneLocationFragment);
-		MoveTargetFragment.ResetNetDirty();
+		if (bLastClient)
+		{
+			MoveTargetFragment.ResetNetDirty();
+		}
 	}
 }
diff --git a/Engine/Plugins/AI/MassCrowd/Source/MassCrowd/Private/MassCrowdReplicator.cpp b/Engine/Plugins/AI/MassCrowd/Source/MassCrowd/Private/MassCrowdReplicator.cpp
--- a/Engine/Plugins/AI/MassCrowd/Source/MassCrowd/Private/MassCrowdReplicator.cpp
+++ b/Engine/Plugins/AI/MassCrowd/Source/MassCrowd/Private/MassCrowdReplicator.cpp
@@ -51,7 +51,8 @@ void UMassCrowdReplicator::ProcessClientReplication(FMassExecutionContext& Conte
 
 		FMassCrowdClientBubbleHandler& Bubble = CrowdBubbleInfo.GetCrowdSerializer().Bubble;
 
-		PathHandler.ModifyEntity<FCrowdFastArrayItem>(Handle, EntityIdx, Bubble.GetPathHandlerMutable());
+		const bool bLastClient = RepSharedFrag->CachedClientHandles.Last() == ClientHandle;
+		PathHandler.ModifyEntity<FCrowdFastArrayItem>(Handle, EntityIdx, Bubble.GetPathHandlerMutable(), bLastClient);
 
 		// Don't call the PositionYawHandler here as we currently only replicate the position and yaw when we add an entity to Mass
 	};
