Enriched Mass gameplay debugger category's display of entity's fragments - now including shared and chunk fragments.

diff --git a/Engine/Plugins/AI/MassAI/Source/MassAIDebug/Private/GameplayDebuggerCategory_Mass.cpp b/Engine/Plugins/AI/MassAI/Source/MassAIDebug/Private/GameplayDebuggerCategory_Mass.cpp
--- a/Engine/Plugins/AI/MassAI/Source/MassAIDebug/Private/GameplayDebuggerCategory_Mass.cpp
+++ b/Engine/Plugins/AI/MassAI/Source/MassAIDebug/Private/GameplayDebuggerCategory_Mass.cpp
@@ -225,7 +225,7 @@ void FGameplayDebuggerCategory_Mass::CollectData(APlayerController* OwnerPC, AAc
 			if (const FTransformFragment* TransformFragment = EntitySystem->GetFragmentDataPtr<FTransformFragment>(CachedEntity))
 			{
 				const FVector Location = TransformFragment->GetTransform().GetLocation();
-				AddShape(FGameplayDebuggerShape::MakeBox(Location, FVector(3,3,250), FColor::Purple));
+				AddShape(FGameplayDebuggerShape::MakeBox(Location, FVector(8,8,500), FColor::Purple,  FString::Printf(TEXT("[%s]"), *CachedEntity.DebugGetDescription())));
 				AddShape(FGameplayDebuggerShape::MakePoint(Location, 10, FColor::Purple));
 			}
 		}
@@ -292,35 +292,38 @@ void FGameplayDebuggerCategory_Mass::CollectData(APlayerController* OwnerPC, AAc
 				else
 				{
 					const FMassArchetypeHandle Archetype = EntitySystem->GetArchetypeForEntity(CachedEntity);
-					TArray<FName> ComponentNames;
-					TArray<FName> TagNames;
-					EntitySystem->DebugGetArchetypeStrings(Archetype, ComponentNames, TagNames);
-
-					FString Tags;
-					int i = 0;
-					for (const FName Name : TagNames)
-					{
-						Tags += FString::Printf(TEXT("%s, "), *Name.ToString());
-						if (i++ % 2)
+					const FMassArchetypeCompositionDescriptor& Composition = EntitySystem->GetArchetypeComposition(Archetype);
+					
+					auto DescriptionBuilder = [](const TArray<FName>& ItemNames) -> FString {
+						constexpr int ColumnsCount = 2;
+						FString Description;
+						int i = 0;
+						for (const FName Name : ItemNames)
 						{
-							Tags += TEXT("\n");
+							if ((i++ % ColumnsCount) == 0)
+							{
+								Description += TEXT("\n");
+							}
+							Description += FString::Printf(TEXT("%s,\t"), *Name.ToString());
 						}
-					}
-					AddTextLine(FString::Printf(TEXT("{Green}Tags:\n{White}%s"), *Tags));
+						return Description;
+					};
 
-					AddTextLine(FString::Printf(TEXT("{Green}Fragments:{White}")));
-					constexpr int ColumnsCount = 2;
-					i = 0;
-					while (i + ColumnsCount < ComponentNames.Num())
-					{
-						AddTextLine(FString::Printf(TEXT("%-42s, %-42s"), *ComponentNames[i].ToString(), *ComponentNames[i + 1].ToString()));
-						i += ColumnsCount;
-					}
-					if (i < ComponentNames.Num())
-					{
-						AddTextLine(FString::Printf(TEXT("%s"), *ComponentNames[i].ToString()));
-					}
+					TArray<FName> ItemNames;
+					Composition.Tags.DebugGetIndividualNames(ItemNames);
+					AddTextLine(FString::Printf(TEXT("{Green}Tags:{White}%s"), *DescriptionBuilder(ItemNames)));
+					
+					ItemNames.Reset();
+					Composition.Fragments.DebugGetIndividualNames(ItemNames);
+					AddTextLine(FString::Printf(TEXT("{Green}Fragments:{White}%s"), *DescriptionBuilder(ItemNames)));
+					
+					ItemNames.Reset();
+					Composition.ChunkFragments.DebugGetIndividualNames(ItemNames);
+					AddTextLine(FString::Printf(TEXT("{Green}Chunk Fragments:{White}%s"), *DescriptionBuilder(ItemNames)));
 
+					ItemNames.Reset();
+					Composition.SharedFragments.DebugGetIndividualNames(ItemNames);
+					AddTextLine(FString::Printf(TEXT("{Green}Shared Fragments:{White}%s"), *DescriptionBuilder(ItemNames)));
 				}
 
 				const FTransformFragment& TransformFragment = EntitySystem->GetFragmentDataChecked<FTransformFragment>(CachedEntity);
