Fixing fragment ownership and dependencies

diff --git a/Engine/Plugins/AI/MassAI/Source/MassNavigation/Private/MassAvoidanceTrait.cpp b/Engine/Plugins/AI/MassAI/Source/MassNavigation/Private/MassAvoidanceTrait.cpp
--- a/Engine/Plugins/AI/MassAI/Source/MassNavigation/Private/MassAvoidanceTrait.cpp
+++ b/Engine/Plugins/AI/MassAI/Source/MassNavigation/Private/MassAvoidanceTrait.cpp
@@ -1,4 +1,4 @@
-﻿// Copyright Epic Games, Inc. All Rights Reserved.
+// Copyright Epic Games, Inc. All Rights Reserved.
 
 #include "Avoidance/MassAvoidanceTrait.h"
 #include "Avoidance/MassAvoidanceFragments.h"
@@ -13,12 +13,12 @@ void UMassObstacleAvoidanceTrait::BuildTemplate(FMassEntityTemplateBuildContext&
 	UMassEntitySubsystem* EntitySubsystem = UWorld::GetSubsystem<UMassEntitySubsystem>(&World);
 	check(EntitySubsystem);
 
-	BuildContext.AddFragment<FAgentRadiusFragment>();
+	BuildContext.RequireFragment<FAgentRadiusFragment>();
 	BuildContext.AddFragment<FMassNavigationEdgesFragment>();
-	BuildContext.AddFragment<FTransformFragment>();
-	BuildContext.AddFragment<FMassVelocityFragment>();
-	BuildContext.AddFragment<FMassForceFragment>();
-	BuildContext.AddFragment<FMassMoveTargetFragment>();
+	BuildContext.RequireFragment<FTransformFragment>();
+	BuildContext.RequireFragment<FMassVelocityFragment>();
+	BuildContext.RequireFragment<FMassForceFragment>();
+	BuildContext.RequireFragment<FMassMoveTargetFragment>();
 
 	const FMassMovingAvoidanceParameters MovingValidated = MovingParameters.GetValidated();
 	const uint32 MovingHash = UE::StructUtils::GetStructCrc32(FConstStructView::Make(MovingValidated));
diff --git a/Engine/Plugins/AI/MassAI/Source/MassNavigation/Private/MassNavigationObstacleTrait.cpp b/Engine/Plugins/AI/MassAI/Source/MassNavigation/Private/MassNavigationObstacleTrait.cpp
--- a/Engine/Plugins/AI/MassAI/Source/MassNavigation/Private/MassNavigationObstacleTrait.cpp
+++ b/Engine/Plugins/AI/MassAI/Source/MassNavigation/Private/MassNavigationObstacleTrait.cpp
@@ -1,4 +1,4 @@
-﻿// Copyright Epic Games, Inc. All Rights Reserved.
+// Copyright Epic Games, Inc. All Rights Reserved.
 #include "Avoidance/MassNavigationObstacleTrait.h"
 #include "MassEntityTemplateRegistry.h"
 #include "MassNavigationFragments.h"
@@ -6,7 +6,7 @@
 
 void UMassNavigationObstacleTrait::BuildTemplate(FMassEntityTemplateBuildContext& BuildContext, UWorld& World) const
 {
-	BuildContext.AddFragment<FAgentRadiusFragment>();
+	BuildContext.RequireFragment<FAgentRadiusFragment>();
 
 	BuildContext.AddFragment<FMassNavigationObstacleGridCellLocationFragment>();
 }
diff --git a/Engine/Plugins/AI/MassAI/Source/MassNavigation/Private/MassSmoothOrientationTrait.cpp b/Engine/Plugins/AI/MassAI/Source/MassNavigation/Private/MassSmoothOrientationTrait.cpp
--- a/Engine/Plugins/AI/MassAI/Source/MassNavigation/Private/MassSmoothOrientationTrait.cpp
+++ b/Engine/Plugins/AI/MassAI/Source/MassNavigation/Private/MassSmoothOrientationTrait.cpp
@@ -1,4 +1,4 @@
-﻿// Copyright Epic Games, Inc. All Rights Reserved.
+// Copyright Epic Games, Inc. All Rights Reserved.
 #include "SmoothOrientation/MassSmoothOrientationTrait.h"
 #include "MassEntityTemplateRegistry.h"
 #include "MassMovementFragments.h"
@@ -11,9 +11,9 @@ void UMassSmoothOrientationTrait::BuildTemplate(FMassEntityTemplateBuildContext&
 	UMassEntitySubsystem* EntitySubsystem = UWorld::GetSubsystem<UMassEntitySubsystem>(&World);
 	check(EntitySubsystem);
 
-	BuildContext.AddFragment<FMassMoveTargetFragment>();
-	BuildContext.AddFragment<FMassVelocityFragment>();
-	BuildContext.AddFragment<FTransformFragment>();
+	BuildContext.RequireFragment<FMassMoveTargetFragment>();
+	BuildContext.RequireFragment<FMassVelocityFragment>();
+	BuildContext.RequireFragment<FTransformFragment>();
 
 	const FConstSharedStruct OrientationFragment = EntitySubsystem->GetOrCreateConstSharedFragment(UE::StructUtils::GetStructCrc32(FConstStructView::Make(Orientation)), Orientation);
 	BuildContext.AddConstSharedFragment(OrientationFragment);
diff --git a/Engine/Plugins/AI/MassAI/Source/MassNavigation/Private/MassSteeringTrait.cpp b/Engine/Plugins/AI/MassAI/Source/MassNavigation/Private/MassSteeringTrait.cpp
--- a/Engine/Plugins/AI/MassAI/Source/MassNavigation/Private/MassSteeringTrait.cpp
+++ b/Engine/Plugins/AI/MassAI/Source/MassNavigation/Private/MassSteeringTrait.cpp
@@ -1,4 +1,4 @@
-﻿// Copyright Epic Games, Inc. All Rights Reserved.
+// Copyright Epic Games, Inc. All Rights Reserved.
 
 #include "Steering/MassSteeringTrait.h"
 #include "MassEntityTemplateRegistry.h"
@@ -13,12 +13,12 @@ void UMassSteeringTrait::BuildTemplate(FMassEntityTemplateBuildContext& BuildCon
 	UMassEntitySubsystem* EntitySubsystem = UWorld::GetSubsystem<UMassEntitySubsystem>(&World);
 	check(EntitySubsystem);
 
-	BuildContext.AddFragment<FAgentRadiusFragment>();
-	BuildContext.AddFragment<FTransformFragment>();
-	BuildContext.AddFragment<FMassVelocityFragment>();
-	BuildContext.AddFragment<FMassForceFragment>();
+	BuildContext.RequireFragment<FAgentRadiusFragment>();
+	BuildContext.RequireFragment<FTransformFragment>();
+	BuildContext.RequireFragment<FMassVelocityFragment>();
+	BuildContext.RequireFragment<FMassForceFragment>();
+
 	BuildContext.AddFragment<FMassMoveTargetFragment>();
-	
 	BuildContext.AddFragment<FMassSteeringFragment>();
 	BuildContext.AddFragment<FMassStandingSteeringFragment>();
 	BuildContext.AddFragment<FMassGhostLocationFragment>();
diff --git a/Engine/Plugins/AI/MassAI/Source/MassZoneGraphNavigation/Private/MassZoneGraphNavigationTrait.cpp b/Engine/Plugins/AI/MassAI/Source/MassZoneGraphNavigation/Private/MassZoneGraphNavigationTrait.cpp
--- a/Engine/Plugins/AI/MassAI/Source/MassZoneGraphNavigation/Private/MassZoneGraphNavigationTrait.cpp
+++ b/Engine/Plugins/AI/MassAI/Source/MassZoneGraphNavigation/Private/MassZoneGraphNavigationTrait.cpp
@@ -1,4 +1,4 @@
-﻿// Copyright Epic Games, Inc. All Rights Reserved.
+// Copyright Epic Games, Inc. All Rights Reserved.
 #include "MassZoneGraphNavigationTrait.h"
 #include "MassEntityTemplateRegistry.h"
 #include "MassCommonFragments.h"
@@ -12,10 +12,10 @@ void UMassZoneGraphNavigationTrait::BuildTemplate(FMassEntityTemplateBuildContex
 	UMassEntitySubsystem* EntitySubsystem = UWorld::GetSubsystem<UMassEntitySubsystem>(&World);
 	check(EntitySubsystem);
 
-	BuildContext.AddFragment<FAgentRadiusFragment>();
-	BuildContext.AddFragment<FTransformFragment>();
-	BuildContext.AddFragment<FMassVelocityFragment>();
-	BuildContext.AddFragment<FMassMoveTargetFragment>();
+	BuildContext.RequireFragment<FAgentRadiusFragment>();
+	BuildContext.RequireFragment<FTransformFragment>();
+	BuildContext.RequireFragment<FMassVelocityFragment>();
+	BuildContext.RequireFragment<FMassMoveTargetFragment>();
 
 	BuildContext.AddFragment<FMassZoneGraphLaneLocationFragment>();
 	BuildContext.AddFragment<FMassZoneGraphPathRequestFragment>();
diff --git a/Engine/Plugins/AI/MassCrowd/Source/MassCrowd/Private/MassCrowdServerRepresentationTrait.cpp b/Engine/Plugins/AI/MassCrowd/Source/MassCrowd/Private/MassCrowdServerRepresentationTrait.cpp
--- a/Engine/Plugins/AI/MassCrowd/Source/MassCrowd/Private/MassCrowdServerRepresentationTrait.cpp
+++ b/Engine/Plugins/AI/MassCrowd/Source/MassCrowd/Private/MassCrowdServerRepresentationTrait.cpp
@@ -8,6 +8,7 @@
 #include "MassRepresentationFragments.h"
 #include "Engine/World.h"
 #include "MassCrowdRepresentationActorManagement.h"
+#include "MassActorSubsystem.h"
 
 UMassCrowdServerRepresentationTrait::UMassCrowdServerRepresentationTrait()
 {
@@ -31,10 +32,9 @@ void UMassCrowdServerRepresentationTrait::BuildTemplate(FMassEntityTemplateBuild
 		return;
 	}
 	
-	// the following needs to be always there for mesh vis to work. Adding following fragments after already 
-	// adding Config.AdditionalDataFragments to let user configure the fragments first. Calling BuildContext.Add() 
-	// won't override any fragments that are already there
-	BuildContext.AddFragment<FTransformFragment>();
+	BuildContext.RequireFragment<FMassViewerInfoFragment>();
+	BuildContext.RequireFragment<FTransformFragment>();
+	BuildContext.RequireFragment<FMassActorFragment>();
 
 	UMassEntitySubsystem* EntitySubsystem = UWorld::GetSubsystem<UMassEntitySubsystem>(&World);
 	check(EntitySubsystem);
diff --git a/Engine/Plugins/Runtime/MassGameplay/Source/MassActors/Private/MassAgentTraits.cpp b/Engine/Plugins/Runtime/MassGameplay/Source/MassActors/Private/MassAgentTraits.cpp
--- a/Engine/Plugins/Runtime/MassGameplay/Source/MassActors/Private/MassAgentTraits.cpp
+++ b/Engine/Plugins/Runtime/MassGameplay/Source/MassActors/Private/MassAgentTraits.cpp
@@ -128,8 +128,7 @@ void UMassAgentMovementSyncTrait::BuildTemplate(FMassEntityTemplateBuildContext&
 //----------------------------------------------------------------------//
 void UMassAgentOrientationSyncTrait::BuildTemplate(FMassEntityTemplateBuildContext& BuildContext, UWorld& World) const
 {
-	BuildContext.AddFragment<FCharacterMovementComponentWrapperFragment>();
-	// @todo: Figure out how we can share init with UMassAgentMovementSyncTrait, or make this this trait to depend on UMassAgentMovementSyncTrait.
+	BuildContext.RequireFragment<FCharacterMovementComponentWrapperFragment>();
 
 	if (EnumHasAnyFlags(SyncDirection, EMassTranslationDirection::ActorToMass))
 	{
diff --git a/Engine/Plugins/Runtime/MassGameplay/Source/MassLOD/Private/MassLODTrait.cpp b/Engine/Plugins/Runtime/MassGameplay/Source/MassLOD/Private/MassLODTrait.cpp
--- a/Engine/Plugins/Runtime/MassGameplay/Source/MassLOD/Private/MassLODTrait.cpp
+++ b/Engine/Plugins/Runtime/MassGameplay/Source/MassLOD/Private/MassLODTrait.cpp
@@ -11,12 +11,13 @@ void UMassLODCollectorTrait::BuildTemplate(FMassEntityTemplateBuildContext& Buil
 {
 	BuildContext.AddFragment<FMassViewerInfoFragment>();
 	BuildContext.AddTag<FMassCollectLODViewerInfoTag>();
+	BuildContext.RequireFragment<FTransformFragment>();
 }
 
 void UMassSimulationLODTrait::BuildTemplate(FMassEntityTemplateBuildContext& BuildContext, UWorld& World) const
 {
-	BuildContext.AddTag<FMassCollectLODViewerInfoTag>(); // Depends on FMassViewerInfoFragment
-	BuildContext.AddFragment<FTransformFragment>();
+	BuildContext.RequireFragment<FMassViewerInfoFragment>();
+	BuildContext.RequireFragment<FTransformFragment>();
 
 	FMassSimulationLODFragment& LODFragment = BuildContext.AddFragment_GetRef<FMassSimulationLODFragment>();
 
diff --git a/Engine/Plugins/Runtime/MassGameplay/Source/MassMovement/Private/MassMovementTrait.cpp b/Engine/Plugins/Runtime/MassGameplay/Source/MassMovement/Private/MassMovementTrait.cpp
--- a/Engine/Plugins/Runtime/MassGameplay/Source/MassMovement/Private/MassMovementTrait.cpp
+++ b/Engine/Plugins/Runtime/MassGameplay/Source/MassMovement/Private/MassMovementTrait.cpp
@@ -11,8 +11,8 @@ void UMassMovementTrait::BuildTemplate(FMassEntityTemplateBuildContext& BuildCon
 	UMassEntitySubsystem* EntitySubsystem = UWorld::GetSubsystem<UMassEntitySubsystem>(&World);
 	check(EntitySubsystem);
 
-	BuildContext.AddFragment<FAgentRadiusFragment>();
-	BuildContext.AddFragment<FTransformFragment>();
+	BuildContext.RequireFragment<FAgentRadiusFragment>();
+	BuildContext.RequireFragment<FTransformFragment>();
 
 	BuildContext.AddFragment<FMassVelocityFragment>();
 	BuildContext.AddFragment<FMassForceFragment>();
diff --git a/Engine/Plugins/Runtime/MassGameplay/Source/MassRepresentation/Private/MassVisualizationTrait.cpp b/Engine/Plugins/Runtime/MassGameplay/Source/MassRepresentation/Private/MassVisualizationTrait.cpp
--- a/Engine/Plugins/Runtime/MassGameplay/Source/MassRepresentation/Private/MassVisualizationTrait.cpp
+++ b/Engine/Plugins/Runtime/MassGameplay/Source/MassRepresentation/Private/MassVisualizationTrait.cpp
@@ -8,6 +8,7 @@
 #include "MassRepresentationActorManagement.h"
 #include "Engine/World.h"
 #include "MassLODFragments.h"
+#include "MassActorSubsystem.h"
 
 UMassVisualizationTrait::UMassVisualizationTrait()
 {
@@ -47,11 +48,9 @@ void UMassVisualizationTrait::BuildTemplate(FMassEntityTemplateBuildContext& Bui
 		return;
 	}
 
-	// the following needs to be always there for mesh vis to work. Adding following fragments after already 
-	// adding Config.AdditionalDataFragments to let user configure the fragments first. Calling BuildContext.Add() 
-	// won't override any fragments that are already there
-	BuildContext.AddTag<FMassCollectLODViewerInfoTag>(); // Depends on FMassViewerInfoFragment
-	BuildContext.AddFragment<FTransformFragment>();
+	BuildContext.RequireFragment<FMassViewerInfoFragment>();
+	BuildContext.RequireFragment<FTransformFragment>();
+	BuildContext.RequireFragment<FMassActorFragment>();
 
 	UMassEntitySubsystem* EntitySubsystem = UWorld::GetSubsystem<UMassEntitySubsystem>(&World);
 	check(EntitySubsystem);
