Fix MassStateTreeProcessor compiling in shipping

diff --git a/Engine/Plugins/AI/MassAI/Source/MassAIBehavior/Private/MassStateTreeProcessors.cpp b/Engine/Plugins/AI/MassAI/Source/MassAIBehavior/Private/MassStateTreeProcessors.cpp
--- a/Engine/Plugins/AI/MassAI/Source/MassAIBehavior/Private/MassStateTreeProcessors.cpp
+++ b/Engine/Plugins/AI/MassAI/Source/MassAIBehavior/Private/MassStateTreeProcessors.cpp
@@ -409,6 +409,7 @@ void UMassStateTreeProcessor::SignalEntities(UMassEntitySubsystem& EntitySubsyst
 					const float AdjustedDeltaTime = TimeInSeconds - StateTreeInstance.LastUpdateTimeInSeconds;
 					StateTreeInstance.LastUpdateTimeInSeconds = TimeInSeconds;
 
+#if WITH_MASSGAMEPLAY_DEBUG
 					const FMassEntityHandle Entity = StateTreeExecutionContext.GetEntity();
 					if (UE::Mass::Debug::IsDebuggingEntity(Entity))
 					{
@@ -425,6 +426,7 @@ void UMassStateTreeProcessor::SignalEntities(UMassEntitySubsystem& EntitySubsyst
 						}
 						UE_VLOG_UELOG(&MassStateTreeSubsystem, LogStateTree, Log, TEXT("%s: Ticking StateTree because of signals: %s"), *Entity.DebugGetDescription(), *SignalsString);
 					}
+#endif // WITH_MASSGAMEPLAY_DEBUG
 
 					// Tick the tree instance
 					StateTreeExecutionContext.Tick(AdjustedDeltaTime, &InstanceData);
