MassStateTreeProcesser fixed non-unity compilation.

diff --git a/Engine/Plugins/AI/MassAI/Source/MassAIBehavior/Private/MassStateTreeProcessors.cpp b/Engine/Plugins/AI/MassAI/Source/MassAIBehavior/Private/MassStateTreeProcessors.cpp
--- a/Engine/Plugins/AI/MassAI/Source/MassAIBehavior/Private/MassStateTreeProcessors.cpp
+++ b/Engine/Plugins/AI/MassAI/Source/MassAIBehavior/Private/MassStateTreeProcessors.cpp
@@ -14,6 +14,7 @@
 #include "ProfilingDebugging/CsvProfiler.h"
 #include "Engine/World.h"
 #include "MassBehaviorSettings.h"
+#include "VisualLogger/VisualLogger.h"
 
 CSV_DEFINE_CATEGORY(StateTreeProcessor, true);
 
