Added a way to mark up MassEntityQueries as owned by a processor and auto-registering such queries. Also added some verication code that ensures all queries run as part of processors are properly marked up as such.
This is a prerequisite step to support fetching processors' requirements.

Updated all the queries owned by processors at this moment.

diff --git a/Engine/Plugins/AI/MassAI/Source/MassAIBehavior/Private/MassLookAtProcessors.cpp b/Engine/Plugins/AI/MassAI/Source/MassAIBehavior/Private/MassLookAtProcessors.cpp
--- a/Engine/Plugins/AI/MassAI/Source/MassAIBehavior/Private/MassLookAtProcessors.cpp
+++ b/Engine/Plugins/AI/MassAI/Source/MassAIBehavior/Private/MassLookAtProcessors.cpp
@@ -89,6 +89,7 @@ namespace UE::MassBehavior
 // UMassLookAtProcessor
 //----------------------------------------------------------------------//
 UMassLookAtProcessor::UMassLookAtProcessor()
+	: EntityQuery_Conditional(*this)
 {
 	ExecutionFlags = (int32)(EProcessorExecutionFlags::Client | EProcessorExecutionFlags::Standalone);
 	ExecutionOrder.ExecuteInGroup = UE::Mass::ProcessorGroupNames::Tasks;
diff --git a/Engine/Plugins/AI/MassAI/Source/MassAIBehavior/Private/MassStateTreeProcessors.cpp b/Engine/Plugins/AI/MassAI/Source/MassAIBehavior/Private/MassStateTreeProcessors.cpp
--- a/Engine/Plugins/AI/MassAI/Source/MassAIBehavior/Private/MassStateTreeProcessors.cpp
+++ b/Engine/Plugins/AI/MassAI/Source/MassAIBehavior/Private/MassStateTreeProcessors.cpp
@@ -173,6 +173,7 @@ void ForEachEntityInChunk(FMassExecutionContext& Context, FMassStateTreeExecutio
 // UMassStateTreeFragmentDestructor
 //----------------------------------------------------------------------//
 UMassStateTreeFragmentDestructor::UMassStateTreeFragmentDestructor()
+	: EntityQuery(*this)
 {
 	ExecutionFlags = (int32)(EProcessorExecutionFlags::Standalone | EProcessorExecutionFlags::Server);
 	ObservedType = FMassStateTreeInstanceFragment::StaticStruct();
@@ -233,6 +234,7 @@ void UMassStateTreeFragmentDestructor::Execute(UMassEntitySubsystem& EntitySubsy
 // UMassStateTreeActivationProcessor
 //----------------------------------------------------------------------//
 UMassStateTreeActivationProcessor::UMassStateTreeActivationProcessor()
+	: EntityQuery(*this)
 {
 	ExecutionOrder.ExecuteAfter.Add(UE::Mass::ProcessorGroupNames::LOD);
 	ExecutionOrder.ExecuteBefore.Add(UE::Mass::ProcessorGroupNames::Behavior);
diff --git a/Engine/Plugins/AI/MassAI/Source/MassAIBehavior/Private/MassZoneGraphAnnotationProcessors.cpp b/Engine/Plugins/AI/MassAI/Source/MassAIBehavior/Private/MassZoneGraphAnnotationProcessors.cpp
--- a/Engine/Plugins/AI/MassAI/Source/MassAIBehavior/Private/MassZoneGraphAnnotationProcessors.cpp
+++ b/Engine/Plugins/AI/MassAI/Source/MassAIBehavior/Private/MassZoneGraphAnnotationProcessors.cpp
@@ -15,6 +15,7 @@
 //  UMassZoneGraphAnnotationTagsInitializer
 //----------------------------------------------------------------------//
 UMassZoneGraphAnnotationTagsInitializer::UMassZoneGraphAnnotationTagsInitializer()
+	: EntityQuery(*this)
 {
 	ObservedType = FMassZoneGraphAnnotationFragment::StaticStruct();
 	Operation = EMassObservedOperation::Add;
diff --git a/Engine/Plugins/AI/MassAI/Source/MassAIDebug/Private/MassDebugStateTreeProcessor.cpp b/Engine/Plugins/AI/MassAI/Source/MassAIDebug/Private/MassDebugStateTreeProcessor.cpp
--- a/Engine/Plugins/AI/MassAI/Source/MassAIDebug/Private/MassDebugStateTreeProcessor.cpp
+++ b/Engine/Plugins/AI/MassAI/Source/MassAIDebug/Private/MassDebugStateTreeProcessor.cpp
@@ -13,6 +13,7 @@
 // UMassDebugStateTreeProcessor
 //----------------------------------------------------------------------//
 UMassDebugStateTreeProcessor::UMassDebugStateTreeProcessor()
+	: EntityQuery(*this)
 {
 	ExecutionOrder.ExecuteInGroup = UE::Mass::ProcessorGroupNames::Behavior;
 	ExecutionOrder.ExecuteAfter.Add(TEXT("MassStateTreeProcessor"));
diff --git a/Engine/Plugins/AI/MassAI/Source/MassNavigation/Private/MassAvoidanceProcessors.cpp b/Engine/Plugins/AI/MassAI/Source/MassNavigation/Private/MassAvoidanceProcessors.cpp
--- a/Engine/Plugins/AI/MassAI/Source/MassNavigation/Private/MassAvoidanceProcessors.cpp
+++ b/Engine/Plugins/AI/MassAI/Source/MassNavigation/Private/MassAvoidanceProcessors.cpp
@@ -345,6 +345,7 @@ namespace UE::MassAvoidance
 //  UMassMovingAvoidanceProcessor
 //----------------------------------------------------------------------//
 UMassMovingAvoidanceProcessor::UMassMovingAvoidanceProcessor()
+	: EntityQuery(*this) 
 {
 	bAutoRegisterWithProcessingPhases = true;
 	ExecutionFlags = (int32)EProcessorExecutionFlags::All;
@@ -1008,6 +1009,8 @@ void UMassMovingAvoidanceProcessor::Execute(UMassEntitySubsystem& EntitySubsyste
 //  UMassStandingAvoidanceProcessor
 //----------------------------------------------------------------------//
 UMassStandingAvoidanceProcessor::UMassStandingAvoidanceProcessor()
+	: EntityQuery(*this)
+
 {
 	bAutoRegisterWithProcessingPhases = true;
 	ExecutionFlags = (int32)EProcessorExecutionFlags::All;
diff --git a/Engine/Plugins/AI/MassAI/Source/MassNavigation/Private/MassNavigationProcessors.cpp b/Engine/Plugins/AI/MassAI/Source/MassNavigation/Private/MassNavigationProcessors.cpp
--- a/Engine/Plugins/AI/MassAI/Source/MassNavigation/Private/MassNavigationProcessors.cpp
+++ b/Engine/Plugins/AI/MassAI/Source/MassNavigation/Private/MassNavigationProcessors.cpp
@@ -21,6 +21,7 @@
 //----------------------------------------------------------------------//
 
 UMassOffLODNavigationProcessor::UMassOffLODNavigationProcessor()
+	: EntityQuery_Conditional(*this)
 {
 	ExecutionFlags = (int32)EProcessorExecutionFlags::All;
 	ExecutionOrder.ExecuteInGroup = UE::Mass::ProcessorGroupNames::Movement;
@@ -69,6 +70,7 @@ void UMassOffLODNavigationProcessor::Execute(UMassEntitySubsystem& EntitySubsyst
 //----------------------------------------------------------------------//
 
 UMassNavigationSmoothHeightProcessor::UMassNavigationSmoothHeightProcessor()
+	: EntityQuery(*this)
 {
 	ExecutionFlags = (int32)EProcessorExecutionFlags::All;
 	ExecutionOrder.ExecuteAfter.Add(UE::Mass::ProcessorGroupNames::Movement);
@@ -121,6 +123,7 @@ void UMassNavigationSmoothHeightProcessor::Execute(UMassEntitySubsystem& EntityS
 //----------------------------------------------------------------------//
 
 UMassMoveTargetFragmentInitializer::UMassMoveTargetFragmentInitializer()
+	: InitializerQuery(*this)
 {
 	ObservedType = FMassMoveTargetFragment::StaticStruct();
 	Operation = EMassObservedOperation::Add;
@@ -175,15 +178,18 @@ void UMassNavigationObstacleGridProcessor::ConfigureQueries()
 	AddToGridEntityQuery.AddRequirement<FMassAvoidanceColliderFragment>(EMassFragmentAccess::ReadOnly, EMassFragmentPresence::Optional);
 	AddToGridEntityQuery.AddTagRequirement<FMassOffLODTag>(EMassFragmentPresence::None);
 	AddToGridEntityQuery.AddTagRequirement<FMassInNavigationObstacleGridTag>(EMassFragmentPresence::None);
+	AddToGridEntityQuery.RegisterWithProcessor(*this);
 
 	UpdateGridEntityQuery = BaseEntityQuery;
 	UpdateGridEntityQuery.AddRequirement<FMassAvoidanceColliderFragment>(EMassFragmentAccess::ReadOnly, EMassFragmentPresence::Optional);
 	UpdateGridEntityQuery.AddTagRequirement<FMassOffLODTag>(EMassFragmentPresence::None);
 	UpdateGridEntityQuery.AddTagRequirement<FMassInNavigationObstacleGridTag>(EMassFragmentPresence::All);
+	UpdateGridEntityQuery.RegisterWithProcessor(*this);
 
 	RemoveFromGridEntityQuery = BaseEntityQuery;
 	RemoveFromGridEntityQuery.AddTagRequirement<FMassOffLODTag>(EMassFragmentPresence::All);
 	RemoveFromGridEntityQuery.AddTagRequirement<FMassInNavigationObstacleGridTag>(EMassFragmentPresence::All);
+	RemoveFromGridEntityQuery.RegisterWithProcessor(*this);
 }
 
 void UMassNavigationObstacleGridProcessor::Execute(UMassEntitySubsystem& EntitySubsystem, FMassExecutionContext& Context)
@@ -272,6 +278,7 @@ void UMassNavigationObstacleGridProcessor::Execute(UMassEntitySubsystem& EntityS
 //  UMassNavigationObstacleRemoverProcessor
 //----------------------------------------------------------------------//
 UMassNavigationObstacleRemoverProcessor::UMassNavigationObstacleRemoverProcessor()
+	: EntityQuery(*this)
 {
 	ObservedType = FMassNavigationObstacleGridCellLocationFragment::StaticStruct();
 	Operation = EMassObservedOperation::Remove;
diff --git a/Engine/Plugins/AI/MassAI/Source/MassNavigation/Private/MassSmoothOrientationProcessors.cpp b/Engine/Plugins/AI/MassAI/Source/MassNavigation/Private/MassSmoothOrientationProcessors.cpp
--- a/Engine/Plugins/AI/MassAI/Source/MassNavigation/Private/MassSmoothOrientationProcessors.cpp
+++ b/Engine/Plugins/AI/MassAI/Source/MassNavigation/Private/MassSmoothOrientationProcessors.cpp
@@ -1,4 +1,4 @@
-﻿// Copyright Epic Games, Inc. All Rights Reserved.
+// Copyright Epic Games, Inc. All Rights Reserved.
 
 #include "SmoothOrientation/MassSmoothOrientationProcessors.h"
 #include "SmoothOrientation/MassSmoothOrientationFragments.h"
@@ -17,6 +17,8 @@
 //  UMassSmoothOrientationProcessor
 //----------------------------------------------------------------------//
 UMassSmoothOrientationProcessor::UMassSmoothOrientationProcessor()
+	: HighResEntityQuery(*this)
+	, LowResEntityQuery_Conditional(*this)
 {
 	ExecutionFlags = (int32)EProcessorExecutionFlags::All;
 	ExecutionOrder.ExecuteInGroup = UE::Mass::ProcessorGroupNames::Movement;
diff --git a/Engine/Plugins/AI/MassAI/Source/MassNavigation/Private/MassSteeringProcessors.cpp b/Engine/Plugins/AI/MassAI/Source/MassNavigation/Private/MassSteeringProcessors.cpp
--- a/Engine/Plugins/AI/MassAI/Source/MassNavigation/Private/MassSteeringProcessors.cpp
+++ b/Engine/Plugins/AI/MassAI/Source/MassNavigation/Private/MassSteeringProcessors.cpp
@@ -1,4 +1,4 @@
-﻿// Copyright Epic Games, Inc. All Rights Reserved.
+// Copyright Epic Games, Inc. All Rights Reserved.
 
 #include "Steering/MassSteeringProcessors.h"
 #include "MassCommonUtils.h"
@@ -51,6 +51,7 @@ namespace UE::MassNavigation
 //  UMassSteerToMoveTargetProcessor
 //----------------------------------------------------------------------//
 UMassSteerToMoveTargetProcessor::UMassSteerToMoveTargetProcessor()
+	: EntityQuery(*this)
 {
 	ExecutionFlags = int32(EProcessorExecutionFlags::All);
 	ExecutionOrder.ExecuteAfter.Add(UE::Mass::ProcessorGroupNames::Tasks);
diff --git a/Engine/Plugins/AI/MassAI/Source/MassZoneGraphNavigation/Private/MassZoneGraphNavigationProcessors.cpp b/Engine/Plugins/AI/MassAI/Source/MassZoneGraphNavigation/Private/MassZoneGraphNavigationProcessors.cpp
--- a/Engine/Plugins/AI/MassAI/Source/MassZoneGraphNavigation/Private/MassZoneGraphNavigationProcessors.cpp
+++ b/Engine/Plugins/AI/MassAI/Source/MassZoneGraphNavigation/Private/MassZoneGraphNavigationProcessors.cpp
@@ -36,6 +36,7 @@ namespace UE::MassNavigation::Debug
 //  UMassZoneGraphLocationInitializer
 //----------------------------------------------------------------------//
 UMassZoneGraphLocationInitializer::UMassZoneGraphLocationInitializer()
+	: EntityQuery(*this)
 {
 	ObservedType = FMassZoneGraphLaneLocationFragment::StaticStruct();
 	Operation = EMassObservedOperation::Add;
@@ -110,6 +111,7 @@ void UMassZoneGraphLocationInitializer::Execute(UMassEntitySubsystem& EntitySubs
 //  UMassZoneGraphPathFollowProcessor
 //----------------------------------------------------------------------//
 UMassZoneGraphPathFollowProcessor::UMassZoneGraphPathFollowProcessor()
+	: EntityQuery_Conditional(*this)
 {
 	ExecutionFlags = (int32)EProcessorExecutionFlags::All;
 	ExecutionOrder.ExecuteInGroup = UE::Mass::ProcessorGroupNames::Tasks;
@@ -409,6 +411,7 @@ void UMassZoneGraphPathFollowProcessor::Execute(UMassEntitySubsystem& EntitySubs
 //  UMassZoneGraphLaneCacheBoundaryProcessor
 //----------------------------------------------------------------------//
 UMassZoneGraphLaneCacheBoundaryProcessor::UMassZoneGraphLaneCacheBoundaryProcessor()
+	: EntityQuery(*this)
 {
 	ExecutionFlags = (int32)EProcessorExecutionFlags::All;
 
diff --git a/Engine/Plugins/AI/MassCrowd/Source/MassCrowd/Private/MassCrowdNavigationProcessor.cpp b/Engine/Plugins/AI/MassCrowd/Source/MassCrowd/Private/MassCrowdNavigationProcessor.cpp
--- a/Engine/Plugins/AI/MassCrowd/Source/MassCrowd/Private/MassCrowdNavigationProcessor.cpp
+++ b/Engine/Plugins/AI/MassCrowd/Source/MassCrowd/Private/MassCrowdNavigationProcessor.cpp
@@ -62,6 +62,7 @@ void UMassCrowdLaneTrackingSignalProcessor::SignalEntities(UMassEntitySubsystem&
 // UMassCrowdLaneTrackingDestructor
 //----------------------------------------------------------------------//
 UMassCrowdLaneTrackingDestructor::UMassCrowdLaneTrackingDestructor()
+	: EntityQuery(*this)
 {
 	ExecutionFlags = (int32)(EProcessorExecutionFlags::Standalone | EProcessorExecutionFlags::Server);
 	ObservedType = FMassCrowdLaneTrackingFragment::StaticStruct();
@@ -99,6 +100,7 @@ void UMassCrowdLaneTrackingDestructor::Execute(UMassEntitySubsystem& EntitySubsy
 //----------------------------------------------------------------------//
 
 UMassCrowdDynamicObstacleProcessor::UMassCrowdDynamicObstacleProcessor()
+	: EntityQuery_Conditional(*this)
 {
 	bAutoRegisterWithProcessingPhases = true;
 
@@ -222,6 +224,7 @@ void UMassCrowdDynamicObstacleProcessor::Execute(UMassEntitySubsystem& EntitySub
 // UMassCrowdDynamicObstacleInitializer
 //----------------------------------------------------------------------//
 UMassCrowdDynamicObstacleInitializer::UMassCrowdDynamicObstacleInitializer()
+	: EntityQuery(*this)
 {
 	ExecutionFlags = (int32)(EProcessorExecutionFlags::Standalone | EProcessorExecutionFlags::Server);
 	ObservedType = FMassCrowdObstacleFragment::StaticStruct();
@@ -261,6 +264,7 @@ void UMassCrowdDynamicObstacleInitializer::Execute(UMassEntitySubsystem& EntityS
 // UMassCrowdDynamicObstacleDeinitializer
 //----------------------------------------------------------------------//
 UMassCrowdDynamicObstacleDeinitializer::UMassCrowdDynamicObstacleDeinitializer()
+	: EntityQuery(*this)
 {
 	ExecutionFlags = (int32)(EProcessorExecutionFlags::Standalone | EProcessorExecutionFlags::Server);
 	ObservedType = FMassCrowdObstacleFragment::StaticStruct();
diff --git a/Engine/Plugins/AI/MassCrowd/Source/MassCrowd/Private/MassCrowdServerRepresentationLODProcessor.cpp b/Engine/Plugins/AI/MassCrowd/Source/MassCrowd/Private/MassCrowdServerRepresentationLODProcessor.cpp
--- a/Engine/Plugins/AI/MassCrowd/Source/MassCrowd/Private/MassCrowdServerRepresentationLODProcessor.cpp
+++ b/Engine/Plugins/AI/MassCrowd/Source/MassCrowd/Private/MassCrowdServerRepresentationLODProcessor.cpp
@@ -11,6 +11,7 @@ namespace UE::MassCrowd
 } // UE::MassCrowd
 
 UMassCrowdServerRepresentationLODProcessor::UMassCrowdServerRepresentationLODProcessor()
+	: EntityQuery(*this)
 {
 	ExecutionFlags = (int32)EProcessorExecutionFlags::Server;
 
diff --git a/Engine/Plugins/AI/MassCrowd/Source/MassCrowd/Private/MassCrowdVisualizationProcessor.cpp b/Engine/Plugins/AI/MassCrowd/Source/MassCrowd/Private/MassCrowdVisualizationProcessor.cpp
--- a/Engine/Plugins/AI/MassCrowd/Source/MassCrowd/Private/MassCrowdVisualizationProcessor.cpp
+++ b/Engine/Plugins/AI/MassCrowd/Source/MassCrowd/Private/MassCrowdVisualizationProcessor.cpp
@@ -45,6 +45,7 @@ void UMassCrowdVisualizationProcessor::ConfigureQueries()
 // UMassDebugCrowdVisualizationProcessor
 //----------------------------------------------------------------------//
 UMassDebugCrowdVisualizationProcessor::UMassDebugCrowdVisualizationProcessor()
+	: EntityQuery(*this)
 {
 	ExecutionFlags = (int32)(EProcessorExecutionFlags::Client | EProcessorExecutionFlags::Standalone);
 
diff --git a/Engine/Plugins/Runtime/MassEntity/Source/MassEntity/Private/MassEntityQuery.cpp b/Engine/Plugins/Runtime/MassEntity/Source/MassEntity/Private/MassEntityQuery.cpp
--- a/Engine/Plugins/Runtime/MassEntity/Source/MassEntity/Private/MassEntityQuery.cpp
+++ b/Engine/Plugins/Runtime/MassEntity/Source/MassEntity/Private/MassEntityQuery.cpp
@@ -8,6 +8,10 @@
 #include "VisualLogger/VisualLogger.h"
 #include "Async/ParallelFor.h"
 #include "Containers/UnrealString.h"
+#include "MassProcessor.h"
+#if WITH_MASSENTITY_DEBUG
+#include "MassRequirementAccessDetector.h"
+#endif // WITH_MASSENTITY_DEBUG
 
 //////////////////////////////////////////////////////////////////////
 // FMassEntityQuery
@@ -36,6 +40,20 @@ FMassEntityQuery::FMassEntityQuery(TConstArrayView<const UScriptStruct*> InitLis
 	}
 }
 
+FMassEntityQuery::FMassEntityQuery(UMassProcessor& Owner)
+{
+	RegisterWithProcessor(Owner);
+}
+
+void FMassEntityQuery::RegisterWithProcessor(UMassProcessor& Owner)
+{
+	ExpectedContextType = EMassExecutionContextType::Processor;
+	Owner.RegisterQuery(*this);
+#if WITH_MASSENTITY_DEBUG
+	bRegistered = true;
+#endif // WITH_MASSENTITY_DEBUG
+}
+
 void FMassEntityQuery::ReadCommandlineParams()
 {
 	int AllowParallelQueries = -1;
@@ -137,6 +155,9 @@ void FMassEntityQuery::ForEachEntityChunk(UMassEntitySubsystem& EntitySubsystem,
 #if WITH_MASSENTITY_DEBUG
 	int32 NumEntitiesToProcess = 0;
 
+	checkf(ExecutionContext.ExecutionType == ExpectedContextType && (ExpectedContextType == EMassExecutionContextType::Local || bRegistered)
+		, TEXT("ExecutionContextType mismatch, make sure all the queries run as part of processor execution are registered with some processor with a FMassEntityQuery::RegisterWithProcessor call"));
+
 	EntitySubsystem.GetRequirementAccessDetector().RequireAccess(*this);
 #endif
 
diff --git a/Engine/Plugins/Runtime/MassEntity/Source/MassEntity/Private/MassExecutor.cpp b/Engine/Plugins/Runtime/MassEntity/Source/MassEntity/Private/MassExecutor.cpp
--- a/Engine/Plugins/Runtime/MassEntity/Source/MassEntity/Private/MassExecutor.cpp
+++ b/Engine/Plugins/Runtime/MassEntity/Source/MassEntity/Private/MassExecutor.cpp
@@ -97,6 +97,7 @@ void RunProcessorsView(TArrayView<UMassProcessor*> Processors, FMassProcessingCo
 	ExecutionContext.SetDeferredCommandBuffer(CommandBuffer);
 	ExecutionContext.SetFlushDeferredCommands(false);
 	ExecutionContext.SetAuxData(ProcessingContext.AuxData);
+	ExecutionContext.SetExecutionType(EMassExecutionContextType::Processor);
 
 	{
 		TRACE_CPUPROFILER_EVENT_SCOPE_STR("Execute Processors")
diff --git a/Engine/Plugins/Runtime/MassEntity/Source/MassEntity/Private/MassProcessor.cpp b/Engine/Plugins/Runtime/MassEntity/Source/MassEntity/Private/MassProcessor.cpp
--- a/Engine/Plugins/Runtime/MassEntity/Source/MassEntity/Private/MassProcessor.cpp
+++ b/Engine/Plugins/Runtime/MassEntity/Source/MassEntity/Private/MassProcessor.cpp
@@ -157,6 +157,25 @@ void UMassProcessor::CallExecute(UMassEntitySubsystem& EntitySubsystem, FMassExe
 	Execute(EntitySubsystem, Context);
 }
 
+void UMassProcessor::RegisterQuery(FMassEntityQuery& Query)
+{
+	const uintptr_t ThisStart = (uintptr_t)this;
+	const uintptr_t ThisEnd = ThisStart + GetClass()->GetStructureSize();
+	const uintptr_t QueryStart = (uintptr_t)&Query;
+	const uintptr_t QueryEnd = QueryStart + sizeof(FMassEntityQuery);
+
+	if (QueryStart >= ThisStart && QueryEnd <= ThisEnd)
+	{
+		OwnedQueries.AddUnique(&Query);
+	}
+	else
+	{
+		constexpr TCHAR MessageFormat[] = TEXT("Registering entity query for %s while the query is not given processor's member variable. Skipping.");
+		checkf(false, MessageFormat, *GetProcessorName());
+		UE_LOG(LogMass, Error, MessageFormat, *GetProcessorName());
+	}
+}
+
 FGraphEventRef UMassProcessor::DispatchProcessorTasks(UMassEntitySubsystem& EntitySubsystem, FMassExecutionContext& ExecutionContext, const FGraphEventArray& Prerequisites)
 {
 	FGraphEventRef ReturnVal;
diff --git a/Engine/Plugins/Runtime/MassEntity/Source/MassEntity/Public/MassEntityQuery.h b/Engine/Plugins/Runtime/MassEntity/Source/MassEntity/Public/MassEntityQuery.h
--- a/Engine/Plugins/Runtime/MassEntity/Source/MassEntity/Public/MassEntityQuery.h
+++ b/Engine/Plugins/Runtime/MassEntity/Source/MassEntity/Public/MassEntityQuery.h
@@ -11,6 +11,7 @@
 
 
 class UMassEntitySubsystem;
+class UMassProcessor;
 struct FMassArchetypeData;
 struct FMassExecutionContext;
 struct FMassFragment;
@@ -103,6 +104,9 @@ public:
 	FMassEntityQuery();
 	FMassEntityQuery(std::initializer_list<UScriptStruct*> InitList);
 	FMassEntityQuery(TConstArrayView<const UScriptStruct*> InitList);
+	FMassEntityQuery(UMassProcessor& Owner);
+
+	void RegisterWithProcessor(UMassProcessor& Owner);
 
 	/** Runs ExecuteFunction on all entities matching Requirements */
 	void ForEachEntityChunk(UMassEntitySubsystem& EntitySubsystem, FMassExecutionContext& ExecutionContext, const FMassExecuteFunction& ExecuteFunction);
@@ -457,6 +461,12 @@ private:
 	TArray<FMassQueryRequirementIndicesMapping> ArchetypeFragmentMapping;
 
 	bool bAllowParallelExecution = false;
+
+	EMassExecutionContextType ExpectedContextType = EMassExecutionContextType::Local;
+
+#if WITH_MASSENTITY_DEBUG
+	uint8 bRegistered : 1;
+#endif // WITH_MASSENTITY_DEBUG
 };
 
 template<>
@@ -484,4 +494,4 @@ FORCEINLINE FMassEntityQuery& FMassEntityQuery::AddTagRequirements<EMassFragment
 	// force recaching the next time this query is used or the following CacheArchetypes call.
 	DirtyCachedData();
 	return *this;
-}
\ No newline at end of file
+}
diff --git a/Engine/Plugins/Runtime/MassEntity/Source/MassEntity/Public/MassEntityTypes.h b/Engine/Plugins/Runtime/MassEntity/Source/MassEntity/Public/MassEntityTypes.h
--- a/Engine/Plugins/Runtime/MassEntity/Source/MassEntity/Public/MassEntityTypes.h
+++ b/Engine/Plugins/Runtime/MassEntity/Source/MassEntity/Public/MassEntityTypes.h
@@ -340,6 +340,13 @@ enum class EMassObservedOperation : uint8
 	MAX
 };
 
+enum class EMassExecutionContextType : uint8
+{
+	Local,
+	Processor,
+	MAX
+};
+
 /** 
  * Note that this is a view and is valid only as long as the source data is valid. Used when flushing mass commands to
  * wrap different kinds of data into a uniform package so that it can be passed over to a common interface.
diff --git a/Engine/Plugins/Runtime/MassEntity/Source/MassEntity/Public/MassExecutionContext.h b/Engine/Plugins/Runtime/MassEntity/Source/MassEntity/Public/MassExecutionContext.h
--- a/Engine/Plugins/Runtime/MassEntity/Source/MassEntity/Public/MassExecutionContext.h
+++ b/Engine/Plugins/Runtime/MassEntity/Source/MassEntity/Public/MassExecutionContext.h
@@ -70,6 +70,8 @@ private:
 	TArrayView<FConstSharedFragmentView> GetMutableConstSharedRequirements() { return ConstSharedFragmentViews; }
 	TArrayView<FSharedFragmentView> GetMutableSharedRequirements() { return SharedFragmentViews; }
 
+	EMassExecutionContextType ExecutionType = EMassExecutionContextType::Local;
+
 	friend FMassArchetypeData;
 	friend FMassEntityQuery;
 
@@ -98,6 +100,7 @@ public:
 	void SetEntityCollection(FMassArchetypeEntityCollection&& InEntityCollection);
 	void ClearEntityCollection() { EntityCollection.Reset(); }
 	void SetAuxData(const FInstancedStruct& InAuxData) { AuxData = InAuxData; }
+	void SetExecutionType(EMassExecutionContextType InExecutionType) { check(InExecutionType != EMassExecutionContextType::MAX); ExecutionType = InExecutionType; }
 
 	float GetDeltaTimeSeconds() const
 	{
diff --git a/Engine/Plugins/Runtime/MassEntity/Source/MassEntity/Public/MassProcessor.h b/Engine/Plugins/Runtime/MassEntity/Source/MassEntity/Public/MassProcessor.h
--- a/Engine/Plugins/Runtime/MassEntity/Source/MassEntity/Public/MassProcessor.h
+++ b/Engine/Plugins/Runtime/MassEntity/Source/MassEntity/Public/MassProcessor.h
@@ -70,6 +70,10 @@ public:
 
 	TConstArrayView<const int32> GetPrerequisiteIndices() const { return DependencyIndices; }
 
+	/** Adds Query to RegisteredQueries list. Query is required to be a member variable of this processor. Not meeting
+	 *  this requirement will cause check failure and the query won't be registered. */
+	void RegisterQuery(FMassEntityQuery& Query);
+
 	bool ShouldAutoAddToGlobalList() const { return bAutoRegisterWithProcessingPhases; }
 #if WITH_EDITOR
 	bool ShouldShowUpInSettings() const { return ShouldAutoAddToGlobalList() || bCanShowUpInSettings; }
@@ -127,6 +131,12 @@ protected:
 	friend class UMassCompositeProcessor;
 	TArray<int32> DependencyIndices;
 	TArray<int32> TransientDependencyIndices;
+
+private:
+	/** Stores processor's queries registered via RegisterQuery. 
+	 *  @note that it's safe to store pointers here since RegisterQuery does verify that a given registered query is 
+	 *  a member variable of a given processor */
+	TArray<FMassEntityQuery*> OwnedQueries;
 };
 
 
diff --git a/Engine/Plugins/Runtime/MassGameplay/Source/MassActors/Private/Translators/MassCapsuleComponentTranslators.cpp b/Engine/Plugins/Runtime/MassGameplay/Source/MassActors/Private/Translators/MassCapsuleComponentTranslators.cpp
--- a/Engine/Plugins/Runtime/MassGameplay/Source/MassActors/Private/Translators/MassCapsuleComponentTranslators.cpp
+++ b/Engine/Plugins/Runtime/MassGameplay/Source/MassActors/Private/Translators/MassCapsuleComponentTranslators.cpp
@@ -10,6 +10,7 @@
 // UMassCapsuleTransformToMassTranslator
 //----------------------------------------------------------------------//
 UMassCapsuleTransformToMassTranslator::UMassCapsuleTransformToMassTranslator()
+	: EntityQuery(*this)
 {
 	ExecutionFlags = (int32)EProcessorExecutionFlags::All;
 	RequiredTags.Add<FMassCapsuleTransformCopyToMassTag>();
@@ -43,6 +44,7 @@ void UMassCapsuleTransformToMassTranslator::Execute(UMassEntitySubsystem& Entity
 // UMassTransformToActorCapsuleTranslator
 //----------------------------------------------------------------------//
 UMassTransformToActorCapsuleTranslator::UMassTransformToActorCapsuleTranslator()
+	: EntityQuery(*this)
 {
 	ExecutionFlags = (int32)EProcessorExecutionFlags::All;
 	RequiredTags.Add<FMassCapsuleTransformCopyToActorTag>();
diff --git a/Engine/Plugins/Runtime/MassGameplay/Source/MassActors/Private/Translators/MassCharacterMovementTranslators.cpp b/Engine/Plugins/Runtime/MassGameplay/Source/MassActors/Private/Translators/MassCharacterMovementTranslators.cpp
--- a/Engine/Plugins/Runtime/MassGameplay/Source/MassActors/Private/Translators/MassCharacterMovementTranslators.cpp
+++ b/Engine/Plugins/Runtime/MassGameplay/Source/MassActors/Private/Translators/MassCharacterMovementTranslators.cpp
@@ -13,6 +13,7 @@
 //  UMassCharacterMovementToMassTranslator
 //----------------------------------------------------------------------//
 UMassCharacterMovementToMassTranslator::UMassCharacterMovementToMassTranslator()
+	: EntityQuery(*this)
 {
 	ExecutionFlags = (int32)EProcessorExecutionFlags::All;
 	ExecutionOrder.ExecuteInGroup = UE::Mass::ProcessorGroupNames::SyncWorldToMass;
@@ -53,6 +54,7 @@ void UMassCharacterMovementToMassTranslator::Execute(UMassEntitySubsystem& Entit
 //  UMassCharacterMovementToActorTranslator
 //----------------------------------------------------------------------//
 UMassCharacterMovementToActorTranslator::UMassCharacterMovementToActorTranslator()
+	: EntityQuery(*this)
 {
 	ExecutionFlags = (int32)EProcessorExecutionFlags::All;
 	ExecutionOrder.ExecuteInGroup = UE::Mass::ProcessorGroupNames::UpdateWorldFromMass;
@@ -90,6 +92,7 @@ void UMassCharacterMovementToActorTranslator::Execute(UMassEntitySubsystem& Enti
 //  UMassCharacterOrientationToMassTranslator
 //----------------------------------------------------------------------//
 UMassCharacterOrientationToMassTranslator::UMassCharacterOrientationToMassTranslator()
+	: EntityQuery(*this)
 {
 	ExecutionFlags = (int32)EProcessorExecutionFlags::All;
 	ExecutionOrder.ExecuteInGroup = UE::Mass::ProcessorGroupNames::SyncWorldToMass;
@@ -129,6 +132,7 @@ void UMassCharacterOrientationToMassTranslator::Execute(UMassEntitySubsystem& En
 //  UMassCharacterOrientationToActorTranslator
 //----------------------------------------------------------------------//
 UMassCharacterOrientationToActorTranslator::UMassCharacterOrientationToActorTranslator()
+	: EntityQuery(*this)
 {
 	ExecutionFlags = (int32)EProcessorExecutionFlags::All;
 	ExecutionOrder.ExecuteInGroup = UE::Mass::ProcessorGroupNames::UpdateWorldFromMass;
diff --git a/Engine/Plugins/Runtime/MassGameplay/Source/MassActors/Private/Translators/MassSceneComponentLocationTranslator.cpp b/Engine/Plugins/Runtime/MassGameplay/Source/MassActors/Private/Translators/MassSceneComponentLocationTranslator.cpp
--- a/Engine/Plugins/Runtime/MassGameplay/Source/MassActors/Private/Translators/MassSceneComponentLocationTranslator.cpp
+++ b/Engine/Plugins/Runtime/MassGameplay/Source/MassActors/Private/Translators/MassSceneComponentLocationTranslator.cpp
@@ -10,6 +10,7 @@
 //  UMassSceneComponentLocationToMassTranslator
 //----------------------------------------------------------------------//
 UMassSceneComponentLocationToMassTranslator::UMassSceneComponentLocationToMassTranslator()
+	: EntityQuery(*this)
 {
 	ExecutionFlags = (int32)EProcessorExecutionFlags::All;
 	ExecutionOrder.ExecuteInGroup = UE::Mass::ProcessorGroupNames::SyncWorldToMass;
@@ -45,6 +46,7 @@ void UMassSceneComponentLocationToMassTranslator::Execute(UMassEntitySubsystem&
 //  UMassSceneComponentLocationToActorTranslator
 //----------------------------------------------------------------------//
 UMassSceneComponentLocationToActorTranslator::UMassSceneComponentLocationToActorTranslator()
+	: EntityQuery(*this)
 {
 	ExecutionFlags = (int32)EProcessorExecutionFlags::All;
 	ExecutionOrder.ExecuteInGroup = UE::Mass::ProcessorGroupNames::UpdateWorldFromMass;
diff --git a/Engine/Plugins/Runtime/MassGameplay/Source/MassActors/Private/Translators/MassTranslators_BehaviorTree.cpp b/Engine/Plugins/Runtime/MassGameplay/Source/MassActors/Private/Translators/MassTranslators_BehaviorTree.cpp
--- a/Engine/Plugins/Runtime/MassGameplay/Source/MassActors/Private/Translators/MassTranslators_BehaviorTree.cpp
+++ b/Engine/Plugins/Runtime/MassGameplay/Source/MassActors/Private/Translators/MassTranslators_BehaviorTree.cpp
@@ -5,6 +5,7 @@
 #include "MassEntitySubsystem.h"
 
 UMassTranslator_BehaviorTree::UMassTranslator_BehaviorTree()
+	: EntityQuery(*this)
 {
 	bAutoRegisterWithProcessingPhases = false;
 }
diff --git a/Engine/Plugins/Runtime/MassGameplay/Source/MassGameplayDebug/Private/AssignDebugVisProcessor.cpp b/Engine/Plugins/Runtime/MassGameplay/Source/MassGameplayDebug/Private/AssignDebugVisProcessor.cpp
--- a/Engine/Plugins/Runtime/MassGameplay/Source/MassGameplayDebug/Private/AssignDebugVisProcessor.cpp
+++ b/Engine/Plugins/Runtime/MassGameplay/Source/MassGameplayDebug/Private/AssignDebugVisProcessor.cpp
@@ -10,6 +10,7 @@
 // UAssignDebugVisProcessor
 //----------------------------------------------------------------------//
 UAssignDebugVisProcessor::UAssignDebugVisProcessor()
+	: EntityQuery(*this)
 {
 	bAutoRegisterWithProcessingPhases = false;
 	ObservedType = FSimDebugVisFragment::StaticStruct();
diff --git a/Engine/Plugins/Runtime/MassGameplay/Source/MassGameplayDebug/Private/DebugVisLocationProcessor.cpp b/Engine/Plugins/Runtime/MassGameplay/Source/MassGameplayDebug/Private/DebugVisLocationProcessor.cpp
--- a/Engine/Plugins/Runtime/MassGameplay/Source/MassGameplayDebug/Private/DebugVisLocationProcessor.cpp
+++ b/Engine/Plugins/Runtime/MassGameplay/Source/MassGameplayDebug/Private/DebugVisLocationProcessor.cpp
@@ -12,6 +12,7 @@
 // UDebugVisLocationProcessor
 //----------------------------------------------------------------------//
 UDebugVisLocationProcessor::UDebugVisLocationProcessor()
+	: EntityQuery(*this)
 {
 	ExecutionOrder.ExecuteAfter.Add(UE::Mass::ProcessorGroupNames::SyncWorldToMass);
 }
@@ -71,6 +72,7 @@ void UDebugVisLocationProcessor::Execute(UMassEntitySubsystem& EntitySubsystem,
 //  UMassProcessor_UpdateDebugVis
 //----------------------------------------------------------------------//
 UMassProcessor_UpdateDebugVis::UMassProcessor_UpdateDebugVis()
+	: EntityQuery(*this)
 {
 	ExecutionOrder.ExecuteAfter.Add(UE::Mass::ProcessorGroupNames::UpdateWorldFromMass);
 }
diff --git a/Engine/Plugins/Runtime/MassGameplay/Source/MassLOD/Private/MassLODCollectorProcessor.cpp b/Engine/Plugins/Runtime/MassGameplay/Source/MassLOD/Private/MassLODCollectorProcessor.cpp
--- a/Engine/Plugins/Runtime/MassGameplay/Source/MassLOD/Private/MassLODCollectorProcessor.cpp
+++ b/Engine/Plugins/Runtime/MassGameplay/Source/MassLOD/Private/MassLODCollectorProcessor.cpp
@@ -33,18 +33,22 @@ void UMassLODCollectorProcessor::ConfigureQueries()
 	EntityQuery_VisibleRangeAndOnLOD = BaseQuery;
 	EntityQuery_VisibleRangeAndOnLOD.AddTagRequirement<FMassVisibilityCulledByDistanceTag>(EMassFragmentPresence::None);
 	EntityQuery_VisibleRangeAndOnLOD.AddTagRequirement<FMassOffLODTag>(EMassFragmentPresence::None);
+	EntityQuery_VisibleRangeAndOnLOD.RegisterWithProcessor(*this);
 
 	EntityQuery_VisibleRangeOnly = BaseQuery;
 	EntityQuery_VisibleRangeOnly.AddTagRequirement<FMassVisibilityCulledByDistanceTag>(EMassFragmentPresence::None);
 	EntityQuery_VisibleRangeOnly.AddTagRequirement<FMassOffLODTag>(EMassFragmentPresence::All);
+	EntityQuery_VisibleRangeOnly.RegisterWithProcessor(*this);
 
 	EntityQuery_OnLODOnly = BaseQuery;
 	EntityQuery_OnLODOnly.AddTagRequirement<FMassVisibilityCulledByDistanceTag>(EMassFragmentPresence::All);
 	EntityQuery_OnLODOnly.AddTagRequirement<FMassOffLODTag>(EMassFragmentPresence::None);
+	EntityQuery_OnLODOnly.RegisterWithProcessor(*this);
 
 	EntityQuery_NotVisibleRangeAndOffLOD = BaseQuery;
 	EntityQuery_NotVisibleRangeAndOffLOD.AddTagRequirement<FMassVisibilityCulledByDistanceTag>(EMassFragmentPresence::All);
 	EntityQuery_NotVisibleRangeAndOffLOD.AddTagRequirement<FMassOffLODTag>(EMassFragmentPresence::All);
+	EntityQuery_NotVisibleRangeAndOffLOD.RegisterWithProcessor(*this);
 }
 
 template <bool bLocalViewersOnly>
diff --git a/Engine/Plugins/Runtime/MassGameplay/Source/MassLOD/Private/MassSimulationLOD.cpp b/Engine/Plugins/Runtime/MassGameplay/Source/MassLOD/Private/MassSimulationLOD.cpp
--- a/Engine/Plugins/Runtime/MassGameplay/Source/MassLOD/Private/MassSimulationLOD.cpp
+++ b/Engine/Plugins/Runtime/MassGameplay/Source/MassLOD/Private/MassSimulationLOD.cpp
@@ -62,6 +62,11 @@ namespace UE::MassLOD
 } // UE::MassLOD
 
 UMassSimulationLODProcessor::UMassSimulationLODProcessor()
+	: EntityQuery(*this)
+	, EntityQueryCalculateLOD(*this)
+	, EntityQueryAdjustDistances(*this)
+	, EntityQueryVariableTick(*this)
+	, EntityQuerySetLODTag(*this)
 {
 	ExecutionFlags = (int32)EProcessorExecutionFlags::All;
 	ExecutionOrder.ExecuteInGroup = UE::Mass::ProcessorGroupNames::LOD;
diff --git a/Engine/Plugins/Runtime/MassGameplay/Source/MassMovement/Private/MassMovementProcessors.cpp b/Engine/Plugins/Runtime/MassGameplay/Source/MassMovement/Private/MassMovementProcessors.cpp
--- a/Engine/Plugins/Runtime/MassGameplay/Source/MassMovement/Private/MassMovementProcessors.cpp
+++ b/Engine/Plugins/Runtime/MassGameplay/Source/MassMovement/Private/MassMovementProcessors.cpp
@@ -1,4 +1,4 @@
-﻿// Copyright Epic Games, Inc. All Rights Reserved.
+// Copyright Epic Games, Inc. All Rights Reserved.
 
 #include "Movement/MassMovementProcessors.h"
 #include "MassCommonUtils.h"
@@ -16,6 +16,7 @@
 //----------------------------------------------------------------------//
 
 UMassApplyMovementProcessor::UMassApplyMovementProcessor()
+	: EntityQuery(*this)
 {
 	ExecutionFlags = (int32)EProcessorExecutionFlags::All;
 	ExecutionOrder.ExecuteInGroup = UE::Mass::ProcessorGroupNames::Movement;
diff --git a/Engine/Plugins/Runtime/MassGameplay/Source/MassMovement/Private/MassSimpleMovementTrait.cpp b/Engine/Plugins/Runtime/MassGameplay/Source/MassMovement/Private/MassSimpleMovementTrait.cpp
--- a/Engine/Plugins/Runtime/MassGameplay/Source/MassMovement/Private/MassSimpleMovementTrait.cpp
+++ b/Engine/Plugins/Runtime/MassGameplay/Source/MassMovement/Private/MassSimpleMovementTrait.cpp
@@ -22,6 +22,7 @@ void UMassSimpleMovementTrait::BuildTemplate(FMassEntityTemplateBuildContext& Bu
 //  UMassSimpleMovementProcessor
 //----------------------------------------------------------------------//
 UMassSimpleMovementProcessor::UMassSimpleMovementProcessor()
+	: EntityQuery(*this)
 {
 	bAutoRegisterWithProcessingPhases = true;
 	ExecutionFlags = (int32)EProcessorExecutionFlags::All;
diff --git a/Engine/Plugins/Runtime/MassGameplay/Source/MassMovement/Private/MassVelocityRandomizerTrait.cpp b/Engine/Plugins/Runtime/MassGameplay/Source/MassMovement/Private/MassVelocityRandomizerTrait.cpp
--- a/Engine/Plugins/Runtime/MassGameplay/Source/MassMovement/Private/MassVelocityRandomizerTrait.cpp
+++ b/Engine/Plugins/Runtime/MassGameplay/Source/MassMovement/Private/MassVelocityRandomizerTrait.cpp
@@ -23,6 +23,7 @@ void UMassVelocityRandomizerTrait::BuildTemplate(FMassEntityTemplateBuildContext
 //  UMassRandomVelocityInitializer
 //----------------------------------------------------------------------//
 UMassRandomVelocityInitializer::UMassRandomVelocityInitializer()
+	: EntityQuery(*this)
 {
 	ObservedType = FMassVelocityFragment::StaticStruct();
 	Operation = EMassObservedOperation::Add;
diff --git a/Engine/Plugins/Runtime/MassGameplay/Source/MassReplication/Private/MassReplicationFragments.cpp b/Engine/Plugins/Runtime/MassGameplay/Source/MassReplication/Private/MassReplicationFragments.cpp
--- a/Engine/Plugins/Runtime/MassGameplay/Source/MassReplication/Private/MassReplicationFragments.cpp
+++ b/Engine/Plugins/Runtime/MassGameplay/Source/MassReplication/Private/MassReplicationFragments.cpp
@@ -9,6 +9,7 @@
 // UMassNetworkIDFragmentInitializer 
 //----------------------------------------------------------------------//
 UMassNetworkIDFragmentInitializer::UMassNetworkIDFragmentInitializer()
+	: EntityQuery(*this)
 {
 	ExecutionFlags = (int32)(EProcessorExecutionFlags::Standalone | EProcessorExecutionFlags::Server);
 	ObservedType = FMassNetworkIDFragment::StaticStruct();
diff --git a/Engine/Plugins/Runtime/MassGameplay/Source/MassReplication/Private/MassReplicationGridProcessor.cpp b/Engine/Plugins/Runtime/MassGameplay/Source/MassReplication/Private/MassReplicationGridProcessor.cpp
--- a/Engine/Plugins/Runtime/MassGameplay/Source/MassReplication/Private/MassReplicationGridProcessor.cpp
+++ b/Engine/Plugins/Runtime/MassGameplay/Source/MassReplication/Private/MassReplicationGridProcessor.cpp
@@ -11,6 +11,9 @@
 //  UMassReplicationGridProcessor
 //----------------------------------------------------------------------//
 UMassReplicationGridProcessor::UMassReplicationGridProcessor()
+	: AddToGridEntityQuery(*this)
+	, UpdateGridEntityQuery(*this)
+	, RemoveFromGridEntityQuery(*this)
 {
 #if !UE_ALLOW_DEBUG_REPLICATION_BUBBLES_STANDALONE
 	ExecutionFlags = int32(EProcessorExecutionFlags::Server);
@@ -122,6 +125,7 @@ void UMassReplicationGridProcessor::Execute(UMassEntitySubsystem& EntitySubsyste
 //  UMassReplicationGridRemoverProcessor
 //----------------------------------------------------------------------//
 UMassReplicationGridRemoverProcessor::UMassReplicationGridRemoverProcessor()
+	: EntityQuery(*this)
 {
 	ObservedType = FMassReplicationGridCellLocationFragment::StaticStruct();
 	Operation = EMassObservedOperation::Remove;
diff --git a/Engine/Plugins/Runtime/MassGameplay/Source/MassReplication/Private/MassReplicationProcessor.cpp b/Engine/Plugins/Runtime/MassGameplay/Source/MassReplication/Private/MassReplicationProcessor.cpp
--- a/Engine/Plugins/Runtime/MassGameplay/Source/MassReplication/Private/MassReplicationProcessor.cpp
+++ b/Engine/Plugins/Runtime/MassGameplay/Source/MassReplication/Private/MassReplicationProcessor.cpp
@@ -15,6 +15,11 @@ namespace UE::Mass::Replication
 //  UMassReplicationProcessor
 //----------------------------------------------------------------------//
 UMassReplicationProcessor::UMassReplicationProcessor()
+	: SyncClientData(*this)
+	, CollectViewerInfoQuery(*this)
+	, CalculateLODQuery(*this)
+	, AdjustLODDistancesQuery(*this)
+	, EntityQuery(*this)
 {
 #if !UE_ALLOW_DEBUG_REPLICATION_BUBBLES_STANDALONE
 	ExecutionFlags = int32(EProcessorExecutionFlags::Server);
diff --git a/Engine/Plugins/Runtime/MassGameplay/Source/MassRepresentation/Private/MassRepresentationProcessor.cpp b/Engine/Plugins/Runtime/MassGameplay/Source/MassRepresentation/Private/MassRepresentationProcessor.cpp
--- a/Engine/Plugins/Runtime/MassGameplay/Source/MassRepresentation/Private/MassRepresentationProcessor.cpp
+++ b/Engine/Plugins/Runtime/MassGameplay/Source/MassRepresentation/Private/MassRepresentationProcessor.cpp
@@ -21,6 +21,7 @@ namespace UE::MassRepresentation
 // UMassRepresentationProcessor(); 
 //----------------------------------------------------------------------//
 UMassRepresentationProcessor::UMassRepresentationProcessor()
+	: EntityQuery(*this)
 {
 	bAutoRegisterWithProcessingPhases = false;
 
@@ -367,6 +368,7 @@ void UMassVisualizationProcessor::Execute(UMassEntitySubsystem& EntitySubsystem,
 // UMassRepresentationFragmentDestructor 
 //----------------------------------------------------------------------//
 UMassRepresentationFragmentDestructor::UMassRepresentationFragmentDestructor()
+	: EntityQuery(*this)
 {
 	ObservedType = FMassRepresentationFragment::StaticStruct();
 	Operation = EMassObservedOperation::Remove;
diff --git a/Engine/Plugins/Runtime/MassGameplay/Source/MassRepresentation/Private/MassUpdateISMProcessor.cpp b/Engine/Plugins/Runtime/MassGameplay/Source/MassRepresentation/Private/MassUpdateISMProcessor.cpp
--- a/Engine/Plugins/Runtime/MassGameplay/Source/MassRepresentation/Private/MassUpdateISMProcessor.cpp
+++ b/Engine/Plugins/Runtime/MassGameplay/Source/MassRepresentation/Private/MassUpdateISMProcessor.cpp
@@ -10,6 +10,7 @@
 #include "Engine/World.h"
 
 UMassUpdateISMProcessor::UMassUpdateISMProcessor()
+	: EntityQuery(*this)
 {
 	ExecutionFlags = (int32)(EProcessorExecutionFlags::Client | EProcessorExecutionFlags::Standalone);
 
diff --git a/Engine/Plugins/Runtime/MassGameplay/Source/MassRepresentation/Private/MassVisualizationLODProcessor.cpp b/Engine/Plugins/Runtime/MassGameplay/Source/MassRepresentation/Private/MassVisualizationLODProcessor.cpp
--- a/Engine/Plugins/Runtime/MassGameplay/Source/MassRepresentation/Private/MassVisualizationLODProcessor.cpp
+++ b/Engine/Plugins/Runtime/MassGameplay/Source/MassRepresentation/Private/MassVisualizationLODProcessor.cpp
@@ -28,20 +28,24 @@ void UMassVisualizationLODProcessor::ConfigureQueries()
 
 	CloseEntityQuery = BaseQuery;
 	CloseEntityQuery.AddTagRequirement<FMassVisibilityCulledByDistanceTag>(EMassFragmentPresence::None);
-	
+	CloseEntityQuery.RegisterWithProcessor(*this);
+
 	CloseEntityAdjustDistanceQuery = CloseEntityQuery;
 	CloseEntityAdjustDistanceQuery.SetArchetypeFilter([](const FMassExecutionContext& Context)
 	{
 		const FMassVisualizationLODSharedFragment& LODSharedFragment = Context.GetSharedFragment<FMassVisualizationLODSharedFragment>();
 		return LODSharedFragment.bHasAdjustedDistancesFromCount;
 	});
+	CloseEntityAdjustDistanceQuery.RegisterWithProcessor(*this);
 
 	FarEntityQuery = BaseQuery;
 	FarEntityQuery.AddTagRequirement<FMassVisibilityCulledByDistanceTag>(EMassFragmentPresence::All);
 	FarEntityQuery.AddChunkRequirement<FMassVisualizationChunkFragment>(EMassFragmentAccess::ReadOnly);
 	FarEntityQuery.SetChunkFilter(&FMassVisualizationChunkFragment::ShouldUpdateVisualizationForChunk);
+	FarEntityQuery.RegisterWithProcessor(*this);
 
 	DebugEntityQuery = BaseQuery;
+	DebugEntityQuery.RegisterWithProcessor(*this);
 }
 
 void UMassVisualizationLODProcessor::Execute(UMassEntitySubsystem& EntitySubsystem, FMassExecutionContext& Context)
diff --git a/Engine/Plugins/Runtime/MassGameplay/Source/MassSignals/Private/MassSignalProcessorBase.cpp b/Engine/Plugins/Runtime/MassGameplay/Source/MassSignals/Private/MassSignalProcessorBase.cpp
--- a/Engine/Plugins/Runtime/MassGameplay/Source/MassSignals/Private/MassSignalProcessorBase.cpp
+++ b/Engine/Plugins/Runtime/MassGameplay/Source/MassSignals/Private/MassSignalProcessorBase.cpp
@@ -5,6 +5,12 @@
 #include "MassArchetypeTypes.h"
 #include "Engine/World.h"
 
+
+UMassSignalProcessorBase::UMassSignalProcessorBase(const FObjectInitializer& ObjectInitializer)
+	: EntityQuery(*this)
+{
+}
+
 void UMassSignalProcessorBase::Initialize(UObject& Owner)
 {
 	SignalSubsystem = UWorld::GetSubsystem<UMassSignalSubsystem>(Owner.GetWorld());
diff --git a/Engine/Plugins/Runtime/MassGameplay/Source/MassSignals/Public/MassSignalProcessorBase.h b/Engine/Plugins/Runtime/MassGameplay/Source/MassSignals/Public/MassSignalProcessorBase.h
--- a/Engine/Plugins/Runtime/MassGameplay/Source/MassSignals/Public/MassSignalProcessorBase.h
+++ b/Engine/Plugins/Runtime/MassGameplay/Source/MassSignals/Public/MassSignalProcessorBase.h
@@ -19,6 +19,9 @@ class MASSSIGNALS_API UMassSignalProcessorBase : public UMassProcessor
 {
 	GENERATED_BODY()
 
+public:
+	UMassSignalProcessorBase(const FObjectInitializer& ObjectInitializer = FObjectInitializer::Get());
+
 protected:
 	/**
 	 * Initialize the processor
diff --git a/Engine/Plugins/Runtime/MassGameplay/Source/MassSmartObjects/Private/MassSmartObjectProcessor.cpp b/Engine/Plugins/Runtime/MassGameplay/Source/MassSmartObjects/Private/MassSmartObjectProcessor.cpp
--- a/Engine/Plugins/Runtime/MassGameplay/Source/MassSmartObjects/Private/MassSmartObjectProcessor.cpp
+++ b/Engine/Plugins/Runtime/MassGameplay/Source/MassSmartObjects/Private/MassSmartObjectProcessor.cpp
@@ -41,6 +41,8 @@ void UMassSmartObjectCandidatesFinderProcessor::ConfigureQueries()
 }
 
 UMassSmartObjectCandidatesFinderProcessor::UMassSmartObjectCandidatesFinderProcessor()
+	: WorldRequestQuery(*this)
+	, LaneRequestQuery(*this)
 {
 	// 1. Frame T Behavior create a request(deferred entity creation)
 	// 2. Frame T+1: Processor execute the request might mark it as done(deferred add tag flushed at the end of the frame)
@@ -312,6 +314,7 @@ void UMassSmartObjectTimedBehaviorProcessor::ConfigureQueries()
 }
 
 UMassSmartObjectTimedBehaviorProcessor::UMassSmartObjectTimedBehaviorProcessor()
+	: EntityQuery(*this)
 {
 	ExecutionOrder.ExecuteInGroup = UE::Mass::ProcessorGroupNames::SyncWorldToMass;
 }
@@ -383,6 +386,7 @@ void UMassSmartObjectTimedBehaviorProcessor::Execute(UMassEntitySubsystem& Entit
 //  UMassSmartObjectUserFragmentDeinitializer
 //----------------------------------------------------------------------//
 UMassSmartObjectUserFragmentDeinitializer::UMassSmartObjectUserFragmentDeinitializer()
+	: EntityQuery(*this)
 {
 	ObservedType = FMassSmartObjectUserFragment::StaticStruct();
 	Operation = EMassObservedOperation::Remove;
diff --git a/Engine/Plugins/Runtime/MassGameplay/Source/MassSpawner/Private/MassSpawnLocationProcessor.cpp b/Engine/Plugins/Runtime/MassGameplay/Source/MassSpawner/Private/MassSpawnLocationProcessor.cpp
--- a/Engine/Plugins/Runtime/MassGameplay/Source/MassSpawner/Private/MassSpawnLocationProcessor.cpp
+++ b/Engine/Plugins/Runtime/MassGameplay/Source/MassSpawner/Private/MassSpawnLocationProcessor.cpp
@@ -11,6 +11,7 @@
 // UMassSpawnLocationProcessor 
 //----------------------------------------------------------------------//
 UMassSpawnLocationProcessor::UMassSpawnLocationProcessor()
+	: EntityQuery(*this)
 {
 	bAutoRegisterWithProcessingPhases = false;
 }
