Fixed MassZoneGraphAnnotationProcessor's access to UZoneGraphAnnotationSubsystem

diff --git a/Engine/Plugins/AI/MassAI/Source/MassAIBehavior/Private/MassZoneGraphAnnotationProcessors.cpp b/Engine/Plugins/AI/MassAI/Source/MassAIBehavior/Private/MassZoneGraphAnnotationProcessors.cpp
--- a/Engine/Plugins/AI/MassAI/Source/MassAIBehavior/Private/MassZoneGraphAnnotationProcessors.cpp
+++ b/Engine/Plugins/AI/MassAI/Source/MassAIBehavior/Private/MassZoneGraphAnnotationProcessors.cpp
@@ -5,6 +5,7 @@
 #include "MassZoneGraphNavigationFragments.h"
 #include "MassZoneGraphAnnotationFragments.h"
 #include "MassZoneGraphAnnotationTypes.h"
+#include "MassGameplayExternalTraits.h"
 #include "ZoneGraphTypes.h"
 #include "ZoneGraphAnnotationSubsystem.h"
 #include "MassSignalSubsystem.h"
@@ -21,22 +22,19 @@ UMassZoneGraphAnnotationTagsInitializer::UMassZoneGraphAnnotationTagsInitializer
 	Operation = EMassObservedOperation::Add;
 }
 
-void UMassZoneGraphAnnotationTagsInitializer::Initialize(UObject& Owner)
-{
-	Super::Initialize(Owner);
-	ZoneGraphAnnotationSubsystem = UWorld::GetSubsystem<UZoneGraphAnnotationSubsystem>(Owner.GetWorld());
-}
-
 void UMassZoneGraphAnnotationTagsInitializer::ConfigureQueries()
 {
 	EntityQuery.AddRequirement<FMassZoneGraphAnnotationFragment>(EMassFragmentAccess::ReadWrite);
 	EntityQuery.AddRequirement<FMassZoneGraphLaneLocationFragment>(EMassFragmentAccess::ReadOnly);
+	EntityQuery.AddSubsystemRequirement<UZoneGraphAnnotationSubsystem>(EMassFragmentAccess::ReadOnly);
 }
 
 void UMassZoneGraphAnnotationTagsInitializer::Execute(UMassEntitySubsystem& EntitySubsystem, FMassExecutionContext& Context)
 {
-	EntityQuery.ForEachEntityChunk(EntitySubsystem, Context, [this](FMassExecutionContext& Context)
+	EntityQuery.ForEachEntityChunk(EntitySubsystem, Context, [this, World = EntitySubsystem.GetWorld()](FMassExecutionContext& Context)
 	{
+		const UZoneGraphAnnotationSubsystem& ZoneGraphAnnotationSubsystem = Context.GetSubsystemChecked<UZoneGraphAnnotationSubsystem>(World);
+
 		const int32 NumEntities = Context.GetNumEntities();
 		const TArrayView<FMassZoneGraphAnnotationFragment> AnnotationTagsList = Context.GetMutableFragmentView<FMassZoneGraphAnnotationFragment>();
 		const TConstArrayView<FMassZoneGraphLaneLocationFragment> LaneLocationList = Context.GetFragmentView<FMassZoneGraphLaneLocationFragment>();
@@ -52,7 +50,7 @@ void UMassZoneGraphAnnotationTagsInitializer::Execute(UMassEntitySubsystem& Enti
 			}
 			else
 			{
-				AnnotationTags.Tags = ZoneGraphAnnotationSubsystem->GetAnnotationTags(LaneLocation.LaneHandle);
+				AnnotationTags.Tags = ZoneGraphAnnotationSubsystem.GetAnnotationTags(LaneLocation.LaneHandle);
 			}
 		}
 	});
@@ -71,7 +69,6 @@ UMassZoneGraphAnnotationTagUpdateProcessor::UMassZoneGraphAnnotationTagUpdatePro
 void UMassZoneGraphAnnotationTagUpdateProcessor::Initialize(UObject& Owner)
 {
 	Super::Initialize(Owner);
-	ZoneGraphAnnotationSubsystem = UWorld::GetSubsystem<UZoneGraphAnnotationSubsystem>(Owner.GetWorld());
 
 	SubscribeToSignal(UE::Mass::Signals::CurrentLaneChanged);
 }
@@ -83,6 +80,7 @@ void UMassZoneGraphAnnotationTagUpdateProcessor::ConfigureQueries()
 	EntityQuery.AddRequirement<FMassZoneGraphLaneLocationFragment>(EMassFragmentAccess::ReadOnly);
 	EntityQuery.AddChunkRequirement<FMassZoneGraphAnnotationVariableTickChunkFragment>(EMassFragmentAccess::ReadOnly);
 	EntityQuery.AddChunkRequirement<FMassSimulationVariableTickChunkFragment>(EMassFragmentAccess::ReadOnly, EMassFragmentPresence::Optional);
+	EntityQuery.AddSubsystemRequirement<UZoneGraphAnnotationSubsystem>(EMassFragmentAccess::ReadWrite);
 }
 
 void UMassZoneGraphAnnotationTagUpdateProcessor::Execute(UMassEntitySubsystem& EntitySubsystem, FMassExecutionContext& Context)
@@ -92,13 +90,15 @@ void UMassZoneGraphAnnotationTagUpdateProcessor::Execute(UMassEntitySubsystem& E
 	// Calling super will update the signals, and call SignalEntities() below.
 	Super::Execute(EntitySubsystem, Context);
 
-	EntityQuery.ForEachEntityChunk(EntitySubsystem, Context, [this](FMassExecutionContext& Context)
+	EntityQuery.ForEachEntityChunk(EntitySubsystem, Context, [this, World = EntitySubsystem.GetWorld()](FMassExecutionContext& Context)
 	{
 		// Periodically update tags.
 		if (!FMassZoneGraphAnnotationVariableTickChunkFragment::UpdateChunk(Context))
 		{
 			return;
 		}
+		
+		UZoneGraphAnnotationSubsystem& ZoneGraphAnnotationSubsystem = Context.GetMutableSubsystemChecked<UZoneGraphAnnotationSubsystem>(World);
 
 		const int32 NumEntities = Context.GetNumEntities();
 		const TArrayView<FMassZoneGraphAnnotationFragment> AnnotationTagsList = Context.GetMutableFragmentView<FMassZoneGraphAnnotationFragment>();
@@ -109,7 +109,7 @@ void UMassZoneGraphAnnotationTagUpdateProcessor::Execute(UMassEntitySubsystem& E
 			FMassZoneGraphAnnotationFragment& AnnotationTags = AnnotationTagsList[EntityIndex];
 			const FMassZoneGraphLaneLocationFragment& LaneLocation = LaneLocationList[EntityIndex];
 
-			UpdateAnnotationTags(AnnotationTags, LaneLocation, Context.GetEntity(EntityIndex));
+			UpdateAnnotationTags(ZoneGraphAnnotationSubsystem, AnnotationTags, LaneLocation, Context.GetEntity(EntityIndex));
 		}
 	});
 
@@ -119,13 +119,13 @@ void UMassZoneGraphAnnotationTagUpdateProcessor::Execute(UMassEntitySubsystem& E
 	}
 }
 
-void UMassZoneGraphAnnotationTagUpdateProcessor::UpdateAnnotationTags(FMassZoneGraphAnnotationFragment& AnnotationTags, const FMassZoneGraphLaneLocationFragment& LaneLocation, const FMassEntityHandle Entity)
+void UMassZoneGraphAnnotationTagUpdateProcessor::UpdateAnnotationTags(UZoneGraphAnnotationSubsystem& ZoneGraphAnnotationSubsystem, FMassZoneGraphAnnotationFragment& AnnotationTags, const FMassZoneGraphLaneLocationFragment& LaneLocation, const FMassEntityHandle Entity)
 {
 	const FZoneGraphTagMask OldTags = AnnotationTags.Tags;
 
 	if (LaneLocation.LaneHandle.IsValid())
 	{
-		AnnotationTags.Tags = ZoneGraphAnnotationSubsystem->GetAnnotationTags(LaneLocation.LaneHandle);
+		AnnotationTags.Tags = ZoneGraphAnnotationSubsystem.GetAnnotationTags(LaneLocation.LaneHandle);
 	}
 	else
 	{
@@ -140,13 +140,10 @@ void UMassZoneGraphAnnotationTagUpdateProcessor::UpdateAnnotationTags(FMassZoneG
 
 void UMassZoneGraphAnnotationTagUpdateProcessor::SignalEntities(UMassEntitySubsystem& EntitySubsystem, FMassExecutionContext& Context, FMassSignalNameLookup& EntitySignals)
 {
-	if (!ZoneGraphAnnotationSubsystem)
+	EntityQuery.ForEachEntityChunk(EntitySubsystem, Context, [this, World = EntitySubsystem.GetWorld()](FMassExecutionContext& Context)
 	{
-		return;
-	}
+		UZoneGraphAnnotationSubsystem& ZoneGraphAnnotationSubsystem = Context.GetMutableSubsystemChecked<UZoneGraphAnnotationSubsystem>(World);
 
-	EntityQuery.ForEachEntityChunk(EntitySubsystem, Context, [this](FMassExecutionContext& Context)
-	{
 		const int32 NumEntities = Context.GetNumEntities();
 		const TConstArrayView<FMassZoneGraphLaneLocationFragment> LaneLocationList = Context.GetFragmentView<FMassZoneGraphLaneLocationFragment>();
 		const TArrayView<FMassZoneGraphAnnotationFragment> AnnotationTagsList = Context.GetMutableFragmentView<FMassZoneGraphAnnotationFragment>();
@@ -156,7 +153,7 @@ void UMassZoneGraphAnnotationTagUpdateProcessor::SignalEntities(UMassEntitySubsy
 			FMassZoneGraphAnnotationFragment& AnnotationTags = AnnotationTagsList[EntityIndex];
 			const FMassZoneGraphLaneLocationFragment& LaneLocation = LaneLocationList[EntityIndex];
 
-			UpdateAnnotationTags(AnnotationTags, LaneLocation, Context.GetEntity(EntityIndex));
+			UpdateAnnotationTags(ZoneGraphAnnotationSubsystem, AnnotationTags, LaneLocation, Context.GetEntity(EntityIndex));
 		}
 	});
 }
diff --git a/Engine/Plugins/AI/MassAI/Source/MassAIBehavior/Public/MassZoneGraphAnnotationProcessors.h b/Engine/Plugins/AI/MassAI/Source/MassAIBehavior/Public/MassZoneGraphAnnotationProcessors.h
--- a/Engine/Plugins/AI/MassAI/Source/MassAIBehavior/Public/MassZoneGraphAnnotationProcessors.h
+++ b/Engine/Plugins/AI/MassAI/Source/MassAIBehavior/Public/MassZoneGraphAnnotationProcessors.h
@@ -24,14 +24,10 @@ public:
 	UMassZoneGraphAnnotationTagsInitializer();
 
 protected:
-	virtual void Initialize(UObject& Owner) override;
 	virtual void ConfigureQueries() override;
 	virtual void Execute(UMassEntitySubsystem& EntitySubsystem, FMassExecutionContext& Context) override;
 
 	FMassEntityQuery EntityQuery;
-
-	UPROPERTY(Transient)
-	UZoneGraphAnnotationSubsystem* ZoneGraphAnnotationSubsystem = nullptr;
 };
 
 /** 
@@ -50,13 +46,10 @@ protected:
 	virtual void ConfigureQueries() override;
 	virtual void Execute(UMassEntitySubsystem& EntitySubsystem, FMassExecutionContext& Context) override;
 
-	void UpdateAnnotationTags(FMassZoneGraphAnnotationFragment& AnnotationTags, const FMassZoneGraphLaneLocationFragment& LaneLocation, FMassEntityHandle Entity);
+	void UpdateAnnotationTags(UZoneGraphAnnotationSubsystem& ZoneGraphAnnotationSubsystem, FMassZoneGraphAnnotationFragment& AnnotationTags, const FMassZoneGraphLaneLocationFragment& LaneLocation, FMassEntityHandle Entity);
 
 	virtual void SignalEntities(UMassEntitySubsystem& EntitySubsystem, FMassExecutionContext& Context, FMassSignalNameLookup& EntitySignals) override;
 
-	UPROPERTY(Transient)
-	UZoneGraphAnnotationSubsystem* ZoneGraphAnnotationSubsystem = nullptr;
-
 	// Frame buffer, it gets reset every frame.
 	TArray<FMassEntityHandle> TransientEntitiesToSignal;
 };
diff --git a/Engine/Plugins/Runtime/MassGameplay/Source/MassGameplayExternalTraits/Public/MassGameplayExternalTraits.h b/Engine/Plugins/Runtime/MassGameplay/Source/MassGameplayExternalTraits/Public/MassGameplayExternalTraits.h
--- a/Engine/Plugins/Runtime/MassGameplay/Source/MassGameplayExternalTraits/Public/MassGameplayExternalTraits.h
+++ b/Engine/Plugins/Runtime/MassGameplay/Source/MassGameplayExternalTraits/Public/MassGameplayExternalTraits.h
@@ -30,3 +30,13 @@ struct TMassExternalSubsystemTraits<UZoneGraphSubsystem> final
 		GameThreadOnly = false
 	};
 };
+
+class UZoneGraphAnnotationSubsystem;
+template<>
+struct TMassExternalSubsystemTraits<UZoneGraphAnnotationSubsystem> final
+{
+	enum
+	{
+		GameThreadOnly = false
+	};
+};
