StateTree struct reference
- added functionality to allow to use property binding to get pointer to another struct property
- converted mass smart object and movement tasks to use struct ref

diff --git a/Engine/Plugins/AI/MassAI/Source/MassAIBehavior/Private/Tasks/MassClaimSmartObjectTask.cpp b/Engine/Plugins/AI/MassAI/Source/MassAIBehavior/Private/Tasks/MassClaimSmartObjectTask.cpp
--- a/Engine/Plugins/AI/MassAI/Source/MassAIBehavior/Private/Tasks/MassClaimSmartObjectTask.cpp
+++ b/Engine/Plugins/AI/MassAI/Source/MassAIBehavior/Private/Tasks/MassClaimSmartObjectTask.cpp
@@ -23,9 +23,6 @@ bool FMassClaimSmartObjectTask::Link(FStateTreeLinker& Linker)
 	Linker.LinkExternalData(SmartObjectSubsystemHandle);
 	Linker.LinkExternalData(MassSignalSubsystemHandle);
 
-	Linker.LinkInstanceDataProperty(CandidateSlotsHandle, STATETREE_INSTANCEDATA_PROPERTY(FInstanceDataType, CandidateSlots));
-	Linker.LinkInstanceDataProperty(ClaimedSlotHandle, STATETREE_INSTANCEDATA_PROPERTY(FInstanceDataType, ClaimedSlot));
-
 	return true;
 }
 
@@ -41,24 +38,31 @@ EStateTreeRunStatus FMassClaimSmartObjectTask::EnterState(FStateTreeExecutionCon
 	UMassSignalSubsystem& SignalSubsystem = Context.GetExternalData(MassSignalSubsystemHandle);
 	FMassSmartObjectUserFragment& SOUser = Context.GetExternalData(SmartObjectUserHandle);
 
-	const FMassSmartObjectCandidateSlots& CandidateSlots = Context.GetInstanceData(CandidateSlotsHandle);
-	FSmartObjectClaimHandle& ClaimedSlot = Context.GetInstanceData(ClaimedSlotHandle);
-	ClaimedSlot.Invalidate();
+	FInstanceDataType& InstanceData = Context.GetInstanceData<FInstanceDataType>(*this);
+
+	const FMassSmartObjectCandidateSlots* CandidateSlots = InstanceData.CandidateSlots.GetPtr<FMassSmartObjectCandidateSlots>();
+	if (CandidateSlots == nullptr)
+	{
+		MASSBEHAVIOR_LOG(Log, TEXT("Candidate slots not set"));
+		return EStateTreeRunStatus::Failed;
+	}
+
+	InstanceData.ClaimedSlot.Invalidate();
 	
 	// Setup MassSmartObject handler and claim
 	const FMassStateTreeExecutionContext& MassContext = static_cast<FMassStateTreeExecutionContext&>(Context);
 	const FMassSmartObjectHandler MassSmartObjectHandler(MassContext.GetEntitySubsystem(), MassContext.GetEntitySubsystemExecutionContext(), SmartObjectSubsystem, SignalSubsystem);
 	
-	ClaimedSlot = MassSmartObjectHandler.ClaimCandidate(MassContext.GetEntity(), SOUser, CandidateSlots);
+	InstanceData.ClaimedSlot = MassSmartObjectHandler.ClaimCandidate(MassContext.GetEntity(), SOUser, *CandidateSlots);
 
 	// Treat claiming a slot as consuming all the candidate slots.
 	// This is done here because of the limited ways we can communicate between FindSmartObject() and ClaimSmartObject().
 	// InteractionCooldownEndTime is used by the FindSmartObject() to invalidate the candidates.
 	SOUser.InteractionCooldownEndTime = Context.GetWorld()->GetTimeSeconds() + InteractionCooldown;
 
-	if (!ClaimedSlot.IsValid())
+	if (!InstanceData.ClaimedSlot.IsValid())
 	{
-		MASSBEHAVIOR_LOG(Log, TEXT("Failed to claim smart object slot from %d candidates"), CandidateSlots.NumSlots);
+		MASSBEHAVIOR_LOG(Log, TEXT("Failed to claim smart object slot from %d candidates"), CandidateSlots->NumSlots);
 		return EStateTreeRunStatus::Failed;
 	}
 
@@ -72,20 +76,20 @@ void FMassClaimSmartObjectTask::ExitState(FStateTreeExecutionContext& Context, c
 		return;
 	}
 
-	const FSmartObjectClaimHandle& ClaimedSlot = Context.GetInstanceData(ClaimedSlotHandle);
 	FMassSmartObjectUserFragment& SOUser = Context.GetExternalData(SmartObjectUserHandle);
+	const FInstanceDataType& InstanceData = Context.GetInstanceData<FInstanceDataType>(*this);
 
 	// Succeeded or not, prevent interactions for a specified duration.
 	SOUser.InteractionCooldownEndTime = Context.GetWorld()->GetTimeSeconds() + InteractionCooldown;
 
-	if (ClaimedSlot.IsValid())
+	if (InstanceData.ClaimedSlot.IsValid())
 	{
 		const FMassStateTreeExecutionContext& MassContext = static_cast<FMassStateTreeExecutionContext&>(Context);
 		USmartObjectSubsystem& SmartObjectSubsystem = Context.GetExternalData(SmartObjectSubsystemHandle);
 		UMassSignalSubsystem& SignalSubsystem = Context.GetExternalData(MassSignalSubsystemHandle);
 		const FMassSmartObjectHandler MassSmartObjectHandler(MassContext.GetEntitySubsystem(), MassContext.GetEntitySubsystemExecutionContext(), SmartObjectSubsystem, SignalSubsystem);
 
-		MassSmartObjectHandler.ReleaseSmartObject(MassContext.GetEntity(), SOUser, ClaimedSlot);
+		MassSmartObjectHandler.ReleaseSmartObject(MassContext.GetEntity(), SOUser, InstanceData.ClaimedSlot);
 	}
 	else
 	{
@@ -96,7 +100,8 @@ void FMassClaimSmartObjectTask::ExitState(FStateTreeExecutionContext& Context, c
 EStateTreeRunStatus FMassClaimSmartObjectTask::Tick(FStateTreeExecutionContext& Context, const float DeltaTime) const
 {
 	FMassSmartObjectUserFragment& SOUser = Context.GetExternalData(SmartObjectUserHandle);
-
+	FInstanceDataType& InstanceData = Context.GetInstanceData<FInstanceDataType>(*this);
+	
 	// Prevent FindSmartObject() to query new objects while claimed.
 	// This is done here because of the limited ways we can communicate between FindSmartObject() and ClaimSmartObject().
 	// InteractionCooldownEndTime is used by the FindSmartObject() to invalidate the candidates.
@@ -104,15 +109,14 @@ EStateTreeRunStatus FMassClaimSmartObjectTask::Tick(FStateTreeExecutionContext&
 
 	// Check that the claimed slot is still valid, and if not, fail the task.
 	// The slot can become invalid if the whole SO or slot becomes invalidated.
-	FSmartObjectClaimHandle& ClaimedSlot = Context.GetInstanceData(ClaimedSlotHandle);
-	if (ClaimedSlot.IsValid())
+	if (InstanceData.ClaimedSlot.IsValid())
 	{
 		const USmartObjectSubsystem& SmartObjectSubsystem = Context.GetExternalData(SmartObjectSubsystemHandle);
-		if (!SmartObjectSubsystem.IsClaimedSmartObjectValid(ClaimedSlot))
+		if (!SmartObjectSubsystem.IsClaimedSmartObjectValid(InstanceData.ClaimedSlot))
 		{
-			ClaimedSlot.Invalidate();
+			InstanceData.ClaimedSlot.Invalidate();
 		}
 	}
 
-	return ClaimedSlot.IsValid() ? EStateTreeRunStatus::Running : EStateTreeRunStatus::Failed;
+	return InstanceData.ClaimedSlot.IsValid() ? EStateTreeRunStatus::Running : EStateTreeRunStatus::Failed;
 }
diff --git a/Engine/Plugins/AI/MassAI/Source/MassAIBehavior/Private/Tasks/MassZoneGraphFindEscapeTarget.cpp b/Engine/Plugins/AI/MassAI/Source/MassAIBehavior/Private/Tasks/MassZoneGraphFindEscapeTarget.cpp
--- a/Engine/Plugins/AI/MassAI/Source/MassAIBehavior/Private/Tasks/MassZoneGraphFindEscapeTarget.cpp
+++ b/Engine/Plugins/AI/MassAI/Source/MassAIBehavior/Private/Tasks/MassZoneGraphFindEscapeTarget.cpp
@@ -22,7 +22,7 @@ bool FMassZoneGraphFindEscapeTarget::Link(FStateTreeLinker& Linker)
 	Linker.LinkExternalData(ZoneGraphSubsystemHandle);
 	Linker.LinkExternalData(ZoneGraphAnnotationSubsystemHandle);
 
-	Linker.LinkInstanceDataProperty(EscapeTargetLocationHandle, STATETREE_INSTANCEDATA_PROPERTY(FMassZoneGraphFindEscapeTargetInstangeData, EscapeTargetLocation));
+	Linker.LinkInstanceDataProperty(EscapeTargetLocationHandle, STATETREE_INSTANCEDATA_PROPERTY(FMassZoneGraphFindEscapeTargetInstanceData, EscapeTargetLocation));
 
 	return true;
 }
diff --git a/Engine/Plugins/AI/MassAI/Source/MassAIBehavior/Private/Tasks/MassZoneGraphPathFollowTask.cpp b/Engine/Plugins/AI/MassAI/Source/MassAIBehavior/Private/Tasks/MassZoneGraphPathFollowTask.cpp
--- a/Engine/Plugins/AI/MassAI/Source/MassAIBehavior/Private/Tasks/MassZoneGraphPathFollowTask.cpp
+++ b/Engine/Plugins/AI/MassAI/Source/MassAIBehavior/Private/Tasks/MassZoneGraphPathFollowTask.cpp
@@ -23,10 +23,6 @@ bool FMassZoneGraphPathFollowTask::Link(FStateTreeLinker& Linker)
 	Linker.LinkExternalData(MovementParamsHandle);
 	Linker.LinkExternalData(ZoneGraphSubsystemHandle);
 
-	Linker.LinkInstanceDataProperty(TargetLocationHandle, STATETREE_INSTANCEDATA_PROPERTY(FMassZoneGraphPathFollowTaskInstanceData, TargetLocation));
-	Linker.LinkInstanceDataProperty(MovementStyleHandle, STATETREE_INSTANCEDATA_PROPERTY(FMassZoneGraphPathFollowTaskInstanceData, MovementStyle));
-	Linker.LinkInstanceDataProperty(SpeedScaleHandle, STATETREE_INSTANCEDATA_PROPERTY(FMassZoneGraphPathFollowTaskInstanceData, SpeedScale));
-
 	return true;
 }
 
@@ -41,6 +37,8 @@ bool FMassZoneGraphPathFollowTask::RequestPath(FMassStateTreeExecutionContext& C
 	FMassZoneGraphPathRequestFragment& RequestFragment = Context.GetExternalData(PathRequestHandle);
 	const FMassMovementParameters& MovementParams = Context.GetExternalData(MovementParamsHandle);
 
+	const FInstanceDataType& InstanceData = Context.GetInstanceData<FInstanceDataType>(*this);
+	
 	bool bDisplayDebug = false;
 #if WITH_MASSGAMEPLAY_DEBUG
 	bDisplayDebug = UE::Mass::Debug::IsDebuggingEntity(Context.GetEntity());
@@ -77,10 +75,7 @@ bool FMassZoneGraphPathFollowTask::RequestPath(FMassStateTreeExecutionContext& C
 	PathRequest.AnticipationDistance = RequestedTargetLocation.AnticipationDistance;
 	PathRequest.EndOfPathOffset.Set(FMath::RandRange(-AgentRadius.Radius, AgentRadius.Radius));
 
-	const FMassMovementStyleRef& MovementStyle = Context.GetInstanceData(MovementStyleHandle);
-	const float SpeedScale = Context.GetInstanceData(SpeedScaleHandle);
-
-	const float DesiredSpeed = FMath::Min(MovementParams.GenerateDesiredSpeed(MovementStyle, Context.GetEntity().Index) * SpeedScale, MovementParams.MaxSpeed);
+	const float DesiredSpeed = FMath::Min(MovementParams.GenerateDesiredSpeed(InstanceData.MovementStyle, Context.GetEntity().Index) * InstanceData.SpeedScale, MovementParams.MaxSpeed);
 	const UWorld* World = Context.GetWorld();
 	checkf(World != nullptr, TEXT("A valid world is expected from the execution context"));
 
@@ -91,6 +86,7 @@ bool FMassZoneGraphPathFollowTask::RequestPath(FMassStateTreeExecutionContext& C
 EStateTreeRunStatus FMassZoneGraphPathFollowTask::EnterState(FStateTreeExecutionContext& Context, const EStateTreeStateChangeType ChangeType, const FStateTreeTransitionResult& Transition) const
 {
 	FMassStateTreeExecutionContext& MassContext = static_cast<FMassStateTreeExecutionContext&>(Context);
+	const FInstanceDataType& InstanceData = Context.GetInstanceData<FInstanceDataType>(*this);
 
 	bool bDisplayDebug = false;
 #if WITH_MASSGAMEPLAY_DEBUG
@@ -102,15 +98,21 @@ EStateTreeRunStatus FMassZoneGraphPathFollowTask::EnterState(FStateTreeExecution
 	}
 
 	const FMassZoneGraphLaneLocationFragment& LaneLocation = Context.GetExternalData(LocationHandle);
-	const FMassZoneGraphTargetLocation& TargetLocation = Context.GetInstanceData(TargetLocationHandle);
+	const FMassZoneGraphTargetLocation* TargetLocation = InstanceData.TargetLocation.GetPtr<FMassZoneGraphTargetLocation>();
+
+	if (TargetLocation == nullptr)
+	{
+		MASSBEHAVIOR_LOG(Error, TEXT("Target is not defined."));
+		return EStateTreeRunStatus::Failed;
+	}
 	
-	if (TargetLocation.LaneHandle != LaneLocation.LaneHandle)
+	if (TargetLocation->LaneHandle != LaneLocation.LaneHandle)
 	{
-		MASSBEHAVIOR_LOG(Error, TEXT("Target is not on current lane, target lane is %s expected %s."), *TargetLocation.LaneHandle.ToString(), *LaneLocation.LaneHandle.ToString());
+		MASSBEHAVIOR_LOG(Error, TEXT("Target is not on current lane, target lane is %s expected %s."), *TargetLocation->LaneHandle.ToString(), *LaneLocation.LaneHandle.ToString());
 		return EStateTreeRunStatus::Failed;
 	}
 
-	if (!RequestPath(MassContext, TargetLocation))
+	if (!RequestPath(MassContext, *TargetLocation))
 	{
 		return EStateTreeRunStatus::Failed;
 	}
@@ -121,6 +123,7 @@ EStateTreeRunStatus FMassZoneGraphPathFollowTask::EnterState(FStateTreeExecution
 EStateTreeRunStatus FMassZoneGraphPathFollowTask::Tick(FStateTreeExecutionContext& Context, const float DeltaTime) const
 {
 	FMassStateTreeExecutionContext& MassContext = static_cast<FMassStateTreeExecutionContext&>(Context);
+	const FInstanceDataType& InstanceData = Context.GetInstanceData<FInstanceDataType>(*this);
 
 	bool bDisplayDebug = false;
 #if WITH_MASSGAMEPLAY_DEBUG
@@ -133,18 +136,23 @@ EStateTreeRunStatus FMassZoneGraphPathFollowTask::Tick(FStateTreeExecutionContex
 
 	const FMassZoneGraphShortPathFragment& ShortPath = Context.GetExternalData(ShortPathHandle);
 	const FMassZoneGraphLaneLocationFragment& LaneLocation = Context.GetExternalData(LocationHandle);
-	const FMassZoneGraphTargetLocation& TargetLocation = Context.GetInstanceData(TargetLocationHandle);
+	const FMassZoneGraphTargetLocation* TargetLocation = InstanceData.TargetLocation.GetPtr<FMassZoneGraphTargetLocation>();
 
 	// Current path follow is done, but it was partial (i.e. many points on a curve), try again until we get there.
 	if (ShortPath.IsDone() && ShortPath.bPartialResult)
 	{
-		if (TargetLocation.LaneHandle != LaneLocation.LaneHandle)
+		if (TargetLocation == nullptr)
+		{
+			MASSBEHAVIOR_LOG(Error, TEXT("Target is not defined."));
+			return EStateTreeRunStatus::Failed;
+		}
+		if (TargetLocation->LaneHandle != LaneLocation.LaneHandle)
 		{
-			MASSBEHAVIOR_LOG(Error, TEXT("Target is not on current lane, target lane is %s expected %s."), *TargetLocation.LaneHandle.ToString(), *LaneLocation.LaneHandle.ToString());
+			MASSBEHAVIOR_LOG(Error, TEXT("Target is not on current lane, target lane is %s expected %s."), *TargetLocation->LaneHandle.ToString(), *LaneLocation.LaneHandle.ToString());
 			return EStateTreeRunStatus::Failed;
 		}
 		
-		if (!RequestPath(MassContext, TargetLocation))
+		if (!RequestPath(MassContext, *TargetLocation))
 		{
 			MASSBEHAVIOR_LOG(Error, TEXT("Failed to request path."));
 			return EStateTreeRunStatus::Failed;
diff --git a/Engine/Plugins/AI/MassAI/Source/MassAIBehavior/Public/Tasks/MassClaimSmartObjectTask.h b/Engine/Plugins/AI/MassAI/Source/MassAIBehavior/Public/Tasks/MassClaimSmartObjectTask.h
--- a/Engine/Plugins/AI/MassAI/Source/MassAIBehavior/Public/Tasks/MassClaimSmartObjectTask.h
+++ b/Engine/Plugins/AI/MassAI/Source/MassAIBehavior/Public/Tasks/MassClaimSmartObjectTask.h
@@ -22,8 +22,8 @@ struct MASSAIBEHAVIOR_API FMassClaimSmartObjectTaskInstanceData
 	GENERATED_BODY()
 
 	/** Result of the candidates search request (Input) */
-	UPROPERTY(VisibleAnywhere, Category = Input)
-	FMassSmartObjectCandidateSlots CandidateSlots;
+	UPROPERTY(VisibleAnywhere, Category = Input, meta = (BaseStruct = "/Script/MassSmartObjects.MassSmartObjectCandidateSlots"))
+	FStateTreeStructRef CandidateSlots;
 
 	UPROPERTY(VisibleAnywhere, Category = Output)
 	FSmartObjectClaimHandle ClaimedSlot;
@@ -47,9 +47,6 @@ protected:
 	TStateTreeExternalDataHandle<USmartObjectSubsystem> SmartObjectSubsystemHandle;
 	TStateTreeExternalDataHandle<UMassSignalSubsystem> MassSignalSubsystemHandle;
 
-	TStateTreeInstanceDataPropertyHandle<FMassSmartObjectCandidateSlots> CandidateSlotsHandle;
-	TStateTreeInstanceDataPropertyHandle<FSmartObjectClaimHandle> ClaimedSlotHandle;
-
 	/** Delay in seconds before trying to use another smart object */
 	UPROPERTY(EditAnywhere, Category = Parameter)
 	float InteractionCooldown = 0.f;
diff --git a/Engine/Plugins/AI/MassAI/Source/MassAIBehavior/Public/Tasks/MassZoneGraphFindEscapeTarget.h b/Engine/Plugins/AI/MassAI/Source/MassAIBehavior/Public/Tasks/MassZoneGraphFindEscapeTarget.h
--- a/Engine/Plugins/AI/MassAI/Source/MassAIBehavior/Public/Tasks/MassZoneGraphFindEscapeTarget.h
+++ b/Engine/Plugins/AI/MassAI/Source/MassAIBehavior/Public/Tasks/MassZoneGraphFindEscapeTarget.h
@@ -15,7 +15,7 @@ class UZoneGraphAnnotationSubsystem;
  * Updates TargetLocation to a escape target based on the agents current location on ZoneGraph, and disturbance annotation.
  */
 USTRUCT()
-struct MASSAIBEHAVIOR_API FMassZoneGraphFindEscapeTargetInstangeData
+struct MASSAIBEHAVIOR_API FMassZoneGraphFindEscapeTargetInstanceData
 {
 	GENERATED_BODY()
 
@@ -32,7 +32,7 @@ struct MASSAIBEHAVIOR_API FMassZoneGraphFindEscapeTarget : public FMassStateTree
 
 protected:
 	virtual bool Link(FStateTreeLinker& Linker) override;
-	virtual const UStruct* GetInstanceDataType() const override { return FMassZoneGraphFindEscapeTargetInstangeData::StaticStruct(); }
+	virtual const UStruct* GetInstanceDataType() const override { return FMassZoneGraphFindEscapeTargetInstanceData::StaticStruct(); }
 	virtual EStateTreeRunStatus EnterState(FStateTreeExecutionContext& Context, const EStateTreeStateChangeType ChangeType, const FStateTreeTransitionResult& Transition) const override;
 
 	TStateTreeExternalDataHandle<FMassZoneGraphLaneLocationFragment> LocationHandle;
diff --git a/Engine/Plugins/AI/MassAI/Source/MassAIBehavior/Public/Tasks/MassZoneGraphPathFollowTask.h b/Engine/Plugins/AI/MassAI/Source/MassAIBehavior/Public/Tasks/MassZoneGraphPathFollowTask.h
--- a/Engine/Plugins/AI/MassAI/Source/MassAIBehavior/Public/Tasks/MassZoneGraphPathFollowTask.h
+++ b/Engine/Plugins/AI/MassAI/Source/MassAIBehavior/Public/Tasks/MassZoneGraphPathFollowTask.h
@@ -72,8 +72,8 @@ struct MASSAIBEHAVIOR_API FMassZoneGraphPathFollowTaskInstanceData
 {
 	GENERATED_BODY()
 
-	UPROPERTY(EditAnywhere, Category = Input)
-	FMassZoneGraphTargetLocation TargetLocation;
+	UPROPERTY(EditAnywhere, Category = Input, meta=(BaseStruct = "/Script/MassAIBehavior.MassZoneGraphTargetLocation"))
+	FStateTreeStructRef TargetLocation; 
 
 	UPROPERTY(EditAnywhere, Category = Parameter)
 	FMassMovementStyleRef MovementStyle;
@@ -87,9 +87,11 @@ struct MASSAIBEHAVIOR_API FMassZoneGraphPathFollowTask : public FMassStateTreeTa
 {
 	GENERATED_BODY()
 
+	using FInstanceDataType = FMassZoneGraphPathFollowTaskInstanceData;
+
 protected:
 	virtual bool Link(FStateTreeLinker& Linker) override;
-	virtual const UStruct* GetInstanceDataType() const override { return FMassZoneGraphPathFollowTaskInstanceData::StaticStruct(); };
+	virtual const UStruct* GetInstanceDataType() const override { return FInstanceDataType::StaticStruct(); };
 	virtual EStateTreeRunStatus EnterState(FStateTreeExecutionContext& Context, const EStateTreeStateChangeType ChangeType, const FStateTreeTransitionResult& Transition) const override;
 	virtual EStateTreeRunStatus Tick(FStateTreeExecutionContext& Context, const float DeltaTime) const override;
 
@@ -103,8 +105,4 @@ protected:
 	TStateTreeExternalDataHandle<FAgentRadiusFragment> AgentRadiusHandle;
 	TStateTreeExternalDataHandle<FMassMovementParameters> MovementParamsHandle;
 	TStateTreeExternalDataHandle<UZoneGraphSubsystem> ZoneGraphSubsystemHandle;
-
-	TStateTreeInstanceDataPropertyHandle<FMassZoneGraphTargetLocation> TargetLocationHandle;
-	TStateTreeInstanceDataPropertyHandle<FMassMovementStyleRef> MovementStyleHandle;
-	TStateTreeInstanceDataPropertyHandle<float> SpeedScaleHandle;
 };
diff --git a/Engine/Plugins/Runtime/StateTree/Source/StateTreeEditorModule/Private/Customizations/StateTreeBindingExtension.cpp b/Engine/Plugins/Runtime/StateTree/Source/StateTreeEditorModule/Private/Customizations/StateTreeBindingExtension.cpp
--- a/Engine/Plugins/Runtime/StateTree/Source/StateTreeEditorModule/Private/Customizations/StateTreeBindingExtension.cpp
+++ b/Engine/Plugins/Runtime/StateTree/Source/StateTreeEditorModule/Private/Customizations/StateTreeBindingExtension.cpp
@@ -8,11 +8,11 @@
 #include "IPropertyAccessEditor.h"
 #include "Features/IModularFeatures.h"
 #include "StateTreeEvaluatorBase.h"
-#include "StateTree.h"
 #include "Algo/Accumulate.h"
 #include "StateTreePropertyHelpers.h"
 #include "StateTreeAnyEnum.h"
 #include "StateTreePropertyBindingCompiler.h"
+#include "StateTreeCompiler.h"
 #include "PropertyEditor/Private/PropertyNode.h"
 
 #define LOCTEXT_NAMESPACE "StateTreeEditor"
@@ -182,16 +182,25 @@ void FStateTreeBindingExtension::ExtendWidgetRow(FDetailWidgetRow& InWidgetRow,
 
 	FProperty* Property = InPropertyHandle->GetProperty();
 
+	bool bIsDataRef = false;
 	bool bIsAnyEnum = false;
 	if (FStructProperty* StructProperty = CastField<FStructProperty>(Property))
 	{
 		bIsAnyEnum = StructProperty->Struct == FStateTreeAnyEnum::StaticStruct();
+		bIsDataRef = StructProperty->Struct == FStateTreeStructRef::StaticStruct();
+	}
+
+	const UScriptStruct* DataRefBaseStruct = nullptr;
+	if (bIsDataRef)
+	{
+		FString BaseStructName;
+		DataRefBaseStruct = UE::StateTree::Compiler::GetBaseStructFromMetaData(Property, BaseStructName);
 	}
 	
 	FPropertyBindingWidgetArgs Args;
 	Args.Property = InPropertyHandle->GetProperty();
 
-	Args.OnCanBindProperty = FOnCanBindProperty::CreateLambda([EditorBindings, OwnerObject, InPropertyHandle, bIsAnyEnum](FProperty* InProperty)
+	Args.OnCanBindProperty = FOnCanBindProperty::CreateLambda([EditorBindings, OwnerObject, InPropertyHandle, bIsAnyEnum, bIsDataRef, DataRefBaseStruct](FProperty* InProperty)
 		{
 			if (!EditorBindings || !OwnerObject)
 			{
@@ -228,6 +237,22 @@ void FStateTreeBindingExtension::ExtendWidgetRow(FDetailWidgetRow& InWidgetRow,
 					bCanBind = !AnyEnum.Enum || AnyEnum.Enum == EnumProperty->GetEnum();
 				}
 			}
+			else if (bIsDataRef && DataRefBaseStruct != nullptr)
+			{
+				if (const FStructProperty* SourceStructProperty = CastField<FStructProperty>(InProperty))
+				{
+					if (SourceStructProperty->Struct == TBaseStructure<FStateTreeStructRef>::Get())
+					{
+						FString SourceBaseStructName;
+						const UScriptStruct* SourceDataRefBaseStruct = UE::StateTree::Compiler::GetBaseStructFromMetaData(SourceStructProperty, SourceBaseStructName);
+						bCanBind = SourceDataRefBaseStruct && SourceDataRefBaseStruct->IsChildOf(DataRefBaseStruct);
+					}
+					else
+					{
+						bCanBind = SourceStructProperty->Struct && SourceStructProperty->Struct->IsChildOf(DataRefBaseStruct);
+					}
+				}
+			}
 			else
 			{
 				// Note: We support type promotion here
diff --git a/Engine/Plugins/Runtime/StateTree/Source/StateTreeEditorModule/Private/StateTreeCompiler.cpp b/Engine/Plugins/Runtime/StateTree/Source/StateTreeEditorModule/Private/StateTreeCompiler.cpp
--- a/Engine/Plugins/Runtime/StateTree/Source/StateTreeEditorModule/Private/StateTreeCompiler.cpp
+++ b/Engine/Plugins/Runtime/StateTree/Source/StateTreeEditorModule/Private/StateTreeCompiler.cpp
@@ -19,6 +19,25 @@ namespace UE::StateTree::Compiler
 		Log.Reportf(EMessageSeverity::Error, ContextStruct, TEXT("The StateTree is too complex. Compact index %s out of range %d/%d."), ContextText, Value, MaxValue);
 	}
 
+	const UScriptStruct* GetBaseStructFromMetaData(const FProperty* Property, FString& OutBaseStructName)
+	{
+		static const FName NAME_BaseStruct = "BaseStruct";
+
+		const UScriptStruct* Result = nullptr;
+		OutBaseStructName = Property->GetMetaData(NAME_BaseStruct);
+	
+		if (!OutBaseStructName.IsEmpty())
+		{
+			Result = UClass::TryFindTypeSlow<UScriptStruct>(OutBaseStructName);
+			if (!Result)
+			{
+				Result = LoadObject<UScriptStruct>(nullptr, *OutBaseStructName);
+			}
+		}
+
+		return Result;
+	}
+
 }; // UE::StateTree::Compiler
 
 bool FStateTreeCompiler::Compile(UStateTree& InStateTree)
@@ -866,13 +885,97 @@ bool FStateTreeCompiler::IsPropertyAnyEnum(const FStateTreeBindableStructDesc& S
 		{
 			if (const FStructProperty* OwnerStructProperty = CastField<FStructProperty>(OwnerProperty))
 			{
-				bIsAnyEnum = OwnerStructProperty->Struct == FStateTreeAnyEnum::StaticStruct();
+				bIsAnyEnum = OwnerStructProperty->Struct == TBaseStructure<FStateTreeAnyEnum>::Get();
 			}
 		}
 	}
 	return bIsAnyEnum;
 }
 
+bool FStateTreeCompiler::ValidateStructRef(const FStateTreeBindableStructDesc& SourceStruct, FStateTreeEditorPropertyPath SourcePath,
+											const FStateTreeBindableStructDesc& TargetStruct, FStateTreeEditorPropertyPath TargetPath) const
+{
+	TArray<FStateTreePropertySegment> Segments;
+
+	const FProperty* TargetLeafProperty = nullptr;
+	int32 TargetLeafArrayIndex = INDEX_NONE;
+	if (FStateTreePropertyBindingCompiler::ResolvePropertyPath(TargetStruct, TargetPath, Segments, TargetLeafProperty, TargetLeafArrayIndex) == false)
+	{
+		// This will later be reported by the bindings compiler.
+		return true;
+	}
+
+	// Early out if the target is not FStateTreeStructRef.
+	const FStructProperty* TargetStructProperty = CastField<FStructProperty>(TargetLeafProperty);
+	if (TargetStructProperty == nullptr || TargetStructProperty->Struct != TBaseStructure<FStateTreeStructRef>::Get())
+	{
+		return true;
+	}
+
+	FString TargetBaseStructName;
+	const UScriptStruct* TargetBaseStruct = UE::StateTree::Compiler::GetBaseStructFromMetaData(TargetStructProperty, TargetBaseStructName);
+	if (TargetBaseStruct == nullptr)
+	{
+		Log.Reportf(EMessageSeverity::Error, TargetStruct,
+				TEXT("Could not find base struct '%s' for target '%s:%s'."),
+				*TargetBaseStructName, *TargetStruct.Name.ToString(), *TargetPath.ToString());
+		return false;
+	}
+	
+	const FProperty* SourceLeafProperty = nullptr;
+	int32 SourceLeafArrayIndex = INDEX_NONE;
+	if (FStateTreePropertyBindingCompiler::ResolvePropertyPath(SourceStruct, SourcePath, Segments, SourceLeafProperty, SourceLeafArrayIndex) == false)
+	{
+		// This will later be reported by the bindings compiler.
+		return true;
+	}
+
+	// Exit if the source is not a struct property.
+	const FStructProperty* SourceStructProperty = CastField<FStructProperty>(SourceLeafProperty);
+	if (SourceStructProperty == nullptr)
+	{
+		return true;
+	}
+	
+	if (SourceStructProperty->Struct == TBaseStructure<FStateTreeStructRef>::Get())
+	{
+		// Source is struct ref too, check the types match.
+		FString SourceBaseStructName;
+		const UScriptStruct* SourceBaseStruct = UE::StateTree::Compiler::GetBaseStructFromMetaData(SourceStructProperty, SourceBaseStructName);
+		if (SourceBaseStruct == nullptr)
+		{
+			Log.Reportf(EMessageSeverity::Error, TargetStruct,
+					TEXT("Could not find base struct '%s' for bidning source '%s:%s'."),
+					*SourceBaseStructName, *SourceStruct.Name.ToString(), *SourcePath.ToString());
+			return false;
+		}
+
+		if (SourceBaseStruct->IsChildOf(TargetBaseStruct) == false)
+		{
+			Log.Reportf(EMessageSeverity::Error, TargetStruct,
+						TEXT("Type mismatch between source '%s:%s' and target '%s:%s' types, '%s' is not child of '%s'."),
+						*SourceStruct.Name.ToString(), *SourcePath.ToString(),
+						*TargetStruct.Name.ToString(), *TargetPath.ToString(),
+						*GetNameSafe(SourceBaseStruct), *GetNameSafe(TargetBaseStruct));
+			return false;
+		}
+	}
+	else
+	{
+		if (!SourceStructProperty->Struct || SourceStructProperty->Struct->IsChildOf(TargetBaseStruct) == false)
+		{
+			Log.Reportf(EMessageSeverity::Error, TargetStruct,
+						TEXT("Type mismatch between source '%s:%s' and target '%s:%s' types, '%s' is not child of '%s'."),
+						*SourceStruct.Name.ToString(), *SourcePath.ToString(),
+						*TargetStruct.Name.ToString(), *TargetPath.ToString(),
+						*GetNameSafe(SourceStructProperty->Struct), *GetNameSafe(TargetBaseStruct));
+			return false;
+		}
+	}
+
+	return true;
+}
+
 bool FStateTreeCompiler::GetAndValidateBindings(const FStateTreeBindableStructDesc& TargetStruct, TArray<FStateTreeEditorPropertyBinding>& OutBindings) const
 {
 	OutBindings.Reset();
@@ -937,6 +1040,12 @@ bool FStateTreeCompiler::GetAndValidateBindings(const FStateTreeBindableStructDe
 		{
 			OutBindings.Add(Binding);
 		}
+
+		// Check if the bindings is for struct ref and validate the types.
+		if (!ValidateStructRef(SourceStruct, Binding.SourcePath, TargetStruct, Binding.TargetPath))
+		{
+			return false;
+		}
 	}
 
 	return true;
diff --git a/Engine/Plugins/Runtime/StateTree/Source/StateTreeEditorModule/Private/StateTreePropertyBindingCompiler.cpp b/Engine/Plugins/Runtime/StateTree/Source/StateTreeEditorModule/Private/StateTreePropertyBindingCompiler.cpp
--- a/Engine/Plugins/Runtime/StateTree/Source/StateTreeEditorModule/Private/StateTreePropertyBindingCompiler.cpp
+++ b/Engine/Plugins/Runtime/StateTree/Source/StateTreeEditorModule/Private/StateTreePropertyBindingCompiler.cpp
@@ -201,6 +201,22 @@ EStateTreePropertyCopyType FStateTreePropertyBindingCompiler::GetCopyType(const
 		}
 	}
 
+	// Handle FStateTreeStructRef
+	if (const FStructProperty* TargetStructProperty = CastField<const FStructProperty>(TargetProperty))
+	{
+		if (TargetStructProperty->Struct == TBaseStructure<FStateTreeStructRef>::Get())
+		{
+			if (const FStructProperty* SourceStructProperty = CastField<const FStructProperty>(SourceProperty))
+			{
+				// FStateTreeStructRef to FStateTreeStructRef is copied as usual.
+				if (SourceStructProperty->Struct != TBaseStructure<FStateTreeStructRef>::Get())
+				{
+					return EStateTreePropertyCopyType::StructReference;
+				}
+			}
+		}
+	}
+
 	const EPropertyAccessCompatibility Compatibility = GetPropertyCompatibility(SourceProperty, TargetProperty);
 
 	// Extract underlying types for enums
diff --git a/Engine/Plugins/Runtime/StateTree/Source/StateTreeEditorModule/Public/StateTreeCompiler.h b/Engine/Plugins/Runtime/StateTree/Source/StateTreeEditorModule/Public/StateTreeCompiler.h
--- a/Engine/Plugins/Runtime/StateTree/Source/StateTreeEditorModule/Public/StateTreeCompiler.h
+++ b/Engine/Plugins/Runtime/StateTree/Source/StateTreeEditorModule/Public/StateTreeCompiler.h
@@ -47,6 +47,8 @@ private:
 	bool CreateEvaluator(const FStateTreeEditorNode& EvalNode);
 	bool GetAndValidateBindings(const FStateTreeBindableStructDesc& TargetStruct, TArray<FStateTreeEditorPropertyBinding>& OutBindings) const;
 	bool IsPropertyAnyEnum(const FStateTreeBindableStructDesc& Struct, FStateTreeEditorPropertyPath Path) const;
+	bool ValidateStructRef(const FStateTreeBindableStructDesc& SourceStruct, FStateTreeEditorPropertyPath SourcePath,
+							const FStateTreeBindableStructDesc& TargetStruct, FStateTreeEditorPropertyPath TargetPath) const;
 
 	FStateTreeCompilerLog& Log;
 	UStateTree* StateTree = nullptr;
@@ -130,4 +132,12 @@ namespace UE::StateTree::Compiler
 		return FValidationResult(Count >= 0 && Count <= MAX_uint8, Count, MAX_uint8);
 	}
 
+	/**
+	 * Returns UScriptStruct defined in "BaseStruct" metadata of given property.
+	 * @param Property Handle to property where value is got from.
+	 * @param OutBaseStructName Handle to property where value is got from.
+	 * @return Script struct defined by the BaseStruct or nullptr if not found.
+	 */
+	const UScriptStruct* GetBaseStructFromMetaData(const FProperty* Property, FString& OutBaseStructName);
+
 }; // UE::StateTree::Compiler
\ No newline at end of file
diff --git a/Engine/Plugins/Runtime/StateTree/Source/StateTreeModule/Private/StateTreePropertyBindings.cpp b/Engine/Plugins/Runtime/StateTree/Source/StateTreeModule/Private/StateTreePropertyBindings.cpp
--- a/Engine/Plugins/Runtime/StateTree/Source/StateTreeModule/Private/StateTreePropertyBindings.cpp
+++ b/Engine/Plugins/Runtime/StateTree/Source/StateTreeModule/Private/StateTreePropertyBindings.cpp
@@ -247,6 +247,17 @@ bool FStateTreePropertyBindings::ValidateCopy(FStateTreePropCopy& Copy) const
 	case EStateTreePropertyCopyType::CopyFixedArray:
 		bResult = SourceProperty->IsA<FArrayProperty>() && TargetProperty->IsA<FArrayProperty>();
 		break;
+	case EStateTreePropertyCopyType::StructReference:
+	{
+		const FStructProperty* SourceStructProperty = CastField<const FStructProperty>(SourceProperty);
+		const FStructProperty* TargetStructProperty = CastField<const FStructProperty>(TargetProperty);
+
+		bResult = SourceStructProperty != nullptr
+			&& TargetStructProperty != nullptr
+			&& SourceStructProperty->Struct != TBaseStructure<FStateTreeStructRef>::Get()
+			&& TargetStructProperty->Struct == TBaseStructure<FStateTreeStructRef>::Get();
+		break;
+	}
 	case EStateTreePropertyCopyType::PromoteBoolToByte:
 		bResult = SourceProperty->IsA<FBoolProperty>() && TargetProperty->IsA<FByteProperty>();
 		break;
@@ -378,7 +389,7 @@ uint8* FStateTreePropertyBindings::GetAddress(FStateTreeDataView InStructView, c
 	return Address;
 }
 
-void FStateTreePropertyBindings::PerformCopy(const FStateTreePropCopy& Copy, const uint8* SourceAddress, uint8* TargetAddress) const
+void FStateTreePropertyBindings::PerformCopy(const FStateTreePropCopy& Copy, uint8* SourceAddress, uint8* TargetAddress) const
 {
 	// Source property can be null
 	check(SourceAddress);
@@ -429,6 +440,13 @@ void FStateTreePropertyBindings::PerformCopy(const FStateTreePropCopy& Copy, con
 		}
 		break;
 	}
+	case EStateTreePropertyCopyType::StructReference:
+	{
+		const FStructProperty* SourceStructProperty = static_cast<const FStructProperty*>(Copy.SourceLeafProperty);
+		FStateTreeStructRef* Target = (FStateTreeStructRef*)TargetAddress;
+		Target->Set(FStructView(SourceStructProperty->Struct, SourceAddress));
+		break;
+	}
 	case EStateTreePropertyCopyType::PromoteBoolToByte:
 		*reinterpret_cast<uint8*>(TargetAddress) = (uint8)static_cast<const FBoolProperty*>(Copy.SourceLeafProperty)->GetPropertyValue(SourceAddress);
 		break;
@@ -526,7 +544,7 @@ bool FStateTreePropertyBindings::CopyTo(TConstArrayView<FStateTreeDataView> Sour
 			check(SourceStructView.GetStruct() == SourceStructs[Copy.SourceStructIndex.Get()].Struct
 				|| (SourceStructView.GetStruct() && SourceStructView.GetStruct()->IsChildOf(SourceStructs[Copy.SourceStructIndex.Get()].Struct)));
 				
-			const uint8* SourceAddress = GetAddress(SourceStructView, Copy.SourceIndirection, Copy.SourceLeafProperty);
+			uint8* SourceAddress = GetAddress(SourceStructView, Copy.SourceIndirection, Copy.SourceLeafProperty);
 			uint8* TargetAddress = GetAddress(TargetStructView, Copy.TargetIndirection, Copy.TargetLeafProperty);
 			
 			check (SourceAddress != nullptr && TargetAddress != nullptr);
diff --git a/Engine/Plugins/Runtime/StateTree/Source/StateTreeModule/Public/StateTreeExecutionContext.h b/Engine/Plugins/Runtime/StateTree/Source/StateTreeModule/Public/StateTreeExecutionContext.h
--- a/Engine/Plugins/Runtime/StateTree/Source/StateTreeModule/Public/StateTreeExecutionContext.h
+++ b/Engine/Plugins/Runtime/StateTree/Source/StateTreeModule/Public/StateTreeExecutionContext.h
@@ -5,6 +5,7 @@
 #include "StateTree.h"
 #include "StateTreePropertyBindings.h"
 #include "StateTreeInstanceData.h"
+#include "StateTreeNodeBase.h"
 #include "StateTreeExecutionContext.generated.h"
 
 struct FStateTreeEvaluatorBase;
@@ -255,9 +256,25 @@ public:
 		}
 		return nullptr;
 	}
-	
+
+	/** @returns pointer to the instance data of specified node. */
+	template <typename T>
+	T* GetInstanceDataPtr(const FStateTreeNodeBase& Node) const
+	{
+		return DataViews[Node.DataViewIndex.Get()].template GetMutablePtr<T>();
+	}
+
+	/** @returns reference to the instance data of specified node. */
+	template <typename T>
+	T& GetInstanceData(const FStateTreeNodeBase& Node) const
+	{
+		return DataViews[Node.DataViewIndex.Get()].template GetMutable<T>();
+	}
+
+	/** @return the status of the last tick function */
 	EStateTreeRunStatus GetLastTickStatus(const FStateTreeInstanceData* ExternalInstanceData = nullptr) const;
 
+	/** @return reference to the list of currently active states. */
 	const FStateTreeActiveStates& GetActiveStates(const FStateTreeInstanceData* ExternalInstanceData = nullptr) const;
 
 #if WITH_GAMEPLAY_DEBUGGER
diff --git a/Engine/Plugins/Runtime/StateTree/Source/StateTreeModule/Public/StateTreePropertyBindings.h b/Engine/Plugins/Runtime/StateTree/Source/StateTreeModule/Public/StateTreePropertyBindings.h
--- a/Engine/Plugins/Runtime/StateTree/Source/StateTreeModule/Public/StateTreePropertyBindings.h
+++ b/Engine/Plugins/Runtime/StateTree/Source/StateTreeModule/Public/StateTreePropertyBindings.h
@@ -227,6 +227,8 @@ enum class EStateTreePropertyCopyType : uint8
 	CopyName,					// FName needs special case because its size changes between editor/compiler and runtime.
 	CopyFixedArray,				// Array needs special handling for fixed size TArrays
 
+	StructReference,			// Copies pointer to a source struct into a FStateTreeStructRef.
+
 	/* Promote the type during the copy */
 
 	/* Bool promotions */
@@ -448,7 +450,7 @@ struct STATETREEMODULE_API FStateTreePropertyBindings
 protected:
 	[[nodiscard]] bool ResolvePath(const UStruct* Struct, const FStateTreePropertySegment& FirstPathSegment, FStateTreePropertyIndirection& OutFirstIndirection, const FProperty*& OutLeafProperty);
 	[[nodiscard]] bool ValidateCopy(FStateTreePropCopy& Copy) const;
-	void PerformCopy(const FStateTreePropCopy& Copy, const uint8* SourceAddress, uint8* TargetAddress) const;
+	void PerformCopy(const FStateTreePropCopy& Copy, uint8* SourceAddress, uint8* TargetAddress) const;
 	uint8* GetAddress(FStateTreeDataView InStructView, const FStateTreePropertyIndirection& FirstIndirection, const FProperty* LeafProperty) const;
 	FString GetPathAsString(const FStateTreePropertySegment& FirstPathSegment, const FStateTreePropertySegment* HighlightedSegment = nullptr, const TCHAR* HighlightPrefix = nullptr, const TCHAR* HighlightPostfix = nullptr);
 
diff --git a/Engine/Plugins/Runtime/StateTree/Source/StateTreeModule/Public/StateTreeTypes.h b/Engine/Plugins/Runtime/StateTree/Source/StateTreeModule/Public/StateTreeTypes.h
--- a/Engine/Plugins/Runtime/StateTree/Source/StateTreeModule/Public/StateTreeTypes.h
+++ b/Engine/Plugins/Runtime/StateTree/Source/StateTreeModule/Public/StateTreeTypes.h
@@ -639,3 +639,83 @@ struct STATETREEMODULE_API FStateTreeExternalDataDesc
 
 #define STATETREE_INSTANCEDATA_PROPERTY(Struct, Member) \
 		decltype(Struct::Member){}, Struct::StaticStruct(), TEXT(#Member)
+
+/**
+ * StateTree struct ref allows to get a reference/pointer to a specified type via property binding.
+ * It is useful for referencing larger properties to avoid copies of the data, or to be able to write to a bounds property.
+ *
+ * The expected type of the reference should be set in "BaseStruct" meta tag.
+ *
+ * Example:
+ *
+ *	USTRUCT()
+ *	struct FAwesomeTaskInstanceData
+ *	{
+ *		GENERATED_BODY()
+ *
+ *		UPROPERTY(VisibleAnywhere, Category = Input, meta = (BaseStruct = "/Script/AwesomeModule.AwesomeData"))
+ *		FStateTreeStructRef Data;
+ *	};
+ *
+ *
+ *	if (const FAwesomeData* Awesome = InstanceData.Data.GetPtr<FAwesomeData>())
+ *	{
+ *		...
+ *	}
+ */
+USTRUCT(BlueprintType)
+struct STATETREEMODULE_API FStateTreeStructRef
+{
+	GENERATED_BODY()
+
+	FStateTreeStructRef() = default;
+
+	/** @return true if the reference is valid (safe to use the reference getters). */
+	bool IsValid() const
+	{
+		return Data.IsValid();
+	}
+
+	/** Sets the struct ref (used by property copy) */
+	void Set(FStructView NewData)
+	{
+		Data = NewData;
+	}
+
+	/** Returns const reference to the struct, this getter assumes that all data is valid. */
+	template <typename T>
+	const T& Get() const
+	{
+		return Data.template Get<T>();
+	}
+
+	/** Returns const pointer to the struct, or nullptr if cast is not valid. */
+	template <typename T>
+	const T* GetPtr() const
+	{
+		return Data.template GetPtr<T>();
+	}
+
+	/** Returns mutable reference to the struct, this getter assumes that all data is valid. */
+	template <typename T>
+    T& GetMutable() const
+	{
+		return Data.template GetMutable<T>();
+	}
+
+	/** Returns mutable pointer to the struct, or nullptr if cast is not valid. */
+	template <typename T>
+    T* GetMutablePtr() const
+	{
+		return Data.template GetMutablePtr<T>();
+	}
+
+	/** @return Struct describing the data type. */
+	const UScriptStruct* GetScriptStruct() const
+	{
+		return Data.GetScriptStruct();
+	}
+
+protected:
+	FStructView Data;
+};
\ No newline at end of file
