Cut UMassDebuggerSubsystem's capability to host entity handles to be debugged and replaced the only use case (in Mass gameplay debugger category) with a custom query to fetch location-owning entities.

The old approach had the downside of allowing to pick only entities that had a special debug fragment in their composition, which resulted in confusion while trying to pick entities.

diff --git a/Engine/Plugins/AI/MassAI/Source/MassAIDebug/Private/GameplayDebuggerCategory_Mass.cpp b/Engine/Plugins/AI/MassAI/Source/MassAIDebug/Private/GameplayDebuggerCategory_Mass.cpp
--- a/Engine/Plugins/AI/MassAI/Source/MassAIDebug/Private/GameplayDebuggerCategory_Mass.cpp
+++ b/Engine/Plugins/AI/MassAI/Source/MassAIDebug/Private/GameplayDebuggerCategory_Mass.cpp
@@ -123,21 +123,40 @@ void FGameplayDebuggerCategory_Mass::PickEntity(const APlayerController& OwnerPC
 	FVector ViewDirection = FVector::ForwardVector;
 	ensureMsgf(GetViewPoint(&OwnerPC, ViewLocation, ViewDirection), TEXT("GetViewPoint is expected to always succeed when passing a valid controller."));
 
+	UMassEntitySubsystem* EntitySystem = UWorld::GetSubsystem<UMassEntitySubsystem>(&World);
+	if (EntitySystem == nullptr)
+	{
+		return;
+	}
+
 	FMassEntityHandle BestEntity;
 	// entities indicated by UE::Mass::Debug take precedence 
     if (UE::Mass::Debug::HasDebugEntities())
     {
-	    if (const UMassEntitySubsystem* EntitySystem = UWorld::GetSubsystem<UMassEntitySubsystem>(&World))
-	    {
-	    	TArray<FMassEntityHandle> Entities;
-	    	TArray<FVector> Locations;
-	    	UE::Mass::Debug::GetDebugEntitiesAndLocations(*EntitySystem, Entities, Locations);
-	    	BestEntity = UE::Mass::Debug::GetBestEntity(ViewLocation, ViewDirection, Entities, Locations, bLimitAngle);
-	    }
+		TArray<FMassEntityHandle> Entities;
+	    TArray<FVector> Locations;
+	    UE::Mass::Debug::GetDebugEntitiesAndLocations(*EntitySystem, Entities, Locations);
+	    BestEntity = UE::Mass::Debug::GetBestEntity(ViewLocation, ViewDirection, Entities, Locations, bLimitAngle);
     }
 	else
 	{
-		BestEntity = UE::Mass::Debug::GetBestEntity(ViewLocation, ViewDirection, Debugger.GetEntities(), Debugger.GetLocations(), bLimitAngle);
+		TArray<FMassEntityHandle> Entities;
+		TArray<FVector> Locations;
+		FMassExecutionContext ExecutionContext;
+		FMassEntityQuery Query;
+		Query.AddRequirement<FTransformFragment>(EMassFragmentAccess::ReadOnly);
+		Query.ForEachEntityChunk(*EntitySystem, ExecutionContext, [&Entities, &Locations](FMassExecutionContext& Context)
+		{
+			Entities.Append(Context.GetEntities().GetData(), Context.GetEntities().Num());
+			TConstArrayView<FTransformFragment> InLocations = Context.GetFragmentView<FTransformFragment>();
+			Locations.Reserve(Locations.Num() + InLocations.Num());
+			for (const FTransformFragment& TransformFragment : InLocations)
+			{
+				Locations.Add(TransformFragment.GetTransform().GetLocation());
+			}
+		});
+
+		BestEntity = UE::Mass::Debug::GetBestEntity(ViewLocation, ViewDirection, Entities, Locations, bLimitAngle);
 	}
 
 	AActor* BestActor = nullptr;
diff --git a/Engine/Plugins/Runtime/MassGameplay/Source/MassGameplayDebug/Private/DebugVisLocationProcessor.cpp b/Engine/Plugins/Runtime/MassGameplay/Source/MassGameplayDebug/Private/DebugVisLocationProcessor.cpp
--- a/Engine/Plugins/Runtime/MassGameplay/Source/MassGameplayDebug/Private/DebugVisLocationProcessor.cpp
+++ b/Engine/Plugins/Runtime/MassGameplay/Source/MassGameplayDebug/Private/DebugVisLocationProcessor.cpp
@@ -105,7 +105,6 @@ void UMassProcessor_UpdateDebugVis::Execute(UMassEntitySubsystem& EntitySubsyste
 			for (int32 i = 0; i < NumEntities; ++i)
 			{
 				Debugger.AddShape(DebugVisList[i].Shape, LocationList[i].GetTransform().GetLocation(), RadiiList[i].Radius);
-				Debugger.AddEntityLocation(Context.GetEntity(i), LocationList[i].GetTransform().GetLocation());
 			}
 		});
 }
diff --git a/Engine/Plugins/Runtime/MassGameplay/Source/MassGameplayDebug/Private/MassDebuggerSubsystem.cpp b/Engine/Plugins/Runtime/MassGameplay/Source/MassGameplayDebug/Private/MassDebuggerSubsystem.cpp
--- a/Engine/Plugins/Runtime/MassGameplay/Source/MassGameplayDebug/Private/MassDebuggerSubsystem.cpp
+++ b/Engine/Plugins/Runtime/MassGameplay/Source/MassGameplayDebug/Private/MassDebuggerSubsystem.cpp
@@ -32,12 +32,6 @@ void UMassDebuggerSubsystem::PreTickProcessors()
 	SelectedEntityDetails.Empty();
 }
 
-void UMassDebuggerSubsystem::AddEntityLocation(FMassEntityHandle Entity, const FVector& Location)
-{
-	Entities.Add(Entity);
-	Locations.Add(Location);
-}
-
 void UMassDebuggerSubsystem::OnProcessingPhaseStarted(const float DeltaSeconds)
 {
 	PreTickProcessors();
diff --git a/Engine/Plugins/Runtime/MassGameplay/Source/MassGameplayDebug/Public/MassDebuggerSubsystem.h b/Engine/Plugins/Runtime/MassGameplay/Source/MassGameplayDebug/Public/MassDebuggerSubsystem.h
--- a/Engine/Plugins/Runtime/MassGameplay/Source/MassGameplayDebug/Public/MassDebuggerSubsystem.h
+++ b/Engine/Plugins/Runtime/MassGameplay/Source/MassGameplayDebug/Public/MassDebuggerSubsystem.h
@@ -30,10 +30,6 @@ public:
 	void AddShape(EMassEntityDebugShape Shape, FVector Location, float Size) { Shapes[uint8(Shape)].Add(FShapeDesc(Location, Size)); }
 	const TArray<FShapeDesc>* GetShapes() const { return Shapes; }
 
-	void AddEntityLocation(FMassEntityHandle Entity, const FVector& Location);
-	TConstArrayView<FMassEntityHandle> GetEntities() const { return Entities; }
-	TConstArrayView<FVector> GetLocations() const { return Locations; }
-
 	FMassEntityHandle GetSelectedEntity() const { return SelectedEntity; }
 	void SetSelectedEntity(const FMassEntityHandle InSelectedEntity);
 
