Made Mass gameplay debugger category's use of UMassDebuggerSubsystem optional and instead require UMassEntitySubsystem.
Also, added a way to broadcast mass entity debug-selection.

diff --git a/Engine/Plugins/AI/MassAI/Source/MassAIDebug/Private/GameplayDebuggerCategory_Mass.cpp b/Engine/Plugins/AI/MassAI/Source/MassAIDebug/Private/GameplayDebuggerCategory_Mass.cpp
--- a/Engine/Plugins/AI/MassAI/Source/MassAIDebug/Private/GameplayDebuggerCategory_Mass.cpp
+++ b/Engine/Plugins/AI/MassAI/Source/MassAIDebug/Private/GameplayDebuggerCategory_Mass.cpp
@@ -111,31 +111,25 @@ FGameplayDebuggerCategory_Mass::FGameplayDebuggerCategory_Mass()
 	BindKeyPress(EKeys::C.GetFName(), FGameplayDebuggerInputModifier::Shift, this, &FGameplayDebuggerCategory_Mass::OnToggleNearEntityPath, EGameplayDebuggerInputMode::Replicated);
 }
 
-void FGameplayDebuggerCategory_Mass::SetCachedEntity(const FMassEntityHandle Entity, UMassDebuggerSubsystem& Debugger)
+void FGameplayDebuggerCategory_Mass::SetCachedEntity(const FMassEntityHandle Entity, UMassEntitySubsystem& EntitySubsystem)
 {
 	CachedEntity = Entity;
-	Debugger.SetSelectedEntity(Entity);
+	FMassDebugger::SelectEntity(EntitySubsystem, Entity);
 }
 
-void FGameplayDebuggerCategory_Mass::PickEntity(const APlayerController& OwnerPC, const UWorld& World, UMassDebuggerSubsystem& Debugger, const bool bLimitAngle)
+void FGameplayDebuggerCategory_Mass::PickEntity(const APlayerController& OwnerPC, const UWorld& World, UMassEntitySubsystem& EntitySubsystem, const bool bLimitAngle)
 {
 	FVector ViewLocation = FVector::ZeroVector;
 	FVector ViewDirection = FVector::ForwardVector;
 	ensureMsgf(GetViewPoint(&OwnerPC, ViewLocation, ViewDirection), TEXT("GetViewPoint is expected to always succeed when passing a valid controller."));
 
-	UMassEntitySubsystem* EntitySystem = UWorld::GetSubsystem<UMassEntitySubsystem>(&World);
-	if (EntitySystem == nullptr)
-	{
-		return;
-	}
-
 	FMassEntityHandle BestEntity;
 	// entities indicated by UE::Mass::Debug take precedence 
     if (UE::Mass::Debug::HasDebugEntities())
     {
 		TArray<FMassEntityHandle> Entities;
 	    TArray<FVector> Locations;
-	    UE::Mass::Debug::GetDebugEntitiesAndLocations(*EntitySystem, Entities, Locations);
+	    UE::Mass::Debug::GetDebugEntitiesAndLocations(EntitySubsystem, Entities, Locations);
 	    BestEntity = UE::Mass::Debug::GetBestEntity(ViewLocation, ViewDirection, Entities, Locations, bLimitAngle);
     }
 	else
@@ -145,7 +139,7 @@ void FGameplayDebuggerCategory_Mass::PickEntity(const APlayerController& OwnerPC
 		FMassExecutionContext ExecutionContext;
 		FMassEntityQuery Query;
 		Query.AddRequirement<FTransformFragment>(EMassFragmentAccess::ReadOnly);
-		Query.ForEachEntityChunk(*EntitySystem, ExecutionContext, [&Entities, &Locations](FMassExecutionContext& Context)
+		Query.ForEachEntityChunk(EntitySubsystem, ExecutionContext, [&Entities, &Locations](FMassExecutionContext& Context)
 		{
 			Entities.Append(Context.GetEntities().GetData(), Context.GetEntities().Num());
 			TConstArrayView<FTransformFragment> InLocations = Context.GetFragmentView<FTransformFragment>();
@@ -168,7 +162,7 @@ void FGameplayDebuggerCategory_Mass::PickEntity(const APlayerController& OwnerPC
 		}
 	}
 
-	SetCachedEntity(BestEntity, Debugger);
+	SetCachedEntity(BestEntity, EntitySubsystem);
 	CachedDebugActor = BestActor;
 	GetReplicator()->SetDebugActor(BestActor);
 }
@@ -182,24 +176,26 @@ void FGameplayDebuggerCategory_Mass::CollectData(APlayerController* OwnerPC, AAc
 {
 	UWorld* World = GetDataWorld(OwnerPC, DebugActor);
 	check(World);
-	UMassDebuggerSubsystem* Debugger = World->GetSubsystem<UMassDebuggerSubsystem>();
-	if (Debugger == nullptr)
+
+	UMassEntitySubsystem* EntitySubsystem = UWorld::GetSubsystem<UMassEntitySubsystem>(World);
+	if (EntitySubsystem == nullptr)
 	{
-		AddTextLine(FString::Printf(TEXT("{Red}MassDebuggerSubsystem instance is missing")));
+		AddTextLine(FString::Printf(TEXT("{Red}EntitySubsystem instance is missing")));
 		return;
 	}
-	Debugger->SetCollectingData();
+
+	UMassDebuggerSubsystem* Debugger = World->GetSubsystem<UMassDebuggerSubsystem>();
 
 	const UMassAgentComponent* AgentComp = nullptr;
 	if (DebugActor)
 	{
 		const FMassEntityHandle EntityHandle = UE::Mass::Debug::GetEntityFromActor(*DebugActor, AgentComp);	
-		SetCachedEntity(EntityHandle, *Debugger);
+		SetCachedEntity(EntityHandle, *EntitySubsystem);
 		CachedDebugActor = DebugActor;
 	}
 	else if (CachedDebugActor)
 	{
-		SetCachedEntity(FMassEntityHandle(), *Debugger);
+		SetCachedEntity(FMassEntityHandle(), *EntitySubsystem);
 		CachedDebugActor = nullptr;
 	}
 
@@ -209,61 +205,53 @@ void FGameplayDebuggerCategory_Mass::CollectData(APlayerController* OwnerPC, AAc
 		// also support client-server picking. For now, we handle explicit mass picking requests on the authority
 		if (bPickEntity)
 		{
-			PickEntity(*OwnerPC, *World, *Debugger);
+			PickEntity(*OwnerPC, *World, *EntitySubsystem);
 			bPickEntity = false;
 		}
 		// if we're debugging based on UE::Mass::Debug and the range changed
 		else if (CachedDebugActor == nullptr && UE::Mass::Debug::HasDebugEntities() && UE::Mass::Debug::IsDebuggingEntity(CachedEntity) == false)
 		{
 			// using bLimitAngle = false to not limit the selection to only the things in from of the player
-			PickEntity(*OwnerPC, *World, *Debugger, /*bLimitAngle=*/false);
+			PickEntity(*OwnerPC, *World, *EntitySubsystem, /*bLimitAngle=*/false);
 		}
 	}
 
-	UMassEntitySubsystem* EntitySystem = UWorld::GetSubsystem<UMassEntitySubsystem>(World);
-	if (EntitySystem)
-	{
-		AddTextLine(FString::Printf(TEXT("{Green}Entities count active{grey}/all: {white}%d{grey}/%d"), EntitySystem->DebugGetEntityCount(), EntitySystem->DebugGetEntityCount()));
-		AddTextLine(FString::Printf(TEXT("{Green}Registered Archetypes count: {white}%d {green}data ver: {white}%d"), EntitySystem->DebugGetArchetypesCount(), EntitySystem->GetArchetypeDataVersion()));
+	AddTextLine(FString::Printf(TEXT("{Green}Entities count active{grey}/all: {white}%d{grey}/%d"), EntitySubsystem->DebugGetEntityCount(), EntitySubsystem->DebugGetEntityCount()));
+	AddTextLine(FString::Printf(TEXT("{Green}Registered Archetypes count: {white}%d {green}data ver: {white}%d"), EntitySubsystem->DebugGetArchetypesCount(), EntitySubsystem->GetArchetypeDataVersion()));
 
-		if (UE::Mass::Debug::HasDebugEntities())
-		{
-			int32 RangeBegin, RangeEnd;
-			UE::Mass::Debug::GetDebugEntitiesRange(RangeBegin, RangeEnd);
-			AddTextLine(FString::Printf(TEXT("{Green}Debugged entity range: {orange}%d-%d"), RangeBegin, RangeEnd));
-		}
+	if (UE::Mass::Debug::HasDebugEntities())
+	{
+		int32 RangeBegin, RangeEnd;
+		UE::Mass::Debug::GetDebugEntitiesRange(RangeBegin, RangeEnd);
+		AddTextLine(FString::Printf(TEXT("{Green}Debugged entity range: {orange}%d-%d"), RangeBegin, RangeEnd));
+	}
 
-		if (bShowArchetypes)
-		{
-			FStringOutputDevice Ar;
-			Ar.SetAutoEmitLineTerminator(true);
-			EntitySystem->DebugPrintArchetypes(Ar, /*bIncludeEmpty*/false);
+	if (bShowArchetypes)
+	{
+		FStringOutputDevice Ar;
+		Ar.SetAutoEmitLineTerminator(true);
+		EntitySubsystem->DebugPrintArchetypes(Ar, /*bIncludeEmpty*/false);
 
-			AddTextLine(Ar);
-		}
+		AddTextLine(Ar);
+	}
 
-		if (CachedEntity.IsSet() && bMarkEntityBeingDebugged)
+	if (CachedEntity.IsSet() && bMarkEntityBeingDebugged)
+	{
+		if (const FTransformFragment* TransformFragment = EntitySubsystem->GetFragmentDataPtr<FTransformFragment>(CachedEntity))
 		{
-			if (const FTransformFragment* TransformFragment = EntitySystem->GetFragmentDataPtr<FTransformFragment>(CachedEntity))
-			{
-				const FVector Location = TransformFragment->GetTransform().GetLocation();
-				AddShape(FGameplayDebuggerShape::MakeBox(Location, FVector(8,8,500), FColor::Purple,  FString::Printf(TEXT("[%s]"), *CachedEntity.DebugGetDescription())));
-				AddShape(FGameplayDebuggerShape::MakePoint(Location, 10, FColor::Purple));
-			}
+			const FVector Location = TransformFragment->GetTransform().GetLocation();
+			AddShape(FGameplayDebuggerShape::MakeBox(Location, FVector(8,8,500), FColor::Purple,  FString::Printf(TEXT("[%s]"), *CachedEntity.DebugGetDescription())));
+			AddShape(FGameplayDebuggerShape::MakePoint(Location, 10, FColor::Purple));
 		}
 	}
-	else
-	{
-		AddTextLine(FString::Printf(TEXT("{Red}EntitySystem instance is missing")));
-	}
 
-	if (CachedEntity.IsSet())
+	if (CachedEntity.IsSet() && Debugger)
 	{
 		AddTextLine(Debugger->GetSelectedEntityInfo());
 	}
 
 	//@todo could shave off some perf cost if UMassDebuggerSubsystem used FGameplayDebuggerShape directly
-	if (bShowShapes)
+	if (bShowShapes && Debugger)
 	{
 		const TArray<UMassDebuggerSubsystem::FShapeDesc>* Shapes = Debugger->GetShapes();
 		check(Shapes);
@@ -294,11 +282,11 @@ void FGameplayDebuggerCategory_Mass::CollectData(APlayerController* OwnerPC, AAc
 
 	if (bShowAgentFragments)
 	{
-		if (CachedEntity.IsSet() && EntitySystem)
+		if (CachedEntity.IsSet() && EntitySubsystem)
 		{
 			// CachedEntity can become invalid if the entity "dies" or in editor mode when related actor gets moved 
 			// (which causes the MassAgentComponent destruction and recreation).
-			if (EntitySystem->IsEntityActive(CachedEntity))
+			if (EntitySubsystem->IsEntityActive(CachedEntity))
 			{
 				AddTextLine(FString::Printf(TEXT("{Green}Entity: {White}%s"), *CachedEntity.DebugGetDescription()));
 				AddTextLine(FString::Printf(TEXT("{Green}Type: {White}%s"), (AgentComp == nullptr) ? TEXT("N/A") : AgentComp->IsPuppet() ? TEXT("PUPPET") : TEXT("AGENT")));
@@ -308,13 +296,13 @@ void FGameplayDebuggerCategory_Mass::CollectData(APlayerController* OwnerPC, AAc
 					FStringOutputDevice FragmentsDesc;
 					FragmentsDesc.SetAutoEmitLineTerminator(true);
 					const TCHAR* PrefixToRemove = TEXT("DataFragment_");
-					FMassDebugger::OutputEntityDescription(FragmentsDesc, *EntitySystem, CachedEntity, PrefixToRemove);
+					FMassDebugger::OutputEntityDescription(FragmentsDesc, *EntitySubsystem, CachedEntity, PrefixToRemove);
 					AddTextLine(FString::Printf(TEXT("{Green}Fragments:\n{White}%s"), *FragmentsDesc));
 				}
 				else
 				{
-					const FMassArchetypeHandle Archetype = EntitySystem->GetArchetypeForEntity(CachedEntity);
-					const FMassArchetypeCompositionDescriptor& Composition = EntitySystem->GetArchetypeComposition(Archetype);
+					const FMassArchetypeHandle Archetype = EntitySubsystem->GetArchetypeForEntity(CachedEntity);
+					const FMassArchetypeCompositionDescriptor& Composition = EntitySubsystem->GetArchetypeComposition(Archetype);
 					
 					auto DescriptionBuilder = [](const TArray<FName>& ItemNames) -> FString {
 						constexpr int ColumnsCount = 2;
@@ -348,7 +336,7 @@ void FGameplayDebuggerCategory_Mass::CollectData(APlayerController* OwnerPC, AAc
 					AddTextLine(FString::Printf(TEXT("{Green}Shared Fragments:{White}%s"), *DescriptionBuilder(ItemNames)));
 				}
 
-				const FTransformFragment& TransformFragment = EntitySystem->GetFragmentDataChecked<FTransformFragment>(CachedEntity);
+				const FTransformFragment& TransformFragment = EntitySubsystem->GetFragmentDataChecked<FTransformFragment>(CachedEntity);
 				constexpr float CapsuleRadius = 50.f;
 				AddShape(FGameplayDebuggerShape::MakeCapsule(TransformFragment.GetTransform().GetLocation() + 2.f * CapsuleRadius * FVector::UpVector, CapsuleRadius, CapsuleRadius * 2.f, FColor::Orange));
 			}
@@ -392,11 +380,11 @@ void FGameplayDebuggerCategory_Mass::CollectData(APlayerController* OwnerPC, AAc
 		UMassSignalSubsystem* SignalSubsystem = World->GetSubsystem<UMassSignalSubsystem>();
 		USmartObjectSubsystem* SmartObjectSubsystem = World->GetSubsystem<USmartObjectSubsystem>();
 		
-		if (MassStateTreeSubsystem && EntitySystem && SignalSubsystem && SmartObjectSubsystem)
+		if (MassStateTreeSubsystem && EntitySubsystem && SignalSubsystem && SmartObjectSubsystem)
 		{
 			FMassExecutionContext Context(0.0f);
 		
-			EntityQuery.ForEachEntityChunk(*EntitySystem, Context, [this, MassStateTreeSubsystem, SignalSubsystem, EntitySystem, SmartObjectSubsystem, OwnerPC, ViewLocation, ViewDirection, CurrentTime](FMassExecutionContext& Context)
+			EntityQuery.ForEachEntityChunk(*EntitySubsystem, Context, [this, MassStateTreeSubsystem, SignalSubsystem, EntitySubsystem, SmartObjectSubsystem, OwnerPC, ViewLocation, ViewDirection, CurrentTime](FMassExecutionContext& Context)
 			{
 				const int32 NumEntities = Context.GetNumEntities();
 				const TConstArrayView<FMassStateTreeInstanceFragment> StateTreeInstanceList = Context.GetFragmentView<FMassStateTreeInstanceFragment>();
@@ -482,9 +470,9 @@ void FGameplayDebuggerCategory_Mass::CollectData(APlayerController* OwnerPC, AAc
 						const FMassLookAtFragment& LookAt = LookAtList[EntityIndex];
 						const FVector WorldLookDirection = Transform.GetTransform().TransformVector(LookAt.Direction);
 						bool bLookArrowDrawn = false;
-						if (LookAt.LookAtMode == EMassLookAtMode::LookAtEntity && EntitySystem->IsEntityValid(LookAt.TrackedEntity))
+						if (LookAt.LookAtMode == EMassLookAtMode::LookAtEntity && EntitySubsystem->IsEntityValid(LookAt.TrackedEntity))
 						{
-							if (const FTransformFragment* TargetTransform = EntitySystem->GetFragmentDataPtr<FTransformFragment>(LookAt.TrackedEntity))
+							if (const FTransformFragment* TargetTransform = EntitySubsystem->GetFragmentDataPtr<FTransformFragment>(LookAt.TrackedEntity))
 							{
 								FVector TargetPosition = TargetTransform->GetTransform().GetLocation();
 								TargetPosition.Z = BasePos.Z;
@@ -496,9 +484,9 @@ void FGameplayDebuggerCategory_Mass::CollectData(APlayerController* OwnerPC, AAc
 							}
 						}
 
-						if (LookAt.bRandomGazeEntities && EntitySystem->IsEntityValid(LookAt.GazeTrackedEntity))
+						if (LookAt.bRandomGazeEntities && EntitySubsystem->IsEntityValid(LookAt.GazeTrackedEntity))
 						{
-							if (const FTransformFragment* TargetTransform = EntitySystem->GetFragmentDataPtr<FTransformFragment>(LookAt.GazeTrackedEntity))
+							if (const FTransformFragment* TargetTransform = EntitySubsystem->GetFragmentDataPtr<FTransformFragment>(LookAt.GazeTrackedEntity))
 							{
 								FVector TargetPosition = TargetTransform->GetTransform().GetLocation();
 								TargetPosition.Z = BasePos.Z;
@@ -592,7 +580,7 @@ void FGameplayDebuggerCategory_Mass::CollectData(APlayerController* OwnerPC, AAc
 						// Current StateTree task
 						if (StateTree != nullptr)
 						{
-							FMassStateTreeExecutionContext StateTreeContext(*EntitySystem, *SignalSubsystem, Context);
+							FMassStateTreeExecutionContext StateTreeContext(*EntitySubsystem, *SignalSubsystem, Context);
 							StateTreeContext.Init(*OwnerPC, *StateTree, EStateTreeStorage::External);
 							StateTreeContext.SetEntity(Entity);
 							
@@ -634,13 +622,13 @@ void FGameplayDebuggerCategory_Mass::CollectData(APlayerController* OwnerPC, AAc
 			});
 		}
 
-		if (bShowNearEntityAvoidance && EntitySystem)
+		if (bShowNearEntityAvoidance && EntitySubsystem)
 		{
 			FMassEntityQuery EntityColliderQuery;
 			EntityColliderQuery.AddRequirement<FMassAvoidanceColliderFragment>(EMassFragmentAccess::ReadOnly);
 			EntityColliderQuery.AddRequirement<FTransformFragment>(EMassFragmentAccess::ReadOnly);
 			FMassExecutionContext Context(0.f);
-			EntityColliderQuery.ForEachEntityChunk(*EntitySystem, Context, [this, ViewLocation, ViewDirection](const FMassExecutionContext& Context)
+			EntityColliderQuery.ForEachEntityChunk(*EntitySubsystem, Context, [this, ViewLocation, ViewDirection](const FMassExecutionContext& Context)
 			{
 				const int32 NumEntities = Context.GetNumEntities();
 				const TConstArrayView<FTransformFragment> TransformList = Context.GetFragmentView<FTransformFragment>();
diff --git a/Engine/Plugins/AI/MassAI/Source/MassAIDebug/Public/GameplayDebuggerCategory_Mass.h b/Engine/Plugins/AI/MassAI/Source/MassAIDebug/Public/GameplayDebuggerCategory_Mass.h
--- a/Engine/Plugins/AI/MassAI/Source/MassAIDebug/Public/GameplayDebuggerCategory_Mass.h
+++ b/Engine/Plugins/AI/MassAI/Source/MassAIDebug/Public/GameplayDebuggerCategory_Mass.h
@@ -9,6 +9,7 @@
 
 #include "GameplayDebuggerCategory.h"
 
+class UMassEntitySubsystem;
 class UMassDebuggerSubsystem;
 class APlayerController;
 class AActor;
@@ -24,7 +25,7 @@ protected:
 	virtual void CollectData(APlayerController* OwnerPC, AActor* DebugActor) override;
 	virtual void DrawData(APlayerController* OwnerPC, FGameplayDebuggerCanvasContext& CanvasContext) override;
 	
-	void SetCachedEntity(FMassEntityHandle BestEntity, UMassDebuggerSubsystem& Debugger);
+	void SetCachedEntity(const FMassEntityHandle Entity, UMassEntitySubsystem& EntitySystem);
 
 	void OnToggleArchetypes() { bShowArchetypes = !bShowArchetypes; }
 	void OnToggleShapes() { bShowShapes = !bShowShapes; }
@@ -35,7 +36,7 @@ protected:
 	void OnToggleNearEntityAvoidance() { bShowNearEntityAvoidance = !bShowNearEntityAvoidance; }
 	void OnToggleNearEntityPath() { bShowNearEntityPath = !bShowNearEntityPath; }
 	
-	void PickEntity(const APlayerController& OwnerPC, const UWorld& World, UMassDebuggerSubsystem& Debugger, const bool bLimitAngle = true);
+	void PickEntity(const APlayerController& OwnerPC, const UWorld& World, UMassEntitySubsystem& EntitySystem, const bool bLimitAngle = true);
 
 protected:
 	AActor* CachedDebugActor;
diff --git a/Engine/Plugins/Runtime/MassEntity/Source/MassEntity/Private/MassDebugger.cpp b/Engine/Plugins/Runtime/MassEntity/Source/MassEntity/Private/MassDebugger.cpp
--- a/Engine/Plugins/Runtime/MassEntity/Source/MassEntity/Private/MassDebugger.cpp
+++ b/Engine/Plugins/Runtime/MassEntity/Source/MassEntity/Private/MassDebugger.cpp
@@ -200,6 +200,8 @@ namespace UE::Mass::Debug
 //----------------------------------------------------------------------//
 // FMassDebugger
 //----------------------------------------------------------------------//
+FMassDebugger::FOnEntitySelected FMassDebugger::OnEntitySelectedDelegate;
+
 TConstArrayView<FMassEntityQuery*> FMassDebugger::GetProcessorQueries(const UMassProcessor& Processor)
 {
 	return Processor.OwnedQueries;
@@ -422,4 +424,9 @@ void FMassDebugger::OutputEntityDescription(FOutputDevice& Ar, const UMassEntity
 	}
 }
 
+void FMassDebugger::SelectEntity(const UMassEntitySubsystem& EntitySubsystem, const FMassEntityHandle EntityHandle)
+{
+	OnEntitySelectedDelegate.Broadcast(EntitySubsystem, EntityHandle);
+}
+
 #endif // WITH_MASSENTITY_DEBUG
\ No newline at end of file
diff --git a/Engine/Plugins/Runtime/MassEntity/Source/MassEntity/Public/MassDebugger.h b/Engine/Plugins/Runtime/MassEntity/Source/MassEntity/Public/MassDebugger.h
--- a/Engine/Plugins/Runtime/MassEntity/Source/MassEntity/Public/MassDebugger.h
+++ b/Engine/Plugins/Runtime/MassEntity/Source/MassEntity/Public/MassDebugger.h
@@ -41,6 +41,8 @@ MASSENTITY_API extern void DebugOutputDescription(TConstArrayView<UMassProcessor
 
 struct MASSENTITY_API FMassDebugger
 {
+	DECLARE_MULTICAST_DELEGATE_TwoParams(FOnEntitySelected, const UMassEntitySubsystem&, const FMassEntityHandle);
+
 	static TConstArrayView<FMassEntityQuery*> GetProcessorQueries(const UMassProcessor& Processor);
 	/** fetches all queries registered for given Processor. Note that in order to get up to date information
 	 *  FMassEntityQuery::CacheArchetypes will be called on each query */
@@ -65,6 +67,10 @@ struct MASSENTITY_API FMassDebugger
 	static void OutputArchetypeDescription(FOutputDevice& Ar, const FMassArchetypeHandle& Archetype);
 	static void OutputEntityDescription(FOutputDevice& Ar, const UMassEntitySubsystem& EntitySubsystem, const int32 EntityIndex, const TCHAR* InPrefix = TEXT(""));
 	static void OutputEntityDescription(FOutputDevice& Ar, const UMassEntitySubsystem& EntitySubsystem, const FMassEntityHandle Entity, const TCHAR* InPrefix = TEXT(""));
+
+	static void SelectEntity(const UMassEntitySubsystem& EntitySubsystem, const FMassEntityHandle EntityHandle);
+
+	static FOnEntitySelected OnEntitySelectedDelegate;
 };
 
 #else
diff --git a/Engine/Plugins/Runtime/MassGameplay/Source/MassGameplayDebug/Private/MassDebuggerSubsystem.cpp b/Engine/Plugins/Runtime/MassGameplay/Source/MassGameplayDebug/Private/MassDebuggerSubsystem.cpp
--- a/Engine/Plugins/Runtime/MassGameplay/Source/MassGameplayDebug/Private/MassDebuggerSubsystem.cpp
+++ b/Engine/Plugins/Runtime/MassGameplay/Source/MassGameplayDebug/Private/MassDebuggerSubsystem.cpp
@@ -7,6 +7,8 @@
 #include "MassDebugVisualizationComponent.h"
 #include "MassEntitySubsystem.h"
 #include "MassDebugVisualizer.h"
+#include "MassDebugger.h"
+#include "MassEntitySubsystem.h"
 
 
 void UMassDebuggerSubsystem::Initialize(FSubsystemCollectionBase& Collection)
@@ -17,7 +19,28 @@ void UMassDebuggerSubsystem::Initialize(FSubsystemCollectionBase& Collection)
 	UMassSimulationSubsystem* SimSystem = UWorld::GetSubsystem<UMassSimulationSubsystem>(GetWorld());
 	check(SimSystem);
 	SimSystem->GetOnProcessingPhaseStarted(EMassProcessingPhase::PrePhysics).AddUObject(this, &UMassDebuggerSubsystem::OnProcessingPhaseStarted);
+
+#if WITH_MASSENTITY_DEBUG
+	OnEntitySelectedHandle = FMassDebugger::OnEntitySelectedDelegate.AddUObject(this, &UMassDebuggerSubsystem::OnEntitySelected);
+#endif // WITH_MASSENTITY_DEBUG
+}
+
+void UMassDebuggerSubsystem::Deinitialize()
+{
+#if WITH_MASSENTITY_DEBUG
+	FMassDebugger::OnEntitySelectedDelegate.Remove(OnEntitySelectedHandle);
+#endif // WITH_MASSENTITY_DEBUG
+}
+
+#if WITH_MASSENTITY_DEBUG
+void UMassDebuggerSubsystem::OnEntitySelected(const UMassEntitySubsystem& EntitySubsystem, const FMassEntityHandle EntityHandle)
+{
+	if (EntitySubsystem.GetWorld() == GetWorld())
+	{
+		SetSelectedEntity(EntityHandle);
+	}
 }
+#endif // WITH_MASSENTITY_DEBUG
 
 void UMassDebuggerSubsystem::PreTickProcessors()
 {
diff --git a/Engine/Plugins/Runtime/MassGameplay/Source/MassGameplayDebug/Public/MassDebuggerSubsystem.h b/Engine/Plugins/Runtime/MassGameplay/Source/MassGameplayDebug/Public/MassDebuggerSubsystem.h
--- a/Engine/Plugins/Runtime/MassGameplay/Source/MassGameplayDebug/Public/MassDebuggerSubsystem.h
+++ b/Engine/Plugins/Runtime/MassGameplay/Source/MassGameplayDebug/Public/MassDebuggerSubsystem.h
@@ -9,6 +9,7 @@
 
 class UMassDebugVisualizationComponent;
 class AMassDebugVisualizer;
+class UMassEntitySubsystem;
 
 UCLASS()
 class MASSGAMEPLAYDEBUG_API UMassDebuggerSubsystem : public UWorldSubsystem
@@ -46,10 +47,12 @@ public:
 protected:
 	// USubsystem BEGIN
 	virtual void Initialize(FSubsystemCollectionBase& Collection) override;
+	virtual void Deinitialize() override;
 	// USubsystem END
 	
 	void OnProcessingPhaseStarted(const float DeltaSeconds);
 	void PreTickProcessors();
+	void OnEntitySelected(const UMassEntitySubsystem& EntitySubsystem, const FMassEntityHandle EntityHandle);
 
 protected:
 	bool bCollectingData = false;
@@ -65,6 +68,8 @@ protected:
 
 	UPROPERTY(Transient)
 	AMassDebugVisualizer* DebugVisualizer;
+
+	FDelegateHandle OnEntitySelectedHandle;
 };
 
 
