Separated "entity management" and USubsystem stuff of UMassEntityManager. Now UMassEntityManager hosts a shared ref of newly created FMassEntityManager, and users are free to create their own FMassEntityManager instances to create isolated entity pools.

Updated Mass code to reflect FMassEntityManager introduction

diff --git a/Engine/Plugins/AI/MassAI/Source/MassAIBehavior/Private/MassLookAtProcessors.cpp b/Engine/Plugins/AI/MassAI/Source/MassAIBehavior/Private/MassLookAtProcessors.cpp
--- a/Engine/Plugins/AI/MassAI/Source/MassAIBehavior/Private/MassLookAtProcessors.cpp
+++ b/Engine/Plugins/AI/MassAI/Source/MassAIBehavior/Private/MassLookAtProcessors.cpp
@@ -114,13 +114,13 @@ void UMassLookAtProcessor::ConfigureQueries()
 	EntityQuery_Conditional.AddSubsystemRequirement<UZoneGraphSubsystem>(EMassFragmentAccess::ReadOnly);
 }
 
-void UMassLookAtProcessor::Execute(UMassEntitySubsystem& EntitySubsystem, FMassExecutionContext& Context)
+void UMassLookAtProcessor::Execute(FMassEntityManager& EntityManager, FMassExecutionContext& Context)
 {
 	QUICK_SCOPE_CYCLE_COUNTER(LookAtProcessor_Run);
 
 	const float CurrentTime = GetWorld()->GetTimeSeconds();
 
-	EntityQuery_Conditional.ForEachEntityChunk(EntitySubsystem, Context, [this, &EntitySubsystem, CurrentTime, World = EntitySubsystem.GetWorld()](FMassExecutionContext& Context)
+	EntityQuery_Conditional.ForEachEntityChunk(EntityManager, Context, [this, &EntityManager, CurrentTime, World = EntityManager.GetWorld()](FMassExecutionContext& Context)
 		{
 			const UMassNavigationSubsystem& MassNavSystem = Context.GetSubsystemChecked<UMassNavigationSubsystem>(World);
 			const UZoneGraphSubsystem& ZoneGraphSubsystem = Context.GetSubsystemChecked<UZoneGraphSubsystem>(World);
@@ -154,7 +154,7 @@ void UMassLookAtProcessor::Execute(UMassEntitySubsystem& EntitySubsystem, FMassE
 					const float TimeSinceUpdate = CurrentTime - LookAt.GazeStartTime;
 					if (TimeSinceUpdate >= LookAt.GazeDuration)
 					{
-						FindNewGazeTarget(MassNavSystem, EntitySubsystem, CurrentTime, TransformFragment.GetTransform(), LookAt);
+						FindNewGazeTarget(MassNavSystem, EntityManager, CurrentTime, TransformFragment.GetTransform(), LookAt);
 					}
 				}
 
@@ -187,7 +187,7 @@ void UMassLookAtProcessor::Execute(UMassEntitySubsystem& EntitySubsystem, FMassE
 					break;
 
 				case EMassLookAtMode::LookAtEntity:
-					UpdateLookAtTrackedEntity(EntitySubsystem, TransformFragment.GetTransform(), bDisplayDebug, LookAt);
+					UpdateLookAtTrackedEntity(EntityManager, TransformFragment.GetTransform(), bDisplayDebug, LookAt);
 					break;
 					
 				default:
@@ -202,7 +202,7 @@ void UMassLookAtProcessor::Execute(UMassEntitySubsystem& EntitySubsystem, FMassE
 
 					if (GazeStrength > KINDA_SMALL_NUMBER)
 					{
-						const bool bHasTarget = UpdateGazeTrackedEntity(EntitySubsystem, TransformFragment.GetTransform(), bDisplayDebug, LookAt);
+						const bool bHasTarget = UpdateGazeTrackedEntity(EntityManager, TransformFragment.GetTransform(), bDisplayDebug, LookAt);
 
 						if (bHasTarget)
 						{
@@ -233,7 +233,7 @@ void UMassLookAtProcessor::Execute(UMassEntitySubsystem& EntitySubsystem, FMassE
 		});
 }
 
-void UMassLookAtProcessor::FindNewGazeTarget(const UMassNavigationSubsystem& MassNavSystem, const UMassEntitySubsystem& EntitySubsystem, const float CurrentTime, const FTransform& Transform, FMassLookAtFragment& LookAt) const
+void UMassLookAtProcessor::FindNewGazeTarget(const UMassNavigationSubsystem& MassNavSystem, const FMassEntityManager& EntityManager, const float CurrentTime, const FTransform& Transform, FMassLookAtFragment& LookAt) const
 {
 	const FNavigationObstacleHashGrid2D& ObstacleGrid = MassNavSystem.GetObstacleGrid();
 	const FMassEntityHandle LastTrackedEntity = LookAt.GazeTrackedEntity;
@@ -262,7 +262,7 @@ void UMassLookAtProcessor::FindNewGazeTarget(const UMassNavigationSubsystem& Mas
 		for (const FNavigationObstacleHashGrid2D::ItemIDType NearbyEntity : NearbyEntities)
 		{
 			// This can happen if we remove entities in the system.
-			if (!EntitySubsystem.IsEntityValid(NearbyEntity.Entity))
+			if (!EntityManager.IsEntityValid(NearbyEntity.Entity))
 			{
 				UE_LOG(LogMassBehavior, VeryVerbose, TEXT("Nearby entity is invalid, skipped."));
 				continue;
@@ -274,7 +274,7 @@ void UMassLookAtProcessor::FindNewGazeTarget(const UMassNavigationSubsystem& Mas
 				continue;
 			}
 
-			FMassEntityView EntityView(EntitySubsystem, NearbyEntity.Entity);
+			FMassEntityView EntityView(EntityManager, NearbyEntity.Entity);
 			if (!EntityView.HasTag<FMassLookAtTargetTag>())
 			{
 				continue;
@@ -350,12 +350,12 @@ void UMassLookAtProcessor::UpdateLookAtTrajectory(const FTransform& Transform, c
 	}
 }
 
-void UMassLookAtProcessor::UpdateLookAtTrackedEntity(const UMassEntitySubsystem& EntitySubsystem, const FTransform& Transform, const bool bDisplayDebug, FMassLookAtFragment& LookAt) const
+void UMassLookAtProcessor::UpdateLookAtTrackedEntity(const FMassEntityManager& EntityManager, const FTransform& Transform, const bool bDisplayDebug, FMassLookAtFragment& LookAt) const
 {
 	// Update direction toward target
-	if (EntitySubsystem.IsEntityValid(LookAt.TrackedEntity))
+	if (EntityManager.IsEntityValid(LookAt.TrackedEntity))
 	{
-		if (const FTransformFragment* TargetTransform = EntitySubsystem.GetFragmentDataPtr<FTransformFragment>(LookAt.TrackedEntity))
+		if (const FTransformFragment* TargetTransform = EntityManager.GetFragmentDataPtr<FTransformFragment>(LookAt.TrackedEntity))
 		{
 			const FVector AgentPosition = Transform.GetLocation();
 			const FVector NewGlobalDirection = (TargetTransform->GetTransform().GetLocation() - AgentPosition).GetSafeNormal();
@@ -372,14 +372,14 @@ void UMassLookAtProcessor::UpdateLookAtTrackedEntity(const UMassEntitySubsystem&
 	}
 }
 
-bool UMassLookAtProcessor::UpdateGazeTrackedEntity(const UMassEntitySubsystem& EntitySubsystem, const FTransform& Transform, const bool bDisplayDebug, FMassLookAtFragment& LookAt) const
+bool UMassLookAtProcessor::UpdateGazeTrackedEntity(const FMassEntityManager& EntityManager, const FTransform& Transform, const bool bDisplayDebug, FMassLookAtFragment& LookAt) const
 {
 	bool bHasTarget = false;
 	
 	// Update direction toward gaze target
-	if (LookAt.GazeTrackedEntity.IsSet() && EntitySubsystem.IsEntityValid(LookAt.GazeTrackedEntity))
+	if (LookAt.GazeTrackedEntity.IsSet() && EntityManager.IsEntityValid(LookAt.GazeTrackedEntity))
 	{
-		if (const FTransformFragment* TargetTransform = EntitySubsystem.GetFragmentDataPtr<FTransformFragment>(LookAt.GazeTrackedEntity))
+		if (const FTransformFragment* TargetTransform = EntityManager.GetFragmentDataPtr<FTransformFragment>(LookAt.GazeTrackedEntity))
 		{
 			const FVector AgentPosition = Transform.GetLocation();
 			const FVector NewGlobalDirection = (TargetTransform->GetTransform().GetLocation() - AgentPosition).GetSafeNormal();
diff --git a/Engine/Plugins/AI/MassAI/Source/MassAIBehavior/Private/MassStateTreeExecutionContext.cpp b/Engine/Plugins/AI/MassAI/Source/MassAIBehavior/Private/MassStateTreeExecutionContext.cpp
--- a/Engine/Plugins/AI/MassAI/Source/MassAIBehavior/Private/MassStateTreeExecutionContext.cpp
+++ b/Engine/Plugins/AI/MassAI/Source/MassAIBehavior/Private/MassStateTreeExecutionContext.cpp
@@ -6,10 +6,10 @@
 #include "MassSignalSubsystem.h"
 #include "Engine/World.h"
 
-FMassStateTreeExecutionContext::FMassStateTreeExecutionContext(UMassEntitySubsystem& InEntitySubsystem,
+FMassStateTreeExecutionContext::FMassStateTreeExecutionContext(FMassEntityManager& InEntityManager,
                                                                UMassSignalSubsystem& InSignalSubsystem,
                                                                FMassExecutionContext& InContext):
-	EntitySubsystem(&InEntitySubsystem), SignalSubsystem(&InSignalSubsystem), EntitySubsystemExecutionContext(&InContext)
+	EntityManager(&InEntityManager), SignalSubsystem(&InSignalSubsystem), EntitySubsystemExecutionContext(&InContext)
 {
 }
 
diff --git a/Engine/Plugins/AI/MassAI/Source/MassAIBehavior/Private/MassStateTreeProcessors.cpp b/Engine/Plugins/AI/MassAI/Source/MassAIBehavior/Private/MassStateTreeProcessors.cpp
--- a/Engine/Plugins/AI/MassAI/Source/MassAIBehavior/Private/MassStateTreeProcessors.cpp
+++ b/Engine/Plugins/AI/MassAI/Source/MassAIBehavior/Private/MassStateTreeProcessors.cpp
@@ -21,10 +21,10 @@ CSV_DEFINE_CATEGORY(StateTreeProcessor, true);
 namespace UE::MassBehavior
 {
 
-bool SetExternalFragments(FMassStateTreeExecutionContext& Context, const UMassEntitySubsystem& EntitySubsystem)
+bool SetExternalFragments(FMassStateTreeExecutionContext& Context, const FMassEntityManager& EntityManager)
 {
 	bool bFoundAllFragments = true;
-	const FMassEntityView EntityView(EntitySubsystem, Context.GetEntity());
+	const FMassEntityView EntityView(EntityManager, Context.GetEntity());
 	for (const FStateTreeExternalDataDesc& DataDesc : Context.GetExternalDataDescs())
 	{
 		if (DataDesc.Struct == nullptr)
@@ -142,7 +142,7 @@ void ForEachEntityInChunk(FMassExecutionContext& Context, FMassStateTreeExecutio
 		// Gather all required fragments.
 		{
 			CSV_SCOPED_TIMING_STAT_EXCLUSIVE(StateTreeProcessorExternalFragments);
-			if (!ensureMsgf(UE::MassBehavior::SetExternalFragments(StateTreeContext, StateTreeContext.GetEntitySubsystem()), TEXT("StateTree will not execute due to missing required fragments.")))
+			if (!ensureMsgf(UE::MassBehavior::SetExternalFragments(StateTreeContext, StateTreeContext.GetEntityManager()), TEXT("StateTree will not execute due to missing required fragments.")))
 			{
 				break;
 			}
@@ -193,17 +193,17 @@ void UMassStateTreeFragmentDestructor::ConfigureQueries()
 	EntityQuery.AddSubsystemRequirement<UMassStateTreeSubsystem>(EMassFragmentAccess::ReadWrite);
 }
 
-void UMassStateTreeFragmentDestructor::Execute(UMassEntitySubsystem& EntitySubsystem, FMassExecutionContext& Context)
+void UMassStateTreeFragmentDestructor::Execute(FMassEntityManager& EntityManager, FMassExecutionContext& Context)
 {
 	if (SignalSubsystem == nullptr)
 	{
 		return;
 	}
 		
-	FMassStateTreeExecutionContext StateTreeContext(EntitySubsystem, *SignalSubsystem, Context);
+	FMassStateTreeExecutionContext StateTreeContext(EntityManager, *SignalSubsystem, Context);
 
-	EntityQuery.ForEachEntityChunk(EntitySubsystem,Context,
-		[this, &StateTreeContext, World = EntitySubsystem.GetWorld()](FMassExecutionContext& Context)
+	EntityQuery.ForEachEntityChunk(EntityManager,Context,
+		[this, &StateTreeContext, World = EntityManager.GetWorld()](FMassExecutionContext& Context)
 		{
 			UMassStateTreeSubsystem& MassStateTreeSubsystem = Context.GetMutableSubsystemChecked<UMassStateTreeSubsystem>(World);
 			const TArrayView<FMassStateTreeInstanceFragment> StateTreeInstanceList = Context.GetMutableFragmentView<FMassStateTreeInstanceFragment>();
@@ -253,9 +253,9 @@ void UMassStateTreeActivationProcessor::ConfigureQueries()
 	ProcessorRequirements.AddSubsystemRequirement<UMassSignalSubsystem>(EMassFragmentAccess::ReadWrite);
 }
 
-void UMassStateTreeActivationProcessor::Execute(UMassEntitySubsystem& EntitySubsystem, FMassExecutionContext& Context)
+void UMassStateTreeActivationProcessor::Execute(FMassEntityManager& EntityManager, FMassExecutionContext& Context)
 {
-	UMassSignalSubsystem& SignalSubsystem = Context.GetMutableSubsystemChecked<UMassSignalSubsystem>(EntitySubsystem.GetWorld());
+	UMassSignalSubsystem& SignalSubsystem = Context.GetMutableSubsystemChecked<UMassSignalSubsystem>(EntityManager.GetWorld());
 
 	const UMassBehaviorSettings* BehaviorSettings = GetDefault<UMassBehaviorSettings>();
 	check(BehaviorSettings);
@@ -263,15 +263,15 @@ void UMassStateTreeActivationProcessor::Execute(UMassEntitySubsystem& EntitySubs
 	// StateTree processor relies on signals to be ticked but we need an 'initial tick' to set the tree in the proper state.
 	// The initializer provides that by sending a signal to all new entities that use StateTree.
 
-	FMassStateTreeExecutionContext StateTreeContext(EntitySubsystem, SignalSubsystem, Context);
+	FMassStateTreeExecutionContext StateTreeContext(EntityManager, SignalSubsystem, Context);
 
-	const float TimeInSeconds = EntitySubsystem.GetWorld()->GetTimeSeconds();
+	const float TimeInSeconds = EntityManager.GetWorld()->GetTimeSeconds();
 
 	TArray<FMassEntityHandle> EntitiesToSignal;
 	int32 ActivationCounts[EMassLOD::Max] {0,0,0,0};
 	
-	EntityQuery.ForEachEntityChunk(EntitySubsystem, Context,
-		[&EntitiesToSignal, &ActivationCounts, MaxActivationsPerLOD = BehaviorSettings->MaxActivationsPerLOD, &StateTreeContext, TimeInSeconds, World = EntitySubsystem.GetWorld()](FMassExecutionContext& Context)
+	EntityQuery.ForEachEntityChunk(EntityManager, Context,
+		[&EntitiesToSignal, &ActivationCounts, MaxActivationsPerLOD = BehaviorSettings->MaxActivationsPerLOD, &StateTreeContext, TimeInSeconds, World = EntityManager.GetWorld()](FMassExecutionContext& Context)
 		{
 			UMassStateTreeSubsystem& MassStateTreeSubsystem = Context.GetMutableSubsystemChecked<UMassStateTreeSubsystem>(World);
 			const int32 NumEntities = Context.GetNumEntities();
@@ -379,20 +379,20 @@ void UMassStateTreeProcessor::ConfigureQueries()
 	ProcessorRequirements.AddSubsystemRequirement<UMassSignalSubsystem>(EMassFragmentAccess::ReadWrite);
 }
 
-void UMassStateTreeProcessor::SignalEntities(UMassEntitySubsystem& EntitySubsystem, FMassExecutionContext& Context, FMassSignalNameLookup& EntitySignals)
+void UMassStateTreeProcessor::SignalEntities(FMassEntityManager& EntityManager, FMassExecutionContext& Context, FMassSignalNameLookup& EntitySignals)
 {
-	UWorld* World = EntitySubsystem.GetWorld();
+	UWorld* World = EntityManager.GetWorld();
 	UMassSignalSubsystem& SignalSubsystem = Context.GetMutableSubsystemChecked<UMassSignalSubsystem>(World);
 	
 	QUICK_SCOPE_CYCLE_COUNTER(StateTreeProcessor_Run);
 	CSV_SCOPED_TIMING_STAT_EXCLUSIVE(StateTreeProcessorExecute);
 
-	const float TimeInSeconds = EntitySubsystem.GetWorld()->GetTimeSeconds();
-	FMassStateTreeExecutionContext StateTreeContext(EntitySubsystem, SignalSubsystem, Context);
+	const float TimeInSeconds = EntityManager.GetWorld()->GetTimeSeconds();
+	FMassStateTreeExecutionContext StateTreeContext(EntityManager, SignalSubsystem, Context);
 
 	TArray<FMassEntityHandle> EntitiesToSignal;
 
-	EntityQuery.ForEachEntityChunk(EntitySubsystem, Context,
+	EntityQuery.ForEachEntityChunk(EntityManager, Context,
 		[this, &StateTreeContext, TimeInSeconds, &EntitiesToSignal, World, &EntitySignals](FMassExecutionContext& Context)
 		{
 			// Keep stats regarding the amount of tree instances ticked per frame
diff --git a/Engine/Plugins/AI/MassAI/Source/MassAIBehavior/Private/MassStateTreeTrait.cpp b/Engine/Plugins/AI/MassAI/Source/MassAIBehavior/Private/MassStateTreeTrait.cpp
--- a/Engine/Plugins/AI/MassAI/Source/MassAIBehavior/Private/MassStateTreeTrait.cpp
+++ b/Engine/Plugins/AI/MassAI/Source/MassAIBehavior/Private/MassStateTreeTrait.cpp
@@ -8,11 +8,12 @@
 #include "VisualLogger/VisualLogger.h"
 #include "MassEntityTemplateRegistry.h"
 #include "Engine/World.h"
+#include "MassEntityUtils.h"
+
 
 void UMassStateTreeTrait::BuildTemplate(FMassEntityTemplateBuildContext& BuildContext, const UWorld& World) const
 {
-	UMassEntitySubsystem* EntitySubsystem = UWorld::GetSubsystem<UMassEntitySubsystem>(&World);
-	check(EntitySubsystem);
+	FMassEntityManager& EntityManager = UE::Mass::Utils::GetEntityManagerChecked(World);
 
 	UMassStateTreeSubsystem* MassStateTreeSubsystem = World.GetSubsystem<UMassStateTreeSubsystem>();
 	if (!MassStateTreeSubsystem)
@@ -35,7 +36,7 @@ void UMassStateTreeTrait::BuildTemplate(FMassEntityTemplateBuildContext& BuildCo
 	FMassStateTreeSharedFragment SharedStateTree;
 	SharedStateTree.StateTree = StateTree;
 	
-	const FConstSharedStruct StateTreeFragment = EntitySubsystem->GetOrCreateConstSharedFragment(UE::StructUtils::GetStructCrc32(FConstStructView::Make(SharedStateTree)), SharedStateTree);
+	const FConstSharedStruct StateTreeFragment = EntityManager.GetOrCreateConstSharedFragment(UE::StructUtils::GetStructCrc32(FConstStructView::Make(SharedStateTree)), SharedStateTree);
 	BuildContext.AddConstSharedFragment(StateTreeFragment);
 
 	BuildContext.AddFragment<FMassStateTreeInstanceFragment>();
diff --git a/Engine/Plugins/AI/MassAI/Source/MassAIBehavior/Private/MassZoneGraphAnnotationProcessors.cpp b/Engine/Plugins/AI/MassAI/Source/MassAIBehavior/Private/MassZoneGraphAnnotationProcessors.cpp
--- a/Engine/Plugins/AI/MassAI/Source/MassAIBehavior/Private/MassZoneGraphAnnotationProcessors.cpp
+++ b/Engine/Plugins/AI/MassAI/Source/MassAIBehavior/Private/MassZoneGraphAnnotationProcessors.cpp
@@ -29,9 +29,9 @@ void UMassZoneGraphAnnotationTagsInitializer::ConfigureQueries()
 	EntityQuery.AddSubsystemRequirement<UZoneGraphAnnotationSubsystem>(EMassFragmentAccess::ReadOnly);
 }
 
-void UMassZoneGraphAnnotationTagsInitializer::Execute(UMassEntitySubsystem& EntitySubsystem, FMassExecutionContext& Context)
+void UMassZoneGraphAnnotationTagsInitializer::Execute(FMassEntityManager& EntityManager, FMassExecutionContext& Context)
 {
-	EntityQuery.ForEachEntityChunk(EntitySubsystem, Context, [this, World = EntitySubsystem.GetWorld()](FMassExecutionContext& Context)
+	EntityQuery.ForEachEntityChunk(EntityManager, Context, [this, World = EntityManager.GetWorld()](FMassExecutionContext& Context)
 	{
 		const UZoneGraphAnnotationSubsystem& ZoneGraphAnnotationSubsystem = Context.GetSubsystemChecked<UZoneGraphAnnotationSubsystem>(World);
 
@@ -86,15 +86,15 @@ void UMassZoneGraphAnnotationTagUpdateProcessor::ConfigureQueries()
 	ProcessorRequirements.AddSubsystemRequirement<UMassSignalSubsystem>(EMassFragmentAccess::ReadWrite);
 }
 
-void UMassZoneGraphAnnotationTagUpdateProcessor::Execute(UMassEntitySubsystem& EntitySubsystem, FMassExecutionContext& Context)
+void UMassZoneGraphAnnotationTagUpdateProcessor::Execute(FMassEntityManager& EntityManager, FMassExecutionContext& Context)
 {
 	TransientEntitiesToSignal.Reset();
 
 	// Calling super will update the signals, and call SignalEntities() below.
-	Super::Execute(EntitySubsystem, Context);
+	Super::Execute(EntityManager, Context);
 
-	UWorld* World = EntitySubsystem.GetWorld();
-	EntityQuery.ForEachEntityChunk(EntitySubsystem, Context, [this, World](FMassExecutionContext& Context)
+	UWorld* World = EntityManager.GetWorld();
+	EntityQuery.ForEachEntityChunk(EntityManager, Context, [this, World](FMassExecutionContext& Context)
 	{
 		// Periodically update tags.
 		if (!FMassZoneGraphAnnotationVariableTickChunkFragment::UpdateChunk(Context))
@@ -143,9 +143,9 @@ void UMassZoneGraphAnnotationTagUpdateProcessor::UpdateAnnotationTags(UZoneGraph
 	}
 }
 
-void UMassZoneGraphAnnotationTagUpdateProcessor::SignalEntities(UMassEntitySubsystem& EntitySubsystem, FMassExecutionContext& Context, FMassSignalNameLookup& EntitySignals)
+void UMassZoneGraphAnnotationTagUpdateProcessor::SignalEntities(FMassEntityManager& EntityManager, FMassExecutionContext& Context, FMassSignalNameLookup& EntitySignals)
 {
-	EntityQuery.ForEachEntityChunk(EntitySubsystem, Context, [this, World = EntitySubsystem.GetWorld()](FMassExecutionContext& Context)
+	EntityQuery.ForEachEntityChunk(EntityManager, Context, [this, World = EntityManager.GetWorld()](FMassExecutionContext& Context)
 	{
 		UZoneGraphAnnotationSubsystem& ZoneGraphAnnotationSubsystem = Context.GetMutableSubsystemChecked<UZoneGraphAnnotationSubsystem>(World);
 
diff --git a/Engine/Plugins/AI/MassAI/Source/MassAIBehavior/Private/Tasks/MassClaimSmartObjectTask.cpp b/Engine/Plugins/AI/MassAI/Source/MassAIBehavior/Private/Tasks/MassClaimSmartObjectTask.cpp
--- a/Engine/Plugins/AI/MassAI/Source/MassAIBehavior/Private/Tasks/MassClaimSmartObjectTask.cpp
+++ b/Engine/Plugins/AI/MassAI/Source/MassAIBehavior/Private/Tasks/MassClaimSmartObjectTask.cpp
@@ -1,4 +1,4 @@
-﻿// Copyright Epic Games, Inc. All Rights Reserved.
+// Copyright Epic Games, Inc. All Rights Reserved.
 
 #include "Tasks/MassClaimSmartObjectTask.h"
 #include "MassCommonFragments.h"
@@ -51,7 +51,7 @@ EStateTreeRunStatus FMassClaimSmartObjectTask::EnterState(FStateTreeExecutionCon
 	
 	// Setup MassSmartObject handler and claim
 	const FMassStateTreeExecutionContext& MassContext = static_cast<FMassStateTreeExecutionContext&>(Context);
-	const FMassSmartObjectHandler MassSmartObjectHandler(MassContext.GetEntitySubsystem(), MassContext.GetEntitySubsystemExecutionContext(), SmartObjectSubsystem, SignalSubsystem);
+	const FMassSmartObjectHandler MassSmartObjectHandler(MassContext.GetEntityManager(), MassContext.GetEntitySubsystemExecutionContext(), SmartObjectSubsystem, SignalSubsystem);
 	
 	InstanceData.ClaimedSlot = MassSmartObjectHandler.ClaimCandidate(MassContext.GetEntity(), SOUser, *CandidateSlots);
 
@@ -87,7 +87,7 @@ void FMassClaimSmartObjectTask::ExitState(FStateTreeExecutionContext& Context, c
 		const FMassStateTreeExecutionContext& MassContext = static_cast<FMassStateTreeExecutionContext&>(Context);
 		USmartObjectSubsystem& SmartObjectSubsystem = Context.GetExternalData(SmartObjectSubsystemHandle);
 		UMassSignalSubsystem& SignalSubsystem = Context.GetExternalData(MassSignalSubsystemHandle);
-		const FMassSmartObjectHandler MassSmartObjectHandler(MassContext.GetEntitySubsystem(), MassContext.GetEntitySubsystemExecutionContext(), SmartObjectSubsystem, SignalSubsystem);
+		const FMassSmartObjectHandler MassSmartObjectHandler(MassContext.GetEntityManager(), MassContext.GetEntitySubsystemExecutionContext(), SmartObjectSubsystem, SignalSubsystem);
 
 		MassSmartObjectHandler.ReleaseSmartObject(MassContext.GetEntity(), SOUser, InstanceData.ClaimedSlot);
 	}
diff --git a/Engine/Plugins/AI/MassAI/Source/MassAIBehavior/Private/Tasks/MassFindSmartObjectTask.cpp b/Engine/Plugins/AI/MassAI/Source/MassAIBehavior/Private/Tasks/MassFindSmartObjectTask.cpp
--- a/Engine/Plugins/AI/MassAI/Source/MassAIBehavior/Private/Tasks/MassFindSmartObjectTask.cpp
+++ b/Engine/Plugins/AI/MassAI/Source/MassAIBehavior/Private/Tasks/MassFindSmartObjectTask.cpp
@@ -1,4 +1,4 @@
-﻿// Copyright Epic Games, Inc. All Rights Reserved.
+// Copyright Epic Games, Inc. All Rights Reserved.
 
 #include "Tasks/MassFindSmartObjectTask.h"
 #include "ZoneGraphAnnotationSubsystem.h"
@@ -47,7 +47,7 @@ void FMassFindSmartObjectTask::ExitState(FStateTreeExecutionContext& Context, co
 		USmartObjectSubsystem& SmartObjectSubsystem = Context.GetExternalData(SmartObjectSubsystemHandle);
 		UMassSignalSubsystem& SignalSubsystem = Context.GetExternalData(MassSignalSubsystemHandle);
 		const FMassSmartObjectHandler MassSmartObjectHandler(
-			MassContext.GetEntitySubsystem(),
+			MassContext.GetEntityManager(),
 			MassContext.GetEntitySubsystemExecutionContext(),
 			SmartObjectSubsystem,
 			SignalSubsystem);
@@ -88,7 +88,7 @@ EStateTreeRunStatus FMassFindSmartObjectTask::Tick(FStateTreeExecutionContext& C
 	UMassSignalSubsystem& SignalSubsystem = Context.GetExternalData(MassSignalSubsystemHandle);
 	const FMassStateTreeExecutionContext& MassContext = static_cast<FMassStateTreeExecutionContext&>(Context);
 	const FMassSmartObjectHandler MassSmartObjectHandler(
-		MassContext.GetEntitySubsystem(),
+		MassContext.GetEntityManager(),
 		MassContext.GetEntitySubsystemExecutionContext(),
 		SmartObjectSubsystem,
 		SignalSubsystem);
diff --git a/Engine/Plugins/AI/MassAI/Source/MassAIBehavior/Private/Tasks/MassUseSmartObjectTask.cpp b/Engine/Plugins/AI/MassAI/Source/MassAIBehavior/Private/Tasks/MassUseSmartObjectTask.cpp
--- a/Engine/Plugins/AI/MassAI/Source/MassAIBehavior/Private/Tasks/MassUseSmartObjectTask.cpp
+++ b/Engine/Plugins/AI/MassAI/Source/MassAIBehavior/Private/Tasks/MassUseSmartObjectTask.cpp
@@ -53,7 +53,7 @@ EStateTreeRunStatus FMassUseSmartObjectTask::EnterState(FStateTreeExecutionConte
 
 	// Setup MassSmartObject handler and start interaction
 	const FMassStateTreeExecutionContext& MassContext = static_cast<FMassStateTreeExecutionContext&>(Context);
-	const FMassSmartObjectHandler MassSmartObjectHandler(MassContext.GetEntitySubsystem(), MassContext.GetEntitySubsystemExecutionContext(), SmartObjectSubsystem, SignalSubsystem);
+	const FMassSmartObjectHandler MassSmartObjectHandler(MassContext.GetEntityManager(), MassContext.GetEntitySubsystemExecutionContext(), SmartObjectSubsystem, SignalSubsystem);
 
 	if (!MassSmartObjectHandler.StartUsingSmartObject(MassContext.GetEntity(), SOUser, ClaimedSlot))
 	{
@@ -83,7 +83,7 @@ void FMassUseSmartObjectTask::ExitState(FStateTreeExecutionContext& Context, con
 			const FMassStateTreeExecutionContext& MassContext = static_cast<FMassStateTreeExecutionContext&>(Context);
 			USmartObjectSubsystem& SmartObjectSubsystem = Context.GetExternalData(SmartObjectSubsystemHandle);
 			UMassSignalSubsystem& SignalSubsystem = Context.GetExternalData(MassSignalSubsystemHandle);
-			const FMassSmartObjectHandler MassSmartObjectHandler(MassContext.GetEntitySubsystem(), MassContext.GetEntitySubsystemExecutionContext(), SmartObjectSubsystem, SignalSubsystem);
+			const FMassSmartObjectHandler MassSmartObjectHandler(MassContext.GetEntityManager(), MassContext.GetEntitySubsystemExecutionContext(), SmartObjectSubsystem, SignalSubsystem);
 			MassSmartObjectHandler.StopUsingSmartObject(MassContext.GetEntity(), SOUser, EMassSmartObjectInteractionStatus::Aborted);
 		}
 		else
@@ -104,7 +104,7 @@ void FMassUseSmartObjectTask::StateCompleted(FStateTreeExecutionContext& Context
 		const FMassStateTreeExecutionContext& MassContext = static_cast<FMassStateTreeExecutionContext&>(Context);
 		USmartObjectSubsystem& SmartObjectSubsystem = Context.GetExternalData(SmartObjectSubsystemHandle);
 		UMassSignalSubsystem& SignalSubsystem = Context.GetExternalData(MassSignalSubsystemHandle);
-		const FMassSmartObjectHandler MassSmartObjectHandler(MassContext.GetEntitySubsystem(), MassContext.GetEntitySubsystemExecutionContext(), SmartObjectSubsystem, SignalSubsystem);
+		const FMassSmartObjectHandler MassSmartObjectHandler(MassContext.GetEntityManager(), MassContext.GetEntitySubsystemExecutionContext(), SmartObjectSubsystem, SignalSubsystem);
 		const EMassSmartObjectInteractionStatus NewStatus = (CompletionStatus == EStateTreeRunStatus::Succeeded)
 																? EMassSmartObjectInteractionStatus::TaskCompleted
 																: EMassSmartObjectInteractionStatus::Aborted;
diff --git a/Engine/Plugins/AI/MassAI/Source/MassAIBehavior/Public/MassLookAtProcessors.h b/Engine/Plugins/AI/MassAI/Source/MassAIBehavior/Public/MassLookAtProcessors.h
--- a/Engine/Plugins/AI/MassAI/Source/MassAIBehavior/Public/MassLookAtProcessors.h
+++ b/Engine/Plugins/AI/MassAI/Source/MassAIBehavior/Public/MassLookAtProcessors.h
@@ -28,20 +28,20 @@ class MASSAIBEHAVIOR_API UMassLookAtProcessor : public UMassProcessor
 protected:
 
 	virtual void ConfigureQueries() override;
-	virtual void Execute(UMassEntitySubsystem& EntitySubsystem, FMassExecutionContext& Context) override;
+	virtual void Execute(FMassEntityManager& EntityManager, FMassExecutionContext& Context) override;
 
 	/** Selects a nearby target if possible or use a random fixed direction */
-	void FindNewGazeTarget(const UMassNavigationSubsystem& MassNavSystem, const UMassEntitySubsystem& EntitySubsystem, const float CurrentTime, const FTransform& Transform, FMassLookAtFragment& LookAt) const;
+	void FindNewGazeTarget(const UMassNavigationSubsystem& MassNavSystem, const FMassEntityManager& EntitySubsystem, const float CurrentTime, const FTransform& Transform, FMassLookAtFragment& LookAt) const;
 
 	/** Updates look direction based on look at trajectory. */
 	void UpdateLookAtTrajectory(const FTransform& Transform, const FMassZoneGraphLaneLocationFragment& ZoneGraphLocation,
 								const FMassLookAtTrajectoryFragment& LookAtTrajectory, const bool bDisplayDebug, FMassLookAtFragment& LookAt) const;
 
 	/** Updates look at based on tracked entity. */
-	void UpdateLookAtTrackedEntity(const UMassEntitySubsystem& EntitySubsystem, const FTransform& Transform, const bool bDisplayDebug, FMassLookAtFragment& LookAt) const;
+	void UpdateLookAtTrackedEntity(const FMassEntityManager& EntitySubsystem, const FTransform& Transform, const bool bDisplayDebug, FMassLookAtFragment& LookAt) const;
 
 	/** Updates gaze based on tracked entity. */
-	bool UpdateGazeTrackedEntity(const UMassEntitySubsystem& EntitySubsystem, const FTransform& Transform, const bool bDisplayDebug, FMassLookAtFragment& LookAt) const;
+	bool UpdateGazeTrackedEntity(const FMassEntityManager& EntitySubsystem, const FTransform& Transform, const bool bDisplayDebug, FMassLookAtFragment& LookAt) const;
 
 	/** Builds look at trajectory along the current path. */
 	void BuildTrajectory(const UZoneGraphSubsystem& ZoneGraphSubsystem, const FMassZoneGraphLaneLocationFragment& LaneLocation, const FMassZoneGraphShortPathFragment& ShortPath,
diff --git a/Engine/Plugins/AI/MassAI/Source/MassAIBehavior/Public/MassStateTreeExecutionContext.h b/Engine/Plugins/AI/MassAI/Source/MassAIBehavior/Public/MassStateTreeExecutionContext.h
--- a/Engine/Plugins/AI/MassAI/Source/MassAIBehavior/Public/MassStateTreeExecutionContext.h
+++ b/Engine/Plugins/AI/MassAI/Source/MassAIBehavior/Public/MassStateTreeExecutionContext.h
@@ -7,7 +7,7 @@
 #include "MassStateTreeExecutionContext.generated.h"
 
 struct FMassExecutionContext;
-class UMassEntitySubsystem;
+struct FMassEntityManager;
 class UMassSignalSubsystem;
 
 /**
@@ -20,9 +20,9 @@ struct MASSAIBEHAVIOR_API FMassStateTreeExecutionContext : public FStateTreeExec
 public:
 	/** Should never be used but has to be public for 'void UScriptStruct::TCppStructOps<FMassStateTreeExecutionContext>::ConstructForTests(void *)' */
 	FMassStateTreeExecutionContext() = default;
-	FMassStateTreeExecutionContext(UMassEntitySubsystem& InEntitySubsystem, UMassSignalSubsystem& InSignalSubsystem, FMassExecutionContext& InContext);
+	FMassStateTreeExecutionContext(FMassEntityManager& InEntityManager, UMassSignalSubsystem& InSignalSubsystem, FMassExecutionContext& InContext);
 
-	UMassEntitySubsystem& GetEntitySubsystem() const { return *EntitySubsystem; }
+	FMassEntityManager& GetEntityManager() const { check(EntityManager); return *EntityManager; }
 	FMassExecutionContext& GetEntitySubsystemExecutionContext() const { return *EntitySubsystemExecutionContext; }
 
 	FMassEntityHandle GetEntity() const { return Entity; }
@@ -35,7 +35,7 @@ protected:
 
 	virtual void BeginGatedTransition(const FStateTreeExecutionState& Exec) override;
 
-	UMassEntitySubsystem* EntitySubsystem = nullptr;
+	FMassEntityManager* EntityManager;
 	UMassSignalSubsystem* SignalSubsystem = nullptr;
 	FMassExecutionContext* EntitySubsystemExecutionContext = nullptr;
 	FMassEntityHandle Entity;
diff --git a/Engine/Plugins/AI/MassAI/Source/MassAIBehavior/Public/MassStateTreeProcessors.h b/Engine/Plugins/AI/MassAI/Source/MassAIBehavior/Public/MassStateTreeProcessors.h
--- a/Engine/Plugins/AI/MassAI/Source/MassAIBehavior/Public/MassStateTreeProcessors.h
+++ b/Engine/Plugins/AI/MassAI/Source/MassAIBehavior/Public/MassStateTreeProcessors.h
@@ -24,7 +24,7 @@ public:
 protected:
 	virtual void Initialize(UObject& Owner) override;
 	virtual void ConfigureQueries() override;
-	virtual void Execute(UMassEntitySubsystem& EntitySubsystem, FMassExecutionContext& Context) override;
+	virtual void Execute(FMassEntityManager& EntityManager, FMassExecutionContext& Context) override;
 
 	FMassEntityQuery EntityQuery;
 
@@ -50,7 +50,7 @@ public:
 	UMassStateTreeActivationProcessor();
 protected:
 	virtual void ConfigureQueries() override;
-	virtual void Execute(UMassEntitySubsystem& EntitySubsystem, FMassExecutionContext& Context) override;
+	virtual void Execute(FMassEntityManager& EntityManager, FMassExecutionContext& Context) override;
 	
 	FMassEntityQuery EntityQuery;
 };
@@ -69,5 +69,5 @@ public:
 protected:
 	virtual void Initialize(UObject& Owner) override;
 	virtual void ConfigureQueries() override;
-	virtual void SignalEntities(UMassEntitySubsystem& EntitySubsystem, FMassExecutionContext& Context, FMassSignalNameLookup& EntitySignals) override;
+	virtual void SignalEntities(FMassEntityManager& EntityManager, FMassExecutionContext& Context, FMassSignalNameLookup& EntitySignals) override;
 };
diff --git a/Engine/Plugins/AI/MassAI/Source/MassAIBehavior/Public/MassZoneGraphAnnotationProcessors.h b/Engine/Plugins/AI/MassAI/Source/MassAIBehavior/Public/MassZoneGraphAnnotationProcessors.h
--- a/Engine/Plugins/AI/MassAI/Source/MassAIBehavior/Public/MassZoneGraphAnnotationProcessors.h
+++ b/Engine/Plugins/AI/MassAI/Source/MassAIBehavior/Public/MassZoneGraphAnnotationProcessors.h
@@ -25,7 +25,7 @@ public:
 
 protected:
 	virtual void ConfigureQueries() override;
-	virtual void Execute(UMassEntitySubsystem& EntitySubsystem, FMassExecutionContext& Context) override;
+	virtual void Execute(FMassEntityManager& EntityManager, FMassExecutionContext& Context) override;
 
 	FMassEntityQuery EntityQuery;
 };
@@ -44,11 +44,11 @@ public:
 protected:
 	virtual void Initialize(UObject& Owner) override;
 	virtual void ConfigureQueries() override;
-	virtual void Execute(UMassEntitySubsystem& EntitySubsystem, FMassExecutionContext& Context) override;
+	virtual void Execute(FMassEntityManager& EntityManager, FMassExecutionContext& Context) override;
 
 	void UpdateAnnotationTags(UZoneGraphAnnotationSubsystem& ZoneGraphAnnotationSubsystem, FMassZoneGraphAnnotationFragment& AnnotationTags, const FMassZoneGraphLaneLocationFragment& LaneLocation, FMassEntityHandle Entity);
 
-	virtual void SignalEntities(UMassEntitySubsystem& EntitySubsystem, FMassExecutionContext& Context, FMassSignalNameLookup& EntitySignals) override;
+	virtual void SignalEntities(FMassEntityManager& EntityManager, FMassExecutionContext& Context, FMassSignalNameLookup& EntitySignals) override;
 
 	// Frame buffer, it gets reset every frame.
 	TArray<FMassEntityHandle> TransientEntitiesToSignal;
diff --git a/Engine/Plugins/AI/MassAI/Source/MassAIDebug/Private/GameplayDebuggerCategory_Mass.cpp b/Engine/Plugins/AI/MassAI/Source/MassAIDebug/Private/GameplayDebuggerCategory_Mass.cpp
--- a/Engine/Plugins/AI/MassAI/Source/MassAIDebug/Private/GameplayDebuggerCategory_Mass.cpp
+++ b/Engine/Plugins/AI/MassAI/Source/MassAIDebug/Private/GameplayDebuggerCategory_Mass.cpp
@@ -111,13 +111,13 @@ FGameplayDebuggerCategory_Mass::FGameplayDebuggerCategory_Mass()
 	BindKeyPress(EKeys::C.GetFName(), FGameplayDebuggerInputModifier::Shift, this, &FGameplayDebuggerCategory_Mass::OnToggleNearEntityPath, EGameplayDebuggerInputMode::Replicated);
 }
 
-void FGameplayDebuggerCategory_Mass::SetCachedEntity(const FMassEntityHandle Entity, UMassEntitySubsystem& EntitySubsystem)
+void FGameplayDebuggerCategory_Mass::SetCachedEntity(const FMassEntityHandle Entity, const FMassEntityManager& EntityManager)
 {
 	CachedEntity = Entity;
-	FMassDebugger::SelectEntity(EntitySubsystem, Entity);
+	FMassDebugger::SelectEntity(EntityManager, Entity);
 }
 
-void FGameplayDebuggerCategory_Mass::PickEntity(const APlayerController& OwnerPC, const UWorld& World, UMassEntitySubsystem& EntitySubsystem, const bool bLimitAngle)
+void FGameplayDebuggerCategory_Mass::PickEntity(const APlayerController& OwnerPC, const UWorld& World, FMassEntityManager& EntityManager, const bool bLimitAngle)
 {
 	FVector ViewLocation = FVector::ZeroVector;
 	FVector ViewDirection = FVector::ForwardVector;
@@ -129,7 +129,7 @@ void FGameplayDebuggerCategory_Mass::PickEntity(const APlayerController& OwnerPC
     {
 		TArray<FMassEntityHandle> Entities;
 	    TArray<FVector> Locations;
-	    UE::Mass::Debug::GetDebugEntitiesAndLocations(EntitySubsystem, Entities, Locations);
+	    UE::Mass::Debug::GetDebugEntitiesAndLocations(EntityManager, Entities, Locations);
 	    BestEntity = UE::Mass::Debug::GetBestEntity(ViewLocation, ViewDirection, Entities, Locations, bLimitAngle);
     }
 	else
@@ -139,7 +139,7 @@ void FGameplayDebuggerCategory_Mass::PickEntity(const APlayerController& OwnerPC
 		FMassExecutionContext ExecutionContext;
 		FMassEntityQuery Query;
 		Query.AddRequirement<FTransformFragment>(EMassFragmentAccess::ReadOnly);
-		Query.ForEachEntityChunk(EntitySubsystem, ExecutionContext, [&Entities, &Locations](FMassExecutionContext& Context)
+		Query.ForEachEntityChunk(EntityManager, ExecutionContext, [&Entities, &Locations](FMassExecutionContext& Context)
 		{
 			Entities.Append(Context.GetEntities().GetData(), Context.GetEntities().Num());
 			TConstArrayView<FTransformFragment> InLocations = Context.GetFragmentView<FTransformFragment>();
@@ -162,7 +162,7 @@ void FGameplayDebuggerCategory_Mass::PickEntity(const APlayerController& OwnerPC
 		}
 	}
 
-	SetCachedEntity(BestEntity, EntitySubsystem);
+	SetCachedEntity(BestEntity, EntityManager);
 	CachedDebugActor = BestActor;
 	GetReplicator()->SetDebugActor(BestActor);
 }
@@ -183,19 +183,19 @@ void FGameplayDebuggerCategory_Mass::CollectData(APlayerController* OwnerPC, AAc
 		AddTextLine(FString::Printf(TEXT("{Red}EntitySubsystem instance is missing")));
 		return;
 	}
-
+	FMassEntityManager& EntityManager = EntitySubsystem->GetMutableEntityManager();
 	UMassDebuggerSubsystem* Debugger = World->GetSubsystem<UMassDebuggerSubsystem>();
 
 	const UMassAgentComponent* AgentComp = nullptr;
 	if (DebugActor)
 	{
 		const FMassEntityHandle EntityHandle = UE::Mass::Debug::GetEntityFromActor(*DebugActor, AgentComp);	
-		SetCachedEntity(EntityHandle, *EntitySubsystem);
+		SetCachedEntity(EntityHandle, EntityManager);
 		CachedDebugActor = DebugActor;
 	}
 	else if (CachedDebugActor)
 	{
-		SetCachedEntity(FMassEntityHandle(), *EntitySubsystem);
+		SetCachedEntity(FMassEntityHandle(), EntityManager);
 		CachedDebugActor = nullptr;
 	}
 
@@ -205,19 +205,19 @@ void FGameplayDebuggerCategory_Mass::CollectData(APlayerController* OwnerPC, AAc
 		// also support client-server picking. For now, we handle explicit mass picking requests on the authority
 		if (bPickEntity)
 		{
-			PickEntity(*OwnerPC, *World, *EntitySubsystem);
+			PickEntity(*OwnerPC, *World, EntityManager);
 			bPickEntity = false;
 		}
 		// if we're debugging based on UE::Mass::Debug and the range changed
 		else if (CachedDebugActor == nullptr && UE::Mass::Debug::HasDebugEntities() && UE::Mass::Debug::IsDebuggingEntity(CachedEntity) == false)
 		{
 			// using bLimitAngle = false to not limit the selection to only the things in from of the player
-			PickEntity(*OwnerPC, *World, *EntitySubsystem, /*bLimitAngle=*/false);
+			PickEntity(*OwnerPC, *World, EntityManager, /*bLimitAngle=*/false);
 		}
 	}
 
-	AddTextLine(FString::Printf(TEXT("{Green}Entities count active{grey}/all: {white}%d{grey}/%d"), EntitySubsystem->DebugGetEntityCount(), EntitySubsystem->DebugGetEntityCount()));
-	AddTextLine(FString::Printf(TEXT("{Green}Registered Archetypes count: {white}%d {green}data ver: {white}%d"), EntitySubsystem->DebugGetArchetypesCount(), EntitySubsystem->GetArchetypeDataVersion()));
+	AddTextLine(FString::Printf(TEXT("{Green}Entities count active{grey}/all: {white}%d{grey}/%d"), EntityManager.DebugGetEntityCount(), EntityManager.DebugGetEntityCount()));
+	AddTextLine(FString::Printf(TEXT("{Green}Registered Archetypes count: {white}%d {green}data ver: {white}%d"), EntityManager.DebugGetArchetypesCount(), EntityManager.GetArchetypeDataVersion()));
 
 	if (UE::Mass::Debug::HasDebugEntities())
 	{
@@ -230,14 +230,14 @@ void FGameplayDebuggerCategory_Mass::CollectData(APlayerController* OwnerPC, AAc
 	{
 		FStringOutputDevice Ar;
 		Ar.SetAutoEmitLineTerminator(true);
-		EntitySubsystem->DebugPrintArchetypes(Ar, /*bIncludeEmpty*/false);
+		EntityManager.DebugPrintArchetypes(Ar, /*bIncludeEmpty*/false);
 
 		AddTextLine(Ar);
 	}
 
 	if (CachedEntity.IsSet() && bMarkEntityBeingDebugged)
 	{
-		if (const FTransformFragment* TransformFragment = EntitySubsystem->GetFragmentDataPtr<FTransformFragment>(CachedEntity))
+		if (const FTransformFragment* TransformFragment = EntityManager.GetFragmentDataPtr<FTransformFragment>(CachedEntity))
 		{
 			const FVector Location = TransformFragment->GetTransform().GetLocation();
 			AddShape(FGameplayDebuggerShape::MakeBox(Location, FVector(8,8,500), FColor::Purple,  FString::Printf(TEXT("[%s]"), *CachedEntity.DebugGetDescription())));
@@ -282,11 +282,11 @@ void FGameplayDebuggerCategory_Mass::CollectData(APlayerController* OwnerPC, AAc
 
 	if (bShowAgentFragments)
 	{
-		if (CachedEntity.IsSet() && EntitySubsystem)
+		if (CachedEntity.IsSet())
 		{
 			// CachedEntity can become invalid if the entity "dies" or in editor mode when related actor gets moved 
 			// (which causes the MassAgentComponent destruction and recreation).
-			if (EntitySubsystem->IsEntityActive(CachedEntity))
+			if (EntityManager.IsEntityActive(CachedEntity))
 			{
 				AddTextLine(FString::Printf(TEXT("{Green}Entity: {White}%s"), *CachedEntity.DebugGetDescription()));
 				AddTextLine(FString::Printf(TEXT("{Green}Type: {White}%s"), (AgentComp == nullptr) ? TEXT("N/A") : AgentComp->IsPuppet() ? TEXT("PUPPET") : TEXT("AGENT")));
@@ -296,13 +296,13 @@ void FGameplayDebuggerCategory_Mass::CollectData(APlayerController* OwnerPC, AAc
 					FStringOutputDevice FragmentsDesc;
 					FragmentsDesc.SetAutoEmitLineTerminator(true);
 					const TCHAR* PrefixToRemove = TEXT("DataFragment_");
-					FMassDebugger::OutputEntityDescription(FragmentsDesc, *EntitySubsystem, CachedEntity, PrefixToRemove);
+					FMassDebugger::OutputEntityDescription(FragmentsDesc, EntityManager, CachedEntity, PrefixToRemove);
 					AddTextLine(FString::Printf(TEXT("{Green}Fragments:\n{White}%s"), *FragmentsDesc));
 				}
 				else
 				{
-					const FMassArchetypeHandle Archetype = EntitySubsystem->GetArchetypeForEntity(CachedEntity);
-					const FMassArchetypeCompositionDescriptor& Composition = EntitySubsystem->GetArchetypeComposition(Archetype);
+					const FMassArchetypeHandle Archetype = EntityManager.GetArchetypeForEntity(CachedEntity);
+					const FMassArchetypeCompositionDescriptor& Composition = EntityManager.GetArchetypeComposition(Archetype);
 					
 					auto DescriptionBuilder = [](const TArray<FName>& ItemNames) -> FString {
 						constexpr int ColumnsCount = 2;
@@ -336,7 +336,7 @@ void FGameplayDebuggerCategory_Mass::CollectData(APlayerController* OwnerPC, AAc
 					AddTextLine(FString::Printf(TEXT("{Green}Shared Fragments:{White}%s"), *DescriptionBuilder(ItemNames)));
 				}
 
-				const FTransformFragment& TransformFragment = EntitySubsystem->GetFragmentDataChecked<FTransformFragment>(CachedEntity);
+				const FTransformFragment& TransformFragment = EntityManager.GetFragmentDataChecked<FTransformFragment>(CachedEntity);
 				constexpr float CapsuleRadius = 50.f;
 				AddShape(FGameplayDebuggerShape::MakeCapsule(TransformFragment.GetTransform().GetLocation() + 2.f * CapsuleRadius * FVector::UpVector, CapsuleRadius, CapsuleRadius * 2.f, FColor::Orange));
 			}
@@ -380,11 +380,11 @@ void FGameplayDebuggerCategory_Mass::CollectData(APlayerController* OwnerPC, AAc
 		UMassSignalSubsystem* SignalSubsystem = World->GetSubsystem<UMassSignalSubsystem>();
 		USmartObjectSubsystem* SmartObjectSubsystem = World->GetSubsystem<USmartObjectSubsystem>();
 		
-		if (MassStateTreeSubsystem && EntitySubsystem && SignalSubsystem && SmartObjectSubsystem)
+		if (MassStateTreeSubsystem && SignalSubsystem && SmartObjectSubsystem)
 		{
 			FMassExecutionContext Context(0.0f);
 		
-			EntityQuery.ForEachEntityChunk(*EntitySubsystem, Context, [this, MassStateTreeSubsystem, SignalSubsystem, EntitySubsystem, SmartObjectSubsystem, OwnerPC, ViewLocation, ViewDirection, CurrentTime](FMassExecutionContext& Context)
+			EntityQuery.ForEachEntityChunk(EntityManager, Context, [this, MassStateTreeSubsystem, SignalSubsystem, &EntityManager, SmartObjectSubsystem, OwnerPC, ViewLocation, ViewDirection, CurrentTime](FMassExecutionContext& Context)
 			{
 				const int32 NumEntities = Context.GetNumEntities();
 				const TConstArrayView<FMassStateTreeInstanceFragment> StateTreeInstanceList = Context.GetFragmentView<FMassStateTreeInstanceFragment>();
@@ -470,9 +470,9 @@ void FGameplayDebuggerCategory_Mass::CollectData(APlayerController* OwnerPC, AAc
 						const FMassLookAtFragment& LookAt = LookAtList[EntityIndex];
 						const FVector WorldLookDirection = Transform.GetTransform().TransformVector(LookAt.Direction);
 						bool bLookArrowDrawn = false;
-						if (LookAt.LookAtMode == EMassLookAtMode::LookAtEntity && EntitySubsystem->IsEntityValid(LookAt.TrackedEntity))
+						if (LookAt.LookAtMode == EMassLookAtMode::LookAtEntity && EntityManager.IsEntityValid(LookAt.TrackedEntity))
 						{
-							if (const FTransformFragment* TargetTransform = EntitySubsystem->GetFragmentDataPtr<FTransformFragment>(LookAt.TrackedEntity))
+							if (const FTransformFragment* TargetTransform = EntityManager.GetFragmentDataPtr<FTransformFragment>(LookAt.TrackedEntity))
 							{
 								FVector TargetPosition = TargetTransform->GetTransform().GetLocation();
 								TargetPosition.Z = BasePos.Z;
@@ -484,9 +484,9 @@ void FGameplayDebuggerCategory_Mass::CollectData(APlayerController* OwnerPC, AAc
 							}
 						}
 
-						if (LookAt.bRandomGazeEntities && EntitySubsystem->IsEntityValid(LookAt.GazeTrackedEntity))
+						if (LookAt.bRandomGazeEntities && EntityManager.IsEntityValid(LookAt.GazeTrackedEntity))
 						{
-							if (const FTransformFragment* TargetTransform = EntitySubsystem->GetFragmentDataPtr<FTransformFragment>(LookAt.GazeTrackedEntity))
+							if (const FTransformFragment* TargetTransform = EntityManager.GetFragmentDataPtr<FTransformFragment>(LookAt.GazeTrackedEntity))
 							{
 								FVector TargetPosition = TargetTransform->GetTransform().GetLocation();
 								TargetPosition.Z = BasePos.Z;
@@ -580,7 +580,7 @@ void FGameplayDebuggerCategory_Mass::CollectData(APlayerController* OwnerPC, AAc
 						// Current StateTree task
 						if (StateTree != nullptr)
 						{
-							FMassStateTreeExecutionContext StateTreeContext(*EntitySubsystem, *SignalSubsystem, Context);
+							FMassStateTreeExecutionContext StateTreeContext(EntityManager, *SignalSubsystem, Context);
 							StateTreeContext.Init(*OwnerPC, *StateTree, EStateTreeStorage::External);
 							StateTreeContext.SetEntity(Entity);
 							
@@ -622,13 +622,13 @@ void FGameplayDebuggerCategory_Mass::CollectData(APlayerController* OwnerPC, AAc
 			});
 		}
 
-		if (bShowNearEntityAvoidance && EntitySubsystem)
+		if (bShowNearEntityAvoidance)
 		{
 			FMassEntityQuery EntityColliderQuery;
 			EntityColliderQuery.AddRequirement<FMassAvoidanceColliderFragment>(EMassFragmentAccess::ReadOnly);
 			EntityColliderQuery.AddRequirement<FTransformFragment>(EMassFragmentAccess::ReadOnly);
 			FMassExecutionContext Context(0.f);
-			EntityColliderQuery.ForEachEntityChunk(*EntitySubsystem, Context, [this, ViewLocation, ViewDirection](const FMassExecutionContext& Context)
+			EntityColliderQuery.ForEachEntityChunk(EntityManager, Context, [this, ViewLocation, ViewDirection](const FMassExecutionContext& Context)
 			{
 				const int32 NumEntities = Context.GetNumEntities();
 				const TConstArrayView<FTransformFragment> TransformList = Context.GetFragmentView<FTransformFragment>();
diff --git a/Engine/Plugins/AI/MassAI/Source/MassAIDebug/Private/MassDebugStateTreeProcessor.cpp b/Engine/Plugins/AI/MassAI/Source/MassAIDebug/Private/MassDebugStateTreeProcessor.cpp
--- a/Engine/Plugins/AI/MassAI/Source/MassAIDebug/Private/MassDebugStateTreeProcessor.cpp
+++ b/Engine/Plugins/AI/MassAI/Source/MassAIDebug/Private/MassDebugStateTreeProcessor.cpp
@@ -26,7 +26,7 @@ void UMassDebugStateTreeProcessor::ConfigureQueries()
 	EntityQuery.AddRequirement<FTransformFragment>(EMassFragmentAccess::ReadOnly);
 }
 
-void UMassDebugStateTreeProcessor::Execute(UMassEntitySubsystem& EntitySubsystem, FMassExecutionContext& Context)
+void UMassDebugStateTreeProcessor::Execute(FMassEntityManager& EntityManager, FMassExecutionContext& Context)
 {
 #if WITH_MASSGAMEPLAY_DEBUG
 	UWorld* World = GetWorld();
@@ -59,7 +59,7 @@ void UMassDebugStateTreeProcessor::Execute(UMassEntitySubsystem& EntitySubsystem
 	}
 	
 	QUICK_SCOPE_CYCLE_COUNTER(UMassDebugStateTreeProcessor_Run);	
-	EntityQuery.ForEachEntityChunk(EntitySubsystem, Context, [this, Debugger, MassStateTreeSubsystem, &EntitySubsystem, MassSignalSubsystem](FMassExecutionContext& Context)
+	EntityQuery.ForEachEntityChunk(EntityManager, Context, [this, Debugger, MassStateTreeSubsystem, &EntityManager, MassSignalSubsystem](FMassExecutionContext& Context)
 	{
 		const FMassEntityHandle SelectedEntity = Debugger->GetSelectedEntity();
 		const int32 NumEntities = Context.GetNumEntities();
@@ -94,7 +94,7 @@ void UMassDebugStateTreeProcessor::Execute(UMassEntitySubsystem& EntitySubsystem
 			
 			if (Entity == SelectedEntity)
 			{
-				FMassStateTreeExecutionContext StateTreeContext(EntitySubsystem, *MassSignalSubsystem, Context);
+				FMassStateTreeExecutionContext StateTreeContext(EntityManager, *MassSignalSubsystem, Context);
 				StateTreeContext.Init(*this, *StateTree, EStateTreeStorage::External);
 				StateTreeContext.SetEntity(Entity);
 				
@@ -112,7 +112,7 @@ void UMassDebugStateTreeProcessor::Execute(UMassEntitySubsystem& EntitySubsystem
 				const FVector ZOffset(0,0,50);
 				const FVector Position = Transform.GetTransform().GetLocation() + ZOffset;
 
-				FMassStateTreeExecutionContext StateTreeContext(EntitySubsystem, *MassSignalSubsystem, Context);
+				FMassStateTreeExecutionContext StateTreeContext(EntityManager, *MassSignalSubsystem, Context);
 				StateTreeContext.Init(*this, *StateTree, EStateTreeStorage::External);
 				StateTreeContext.SetEntity(Entity);
 
diff --git a/Engine/Plugins/AI/MassAI/Source/MassAIDebug/Public/GameplayDebuggerCategory_Mass.h b/Engine/Plugins/AI/MassAI/Source/MassAIDebug/Public/GameplayDebuggerCategory_Mass.h
--- a/Engine/Plugins/AI/MassAI/Source/MassAIDebug/Public/GameplayDebuggerCategory_Mass.h
+++ b/Engine/Plugins/AI/MassAI/Source/MassAIDebug/Public/GameplayDebuggerCategory_Mass.h
@@ -9,7 +9,7 @@
 
 #include "GameplayDebuggerCategory.h"
 
-class UMassEntitySubsystem;
+struct FMassEntityManager;
 class UMassDebuggerSubsystem;
 class APlayerController;
 class AActor;
@@ -25,7 +25,7 @@ protected:
 	virtual void CollectData(APlayerController* OwnerPC, AActor* DebugActor) override;
 	virtual void DrawData(APlayerController* OwnerPC, FGameplayDebuggerCanvasContext& CanvasContext) override;
 	
-	void SetCachedEntity(const FMassEntityHandle Entity, UMassEntitySubsystem& EntitySystem);
+	void SetCachedEntity(const FMassEntityHandle Entity, const FMassEntityManager& EntityManager);
 
 	void OnToggleArchetypes() { bShowArchetypes = !bShowArchetypes; }
 	void OnToggleShapes() { bShowShapes = !bShowShapes; }
@@ -36,7 +36,7 @@ protected:
 	void OnToggleNearEntityAvoidance() { bShowNearEntityAvoidance = !bShowNearEntityAvoidance; }
 	void OnToggleNearEntityPath() { bShowNearEntityPath = !bShowNearEntityPath; }
 	
-	void PickEntity(const APlayerController& OwnerPC, const UWorld& World, UMassEntitySubsystem& EntitySystem, const bool bLimitAngle = true);
+	void PickEntity(const APlayerController& OwnerPC, const UWorld& World, FMassEntityManager& EntityManager, const bool bLimitAngle = true);
 
 protected:
 	AActor* CachedDebugActor;
diff --git a/Engine/Plugins/AI/MassAI/Source/MassAIDebug/Public/MassDebugStateTreeProcessor.h b/Engine/Plugins/AI/MassAI/Source/MassAIDebug/Public/MassDebugStateTreeProcessor.h
--- a/Engine/Plugins/AI/MassAI/Source/MassAIDebug/Public/MassDebugStateTreeProcessor.h
+++ b/Engine/Plugins/AI/MassAI/Source/MassAIDebug/Public/MassDebugStateTreeProcessor.h
@@ -5,7 +5,7 @@
 #include "MassProcessor.h"
 #include "MassDebugStateTreeProcessor.generated.h"
 
-class UMassEntitySubsystem;
+struct FMassEntityManager;
 struct FMassEntityQuery;
 struct FMassExecutionContext;
 
@@ -18,7 +18,7 @@ protected:
 	UMassDebugStateTreeProcessor();
 
 	virtual void ConfigureQueries() override;
-	virtual void Execute(UMassEntitySubsystem& EntitySubsystem, FMassExecutionContext& Context) override;
+	virtual void Execute(FMassEntityManager& EntityManager, FMassExecutionContext& Context) override;
 
 	FMassEntityQuery EntityQuery;
 };
diff --git a/Engine/Plugins/AI/MassAI/Source/MassNavigation/Private/MassAvoidanceProcessors.cpp b/Engine/Plugins/AI/MassAI/Source/MassNavigation/Private/MassAvoidanceProcessors.cpp
--- a/Engine/Plugins/AI/MassAI/Source/MassNavigation/Private/MassAvoidanceProcessors.cpp
+++ b/Engine/Plugins/AI/MassAI/Source/MassNavigation/Private/MassAvoidanceProcessors.cpp
@@ -376,7 +376,7 @@ void UMassMovingAvoidanceProcessor::Initialize(UObject& Owner)
 	NavigationSubsystem = UWorld::GetSubsystem<UMassNavigationSubsystem>(Owner.GetWorld());
 }
 
-void UMassMovingAvoidanceProcessor::Execute(UMassEntitySubsystem& EntitySubsystem, FMassExecutionContext& Context)
+void UMassMovingAvoidanceProcessor::Execute(FMassEntityManager& EntityManager, FMassExecutionContext& Context)
 {
 	QUICK_SCOPE_CYCLE_COUNTER(UMassMovingAvoidanceProcessor);
 
@@ -385,7 +385,7 @@ void UMassMovingAvoidanceProcessor::Execute(UMassEntitySubsystem& EntitySubsyste
 		return;
 	}
 
-	EntityQuery.ForEachEntityChunk(EntitySubsystem, Context, [this, &EntitySubsystem](FMassExecutionContext& Context)
+	EntityQuery.ForEachEntityChunk(EntityManager, Context, [this, &EntityManager](FMassExecutionContext& Context)
 	{
 		const float DeltaTime = Context.GetDeltaTimeSeconds();
 		const float CurrentTime = World->GetTimeSeconds();
@@ -739,14 +739,14 @@ void UMassMovingAvoidanceProcessor::Execute(UMassEntitySubsystem& EntitySubsyste
 				}
 
 				// Skip invalid entities.
-				if (!EntitySubsystem.IsEntityValid(OtherEntity.Entity))
+				if (!EntityManager.IsEntityValid(OtherEntity.Entity))
 				{
 					UE_LOG(LogAvoidanceObstacles, VeryVerbose, TEXT("Close entity is invalid, skipped."));
 					continue;
 				}
 				
 				// Skip too far
-				const FTransform& Transform = EntitySubsystem.GetFragmentDataChecked<FTransformFragment>(OtherEntity.Entity).GetTransform();
+				const FTransform& Transform = EntityManager.GetFragmentDataChecked<FTransformFragment>(OtherEntity.Entity).GetTransform();
 				const FVector OtherLocation = Transform.GetLocation();
 				
 				const float SqDist = FVector::DistSquared(AgentLocation, OtherLocation);
@@ -779,7 +779,7 @@ void UMassMovingAvoidanceProcessor::Execute(UMassEntitySubsystem& EntitySubsyste
 				}
 
 				FSortedObstacle& Obstacle = ClosestObstacles[Index];
-				FMassEntityView OtherEntityView(EntitySubsystem, Obstacle.ObstacleItem.Entity);
+				FMassEntityView OtherEntityView(EntityManager, Obstacle.ObstacleItem.Entity);
 
 				const FMassVelocityFragment* OtherVelocityFragment = OtherEntityView.GetFragmentDataPtr<FMassVelocityFragment>();
 				const FVector OtherVelocity = OtherVelocityFragment != nullptr ? OtherVelocityFragment->Value : FVector::ZeroVector; // Get velocity from FAvoidanceComponent
@@ -1039,7 +1039,7 @@ void UMassStandingAvoidanceProcessor::Initialize(UObject& Owner)
 	NavigationSubsystem = UWorld::GetSubsystem<UMassNavigationSubsystem>(Owner.GetWorld());
 }
 
-void UMassStandingAvoidanceProcessor::Execute(UMassEntitySubsystem& EntitySubsystem, FMassExecutionContext& Context)
+void UMassStandingAvoidanceProcessor::Execute(FMassEntityManager& EntityManager, FMassExecutionContext& Context)
 {
 	QUICK_SCOPE_CYCLE_COUNTER(UMassStandingAvoidanceProcessor);
 
@@ -1049,7 +1049,7 @@ void UMassStandingAvoidanceProcessor::Execute(UMassEntitySubsystem& EntitySubsys
 	}
 
 	// Avoidance while standing
-	EntityQuery.ForEachEntityChunk(EntitySubsystem, Context, [this, &EntitySubsystem](FMassExecutionContext& Context)
+	EntityQuery.ForEachEntityChunk(EntityManager, Context, [this, &EntityManager](FMassExecutionContext& Context)
 	{
 		const int32 NumEntities = Context.GetNumEntities();
 		const float DeltaTime = Context.GetDeltaTimeSeconds();
@@ -1139,14 +1139,14 @@ void UMassStandingAvoidanceProcessor::Execute(UMassEntitySubsystem& EntitySubsys
 				}
 
 				// Skip invalid entities.
-				if (!EntitySubsystem.IsEntityValid(OtherEntity.Entity))
+				if (!EntityManager.IsEntityValid(OtherEntity.Entity))
 				{
 					UE_LOG(LogAvoidanceObstacles, VeryVerbose, TEXT("Close entity is invalid, skipped."));
 					continue;
 				}
 
 				// Skip too far
-				const FTransformFragment& OtherTransform = EntitySubsystem.GetFragmentDataChecked<FTransformFragment>(OtherEntity.Entity);
+				const FTransformFragment& OtherTransform = EntityManager.GetFragmentDataChecked<FTransformFragment>(OtherEntity.Entity);
 				const FVector OtherLocation = OtherTransform.GetTransform().GetLocation();
 				const float DistSq = FVector::DistSquared(AgentLocation, OtherLocation);
 				if (DistSq > DistanceCutOffSqr)
@@ -1166,7 +1166,7 @@ void UMassStandingAvoidanceProcessor::Execute(UMassEntitySubsystem& EntitySubsys
 			for (int32 Index = 0; Index < NumCloseObstacles; Index++)
 			{
 				FSortedObstacle& OtherAgent = ClosestObstacles[Index];
-				FMassEntityView OtherEntityView(EntitySubsystem, OtherAgent.Entity);
+				FMassEntityView OtherEntityView(EntityManager, OtherAgent.Entity);
 
 				const float OtherRadius = OtherEntityView.GetFragmentData<FAgentRadiusFragment>().Radius;
 				const float TotalRadius = GhostRadius + OtherRadius;
diff --git a/Engine/Plugins/AI/MassAI/Source/MassNavigation/Private/MassAvoidanceTrait.cpp b/Engine/Plugins/AI/MassAI/Source/MassNavigation/Private/MassAvoidanceTrait.cpp
--- a/Engine/Plugins/AI/MassAI/Source/MassNavigation/Private/MassAvoidanceTrait.cpp
+++ b/Engine/Plugins/AI/MassAI/Source/MassNavigation/Private/MassAvoidanceTrait.cpp
@@ -7,11 +7,12 @@
 #include "MassCommonFragments.h"
 #include "MassNavigationFragments.h"
 #include "Engine/World.h"
+#include "MassEntityUtils.h"
+
 
 void UMassObstacleAvoidanceTrait::BuildTemplate(FMassEntityTemplateBuildContext& BuildContext, const UWorld& World) const
 {
-	UMassEntitySubsystem* EntitySubsystem = UWorld::GetSubsystem<UMassEntitySubsystem>(&World);
-	check(EntitySubsystem);
+	FMassEntityManager& EntityManager = UE::Mass::Utils::GetEntityManagerChecked(World);
 
 	BuildContext.RequireFragment<FAgentRadiusFragment>();
 	BuildContext.AddFragment<FMassNavigationEdgesFragment>();
@@ -22,11 +23,11 @@ void UMassObstacleAvoidanceTrait::BuildTemplate(FMassEntityTemplateBuildContext&
 
 	const FMassMovingAvoidanceParameters MovingValidated = MovingParameters.GetValidated();
 	const uint32 MovingHash = UE::StructUtils::GetStructCrc32(FConstStructView::Make(MovingValidated));
-	const FConstSharedStruct MovingFragment = EntitySubsystem->GetOrCreateConstSharedFragment(MovingHash, MovingValidated);
+	const FConstSharedStruct MovingFragment = EntityManager.GetOrCreateConstSharedFragment(MovingHash, MovingValidated);
 	BuildContext.AddConstSharedFragment(MovingFragment);
 
 	const FMassStandingAvoidanceParameters StandingValidated = StandingParameters.GetValidated();
 	const uint32 StandingHash = UE::StructUtils::GetStructCrc32(FConstStructView::Make(StandingValidated));
-	const FConstSharedStruct StandingFragment = EntitySubsystem->GetOrCreateConstSharedFragment(StandingHash, StandingValidated);
+	const FConstSharedStruct StandingFragment = EntityManager.GetOrCreateConstSharedFragment(StandingHash, StandingValidated);
 	BuildContext.AddConstSharedFragment(StandingFragment);
 }
diff --git a/Engine/Plugins/AI/MassAI/Source/MassNavigation/Private/MassNavigationProcessors.cpp b/Engine/Plugins/AI/MassAI/Source/MassNavigation/Private/MassNavigationProcessors.cpp
--- a/Engine/Plugins/AI/MassAI/Source/MassNavigation/Private/MassNavigationProcessors.cpp
+++ b/Engine/Plugins/AI/MassAI/Source/MassNavigation/Private/MassNavigationProcessors.cpp
@@ -37,10 +37,10 @@ void UMassOffLODNavigationProcessor::ConfigureQueries()
 	EntityQuery_Conditional.SetChunkFilter(&FMassSimulationVariableTickChunkFragment::ShouldTickChunkThisFrame);
 }
 
-void UMassOffLODNavigationProcessor::Execute(UMassEntitySubsystem& EntitySubsystem,
+void UMassOffLODNavigationProcessor::Execute(FMassEntityManager& EntityManager,
 													FMassExecutionContext& Context)
 {
-	EntityQuery_Conditional.ForEachEntityChunk(EntitySubsystem, Context, [this](FMassExecutionContext& Context)
+	EntityQuery_Conditional.ForEachEntityChunk(EntityManager, Context, [this](FMassExecutionContext& Context)
 	{
 #if WITH_MASSGAMEPLAY_DEBUG
 		if (UE::MassMovement::bFreezeMovement)
@@ -84,10 +84,10 @@ void UMassNavigationSmoothHeightProcessor::ConfigureQueries()
 	EntityQuery.AddConstSharedRequirement<FMassMovementParameters>(EMassFragmentPresence::All);
 }
 
-void UMassNavigationSmoothHeightProcessor::Execute(UMassEntitySubsystem& EntitySubsystem,
+void UMassNavigationSmoothHeightProcessor::Execute(FMassEntityManager& EntityManager,
 													FMassExecutionContext& Context)
 {
-	EntityQuery.ForEachEntityChunk(EntitySubsystem, Context, [this](FMassExecutionContext& Context)
+	EntityQuery.ForEachEntityChunk(EntityManager, Context, [this](FMassExecutionContext& Context)
 	{
 #if WITH_MASSGAMEPLAY_DEBUG
 		if (UE::MassMovement::bFreezeMovement)
@@ -135,9 +135,9 @@ void UMassMoveTargetFragmentInitializer::ConfigureQueries()
 	InitializerQuery.AddRequirement<FTransformFragment>(EMassFragmentAccess::ReadOnly);
 }
 
-void UMassMoveTargetFragmentInitializer::Execute(UMassEntitySubsystem& EntitySubsystem, FMassExecutionContext& Context)
+void UMassMoveTargetFragmentInitializer::Execute(FMassEntityManager& EntityManager, FMassExecutionContext& Context)
 {
-	InitializerQuery.ForEachEntityChunk(EntitySubsystem, Context, [](FMassExecutionContext& Context)
+	InitializerQuery.ForEachEntityChunk(EntityManager, Context, [](FMassExecutionContext& Context)
 	{
 		const int32 NumEntities = Context.GetNumEntities();
 		const TArrayView<FMassMoveTargetFragment> MoveTargetList = Context.GetMutableFragmentView<FMassMoveTargetFragment>();
@@ -192,10 +192,10 @@ void UMassNavigationObstacleGridProcessor::ConfigureQueries()
 	RemoveFromGridEntityQuery.RegisterWithProcessor(*this);
 }
 
-void UMassNavigationObstacleGridProcessor::Execute(UMassEntitySubsystem& EntitySubsystem, FMassExecutionContext& Context)
+void UMassNavigationObstacleGridProcessor::Execute(FMassEntityManager& EntityManager, FMassExecutionContext& Context)
 {
 	// can't be ParallelFor due to MovementSubsystem->GetGridMutable().Move not being thread-safe
-	AddToGridEntityQuery.ForEachEntityChunk(EntitySubsystem, Context, [this, &EntitySubsystem, World = EntitySubsystem.GetWorld()](FMassExecutionContext& Context)
+	AddToGridEntityQuery.ForEachEntityChunk(EntityManager, Context, [this, &EntityManager, World = EntityManager.GetWorld()](FMassExecutionContext& Context)
 	{
 		FNavigationObstacleHashGrid2D& HashGrid = Context.GetMutableSubsystemChecked<UMassNavigationSubsystem>(World).GetObstacleGridMutable();
 		const int32 NumEntities = Context.GetNumEntities();
@@ -222,7 +222,7 @@ void UMassNavigationObstacleGridProcessor::Execute(UMassEntitySubsystem& EntityS
 		}
 	});
 
-	UpdateGridEntityQuery.ForEachEntityChunk(EntitySubsystem, Context, [this, &EntitySubsystem, World = EntitySubsystem.GetWorld()](FMassExecutionContext& Context)
+	UpdateGridEntityQuery.ForEachEntityChunk(EntityManager, Context, [this, &EntityManager, World = EntityManager.GetWorld()](FMassExecutionContext& Context)
 	{
 		FNavigationObstacleHashGrid2D& HashGrid = Context.GetMutableSubsystemChecked<UMassNavigationSubsystem>(World).GetObstacleGridMutable();
 		const int32 NumEntities = Context.GetNumEntities();
@@ -256,7 +256,7 @@ void UMassNavigationObstacleGridProcessor::Execute(UMassEntitySubsystem& EntityS
 		}
 	});
 
-	RemoveFromGridEntityQuery.ForEachEntityChunk(EntitySubsystem, Context, [this, &EntitySubsystem, World = EntitySubsystem.GetWorld()](FMassExecutionContext& Context)
+	RemoveFromGridEntityQuery.ForEachEntityChunk(EntityManager, Context, [this, &EntityManager, World = EntityManager.GetWorld()](FMassExecutionContext& Context)
 	{
 		FNavigationObstacleHashGrid2D& HashGrid = Context.GetMutableSubsystemChecked<UMassNavigationSubsystem>(World).GetObstacleGridMutable();
 		const int32 NumEntities = Context.GetNumEntities();
@@ -291,9 +291,9 @@ void UMassNavigationObstacleRemoverProcessor::ConfigureQueries()
 	EntityQuery.AddSubsystemRequirement<UMassNavigationSubsystem>(EMassFragmentAccess::ReadWrite);
 }
 
-void UMassNavigationObstacleRemoverProcessor::Execute(UMassEntitySubsystem& EntitySubsystem, FMassExecutionContext& Context)
+void UMassNavigationObstacleRemoverProcessor::Execute(FMassEntityManager& EntityManager, FMassExecutionContext& Context)
 {
-	EntityQuery.ForEachEntityChunk(EntitySubsystem, Context, [this, World = EntitySubsystem.GetWorld()](FMassExecutionContext& Context)
+	EntityQuery.ForEachEntityChunk(EntityManager, Context, [this, World = EntityManager.GetWorld()](FMassExecutionContext& Context)
 	{
 		FNavigationObstacleHashGrid2D& HashGrid = Context.GetMutableSubsystemChecked<UMassNavigationSubsystem>(World).GetObstacleGridMutable();
 		const int32 NumEntities = Context.GetNumEntities();
diff --git a/Engine/Plugins/AI/MassAI/Source/MassNavigation/Private/MassSmoothOrientationProcessors.cpp b/Engine/Plugins/AI/MassAI/Source/MassNavigation/Private/MassSmoothOrientationProcessors.cpp
--- a/Engine/Plugins/AI/MassAI/Source/MassNavigation/Private/MassSmoothOrientationProcessors.cpp
+++ b/Engine/Plugins/AI/MassAI/Source/MassNavigation/Private/MassSmoothOrientationProcessors.cpp
@@ -39,7 +39,7 @@ void UMassSmoothOrientationProcessor::ConfigureQueries()
 	LowResEntityQuery_Conditional.SetChunkFilter(&FMassSimulationVariableTickChunkFragment::ShouldTickChunkThisFrame);
 }
 
-void UMassSmoothOrientationProcessor::Execute(UMassEntitySubsystem& EntitySubsystem,
+void UMassSmoothOrientationProcessor::Execute(FMassEntityManager& EntityManager,
 													FMassExecutionContext& Context)
 {
 	// Clamp max delta time to avoid force explosion on large time steps (i.e. during initialization).
@@ -48,7 +48,7 @@ void UMassSmoothOrientationProcessor::Execute(UMassEntitySubsystem& EntitySubsys
 	{
 		QUICK_SCOPE_CYCLE_COUNTER(HighRes);
 
-		HighResEntityQuery.ForEachEntityChunk(EntitySubsystem, Context, [this, DeltaTime](FMassExecutionContext& Context)
+		HighResEntityQuery.ForEachEntityChunk(EntityManager, Context, [this, DeltaTime](FMassExecutionContext& Context)
 		{
 			const int32 NumEntities = Context.GetNumEntities();
 
@@ -117,7 +117,7 @@ void UMassSmoothOrientationProcessor::Execute(UMassEntitySubsystem& EntitySubsys
 	{
 		QUICK_SCOPE_CYCLE_COUNTER(LowRes);
 
-		LowResEntityQuery_Conditional.ForEachEntityChunk(EntitySubsystem, Context, [this](FMassExecutionContext& Context)
+		LowResEntityQuery_Conditional.ForEachEntityChunk(EntityManager, Context, [this](FMassExecutionContext& Context)
 		{
 			const int32 NumEntities = Context.GetNumEntities();
 
diff --git a/Engine/Plugins/AI/MassAI/Source/MassNavigation/Private/MassSmoothOrientationTrait.cpp b/Engine/Plugins/AI/MassAI/Source/MassNavigation/Private/MassSmoothOrientationTrait.cpp
--- a/Engine/Plugins/AI/MassAI/Source/MassNavigation/Private/MassSmoothOrientationTrait.cpp
+++ b/Engine/Plugins/AI/MassAI/Source/MassNavigation/Private/MassSmoothOrientationTrait.cpp
@@ -5,16 +5,17 @@
 #include "MassCommonFragments.h"
 #include "MassNavigationFragments.h"
 #include "Engine/World.h"
+#include "MassEntityUtils.h"
+
 
 void UMassSmoothOrientationTrait::BuildTemplate(FMassEntityTemplateBuildContext& BuildContext, const UWorld& World) const
 {
-	UMassEntitySubsystem* EntitySubsystem = UWorld::GetSubsystem<UMassEntitySubsystem>(&World);
-	check(EntitySubsystem);
+	FMassEntityManager& EntityManager = UE::Mass::Utils::GetEntityManagerChecked(World);
 
 	BuildContext.RequireFragment<FMassMoveTargetFragment>();
 	BuildContext.RequireFragment<FMassVelocityFragment>();
 	BuildContext.RequireFragment<FTransformFragment>();
 
-	const FConstSharedStruct OrientationFragment = EntitySubsystem->GetOrCreateConstSharedFragment(UE::StructUtils::GetStructCrc32(FConstStructView::Make(Orientation)), Orientation);
+	const FConstSharedStruct OrientationFragment = EntityManager.GetOrCreateConstSharedFragment(UE::StructUtils::GetStructCrc32(FConstStructView::Make(Orientation)), Orientation);
 	BuildContext.AddConstSharedFragment(OrientationFragment);
 }
diff --git a/Engine/Plugins/AI/MassAI/Source/MassNavigation/Private/MassSteeringProcessors.cpp b/Engine/Plugins/AI/MassAI/Source/MassNavigation/Private/MassSteeringProcessors.cpp
--- a/Engine/Plugins/AI/MassAI/Source/MassNavigation/Private/MassSteeringProcessors.cpp
+++ b/Engine/Plugins/AI/MassAI/Source/MassNavigation/Private/MassSteeringProcessors.cpp
@@ -75,9 +75,9 @@ void UMassSteerToMoveTargetProcessor::ConfigureQueries()
 	EntityQuery.AddTagRequirement<FMassOffLODTag>(EMassFragmentPresence::None);
 }
 
-void UMassSteerToMoveTargetProcessor::Execute(UMassEntitySubsystem& EntitySubsystem, FMassExecutionContext& Context)
+void UMassSteerToMoveTargetProcessor::Execute(FMassEntityManager& EntityManager, FMassExecutionContext& Context)
 {
-	EntityQuery.ForEachEntityChunk(EntitySubsystem, Context, [this](FMassExecutionContext& Context)
+	EntityQuery.ForEachEntityChunk(EntityManager, Context, [this](FMassExecutionContext& Context)
 	{
 		const int32 NumEntities = Context.GetNumEntities();
 		const TConstArrayView<FTransformFragment> TransformList = Context.GetFragmentView<FTransformFragment>();
diff --git a/Engine/Plugins/AI/MassAI/Source/MassNavigation/Private/MassSteeringTrait.cpp b/Engine/Plugins/AI/MassAI/Source/MassNavigation/Private/MassSteeringTrait.cpp
--- a/Engine/Plugins/AI/MassAI/Source/MassNavigation/Private/MassSteeringTrait.cpp
+++ b/Engine/Plugins/AI/MassAI/Source/MassNavigation/Private/MassSteeringTrait.cpp
@@ -7,11 +7,12 @@
 #include "MassNavigationFragments.h"
 #include "Steering/MassSteeringFragments.h"
 #include "Engine/World.h"
+#include "MassEntityUtils.h"
+
 
 void UMassSteeringTrait::BuildTemplate(FMassEntityTemplateBuildContext& BuildContext, const UWorld& World) const
 {
-	UMassEntitySubsystem* EntitySubsystem = UWorld::GetSubsystem<UMassEntitySubsystem>(&World);
-	check(EntitySubsystem);
+	FMassEntityManager& EntityManager = UE::Mass::Utils::GetEntityManagerChecked(World);
 
 	BuildContext.RequireFragment<FAgentRadiusFragment>();
 	BuildContext.RequireFragment<FTransformFragment>();
@@ -23,9 +24,9 @@ void UMassSteeringTrait::BuildTemplate(FMassEntityTemplateBuildContext& BuildCon
 	BuildContext.AddFragment<FMassStandingSteeringFragment>();
 	BuildContext.AddFragment<FMassGhostLocationFragment>();
 
-	const FConstSharedStruct MovingSteeringFragment = EntitySubsystem->GetOrCreateConstSharedFragment(UE::StructUtils::GetStructCrc32(FConstStructView::Make(MovingSteering)), MovingSteering);
+	const FConstSharedStruct MovingSteeringFragment = EntityManager.GetOrCreateConstSharedFragment(UE::StructUtils::GetStructCrc32(FConstStructView::Make(MovingSteering)), MovingSteering);
 	BuildContext.AddConstSharedFragment(MovingSteeringFragment);
 
-	const FConstSharedStruct StandingSteeringFragment = EntitySubsystem->GetOrCreateConstSharedFragment(UE::StructUtils::GetStructCrc32(FConstStructView::Make(StandingSteering)), StandingSteering);
+	const FConstSharedStruct StandingSteeringFragment = EntityManager.GetOrCreateConstSharedFragment(UE::StructUtils::GetStructCrc32(FConstStructView::Make(StandingSteering)), StandingSteering);
 	BuildContext.AddConstSharedFragment(StandingSteeringFragment);
 }
diff --git a/Engine/Plugins/AI/MassAI/Source/MassNavigation/Public/Avoidance/MassAvoidanceProcessors.h b/Engine/Plugins/AI/MassAI/Source/MassNavigation/Public/Avoidance/MassAvoidanceProcessors.h
--- a/Engine/Plugins/AI/MassAI/Source/MassNavigation/Public/Avoidance/MassAvoidanceProcessors.h
+++ b/Engine/Plugins/AI/MassAI/Source/MassNavigation/Public/Avoidance/MassAvoidanceProcessors.h
@@ -24,7 +24,7 @@ public:
 protected:
 	virtual void ConfigureQueries() override;
 	virtual void Initialize(UObject& Owner) override;
-	virtual void Execute(UMassEntitySubsystem& EntitySubsystem, FMassExecutionContext& Context) override;
+	virtual void Execute(FMassEntityManager& EntityManager, FMassExecutionContext& Context) override;
 
 private:
 	TObjectPtr<UWorld> World;
@@ -44,7 +44,7 @@ public:
 protected:
 	virtual void ConfigureQueries() override;
 	virtual void Initialize(UObject& Owner) override;
-	virtual void Execute(UMassEntitySubsystem& EntitySubsystem, FMassExecutionContext& Context) override;
+	virtual void Execute(FMassEntityManager& EntityManager, FMassExecutionContext& Context) override;
 
 private:
 	TObjectPtr<UWorld> World;
diff --git a/Engine/Plugins/AI/MassAI/Source/MassNavigation/Public/MassNavigationProcessors.h b/Engine/Plugins/AI/MassAI/Source/MassNavigation/Public/MassNavigationProcessors.h
--- a/Engine/Plugins/AI/MassAI/Source/MassNavigation/Public/MassNavigationProcessors.h
+++ b/Engine/Plugins/AI/MassAI/Source/MassNavigation/Public/MassNavigationProcessors.h
@@ -21,7 +21,7 @@ public:
 
 protected:
 	virtual void ConfigureQueries() override;
-	virtual void Execute(UMassEntitySubsystem& EntitySubsystem, FMassExecutionContext& Context) override;
+	virtual void Execute(FMassEntityManager& EntityManager, FMassExecutionContext& Context) override;
 
 private:
 	FMassEntityQuery EntityQuery_Conditional;
@@ -41,7 +41,7 @@ public:
 
 protected:
 	virtual void ConfigureQueries() override;
-	virtual void Execute(UMassEntitySubsystem& EntitySubsystem, FMassExecutionContext& Context) override;
+	virtual void Execute(FMassEntityManager& EntityManager, FMassExecutionContext& Context) override;
 
 private:
 	FMassEntityQuery EntityQuery;
@@ -59,7 +59,7 @@ public:
 	
 protected:
 	virtual void ConfigureQueries() override;
-	virtual void Execute(UMassEntitySubsystem& EntitySubsystem, FMassExecutionContext& Context) override;
+	virtual void Execute(FMassEntityManager& EntityManager, FMassExecutionContext& Context) override;
 
 	FMassEntityQuery InitializerQuery;
 };
@@ -75,7 +75,7 @@ public:
 
 protected:
 	virtual void ConfigureQueries() override;
-	virtual void Execute(UMassEntitySubsystem& EntitySubsystem, FMassExecutionContext& Context) override;
+	virtual void Execute(FMassEntityManager& EntityManager, FMassExecutionContext& Context) override;
 
 private:
 	FMassEntityQuery AddToGridEntityQuery;
@@ -93,7 +93,7 @@ class MASSNAVIGATION_API UMassNavigationObstacleRemoverProcessor : public UMassO
 
 protected:
 	virtual void ConfigureQueries() override;
-	virtual void Execute(UMassEntitySubsystem& EntitySubsystem, FMassExecutionContext& Context) override;
+	virtual void Execute(FMassEntityManager& EntityManager, FMassExecutionContext& Context) override;
 
 	FMassEntityQuery EntityQuery;
 };
diff --git a/Engine/Plugins/AI/MassAI/Source/MassNavigation/Public/SmoothOrientation/MassSmoothOrientationProcessors.h b/Engine/Plugins/AI/MassAI/Source/MassNavigation/Public/SmoothOrientation/MassSmoothOrientationProcessors.h
--- a/Engine/Plugins/AI/MassAI/Source/MassNavigation/Public/SmoothOrientation/MassSmoothOrientationProcessors.h
+++ b/Engine/Plugins/AI/MassAI/Source/MassNavigation/Public/SmoothOrientation/MassSmoothOrientationProcessors.h
@@ -18,7 +18,7 @@ public:
 
 protected:
 	virtual void ConfigureQueries() override;
-	virtual void Execute(UMassEntitySubsystem& EntitySubsystem, FMassExecutionContext& Context) override;
+	virtual void Execute(FMassEntityManager& EntityManager, FMassExecutionContext& Context) override;
 
 private:
 	FMassEntityQuery HighResEntityQuery;
diff --git a/Engine/Plugins/AI/MassAI/Source/MassNavigation/Public/Steering/MassSteeringProcessors.h b/Engine/Plugins/AI/MassAI/Source/MassNavigation/Public/Steering/MassSteeringProcessors.h
--- a/Engine/Plugins/AI/MassAI/Source/MassNavigation/Public/Steering/MassSteeringProcessors.h
+++ b/Engine/Plugins/AI/MassAI/Source/MassNavigation/Public/Steering/MassSteeringProcessors.h
@@ -18,7 +18,7 @@ protected:
 	UMassSteerToMoveTargetProcessor();
 	
 	virtual void ConfigureQueries() override;
-	virtual void Execute(UMassEntitySubsystem& EntitySubsystem, FMassExecutionContext& Context) override;
+	virtual void Execute(FMassEntityManager& EntityManager, FMassExecutionContext& Context) override;
 
 	FMassEntityQuery EntityQuery;
 };
diff --git a/Engine/Plugins/AI/MassAI/Source/MassZoneGraphNavigation/Private/MassZoneGraphNavigationProcessors.cpp b/Engine/Plugins/AI/MassAI/Source/MassZoneGraphNavigation/Private/MassZoneGraphNavigationProcessors.cpp
--- a/Engine/Plugins/AI/MassAI/Source/MassZoneGraphNavigation/Private/MassZoneGraphNavigationProcessors.cpp
+++ b/Engine/Plugins/AI/MassAI/Source/MassZoneGraphNavigation/Private/MassZoneGraphNavigationProcessors.cpp
@@ -53,9 +53,9 @@ void UMassZoneGraphLocationInitializer::ConfigureQueries()
 	EntityQuery.AddSubsystemRequirement<UZoneGraphSubsystem>(EMassFragmentAccess::ReadOnly);
 }
 
-void UMassZoneGraphLocationInitializer::Execute(UMassEntitySubsystem& EntitySubsystem, FMassExecutionContext& Context)
+void UMassZoneGraphLocationInitializer::Execute(FMassEntityManager& EntityManager, FMassExecutionContext& Context)
 {
-	EntityQuery.ForEachEntityChunk(EntitySubsystem, Context, [this, World = EntitySubsystem.GetWorld()](FMassExecutionContext& Context)
+	EntityQuery.ForEachEntityChunk(EntityManager, Context, [this, World = EntityManager.GetWorld()](FMassExecutionContext& Context)
 	{
 		const UZoneGraphSubsystem& ZoneGraphSubsystem = Context.GetSubsystemChecked<UZoneGraphSubsystem>(World);
 
@@ -139,7 +139,7 @@ void UMassZoneGraphPathFollowProcessor::ConfigureQueries()
 	EntityQuery_Conditional.AddSubsystemRequirement<UZoneGraphSubsystem>(EMassFragmentAccess::ReadOnly);
 }
 
-void UMassZoneGraphPathFollowProcessor::Execute(UMassEntitySubsystem& EntitySubsystem, FMassExecutionContext& Context)
+void UMassZoneGraphPathFollowProcessor::Execute(FMassEntityManager& EntityManager, FMassExecutionContext& Context)
 {
 	if (!SignalSubsystem)
 	{
@@ -149,7 +149,7 @@ void UMassZoneGraphPathFollowProcessor::Execute(UMassEntitySubsystem& EntitySubs
 	TArray<FMassEntityHandle> EntitiesToSignalPathDone;
 	TArray<FMassEntityHandle> EntitiesToSignalLaneChanged;
 
-	EntityQuery_Conditional.ForEachEntityChunk(EntitySubsystem, Context, [this, &EntitiesToSignalPathDone, &EntitiesToSignalLaneChanged, World = EntitySubsystem.GetWorld()](FMassExecutionContext& Context)
+	EntityQuery_Conditional.ForEachEntityChunk(EntityManager, Context, [this, &EntitiesToSignalPathDone, &EntitiesToSignalLaneChanged, World = EntityManager.GetWorld()](FMassExecutionContext& Context)
 	{
 		const UZoneGraphSubsystem& ZoneGraphSubsystem = Context.GetSubsystemChecked<UZoneGraphSubsystem>(World);
 
@@ -438,7 +438,7 @@ void UMassZoneGraphLaneCacheBoundaryProcessor::Initialize(UObject& Owner)
 	WeakWorld = Owner.GetWorld();
 }
 
-void UMassZoneGraphLaneCacheBoundaryProcessor::Execute(UMassEntitySubsystem& EntitySubsystem, FMassExecutionContext& Context)
+void UMassZoneGraphLaneCacheBoundaryProcessor::Execute(FMassEntityManager& EntityManager, FMassExecutionContext& Context)
 {
 	QUICK_SCOPE_CYCLE_COUNTER(MassLaneCacheBoundaryProcessor);
 
@@ -448,7 +448,7 @@ void UMassZoneGraphLaneCacheBoundaryProcessor::Execute(UMassEntitySubsystem& Ent
 		return;
 	}
 
-	EntityQuery.ForEachEntityChunk(EntitySubsystem, Context, [this, World](FMassExecutionContext& Context)
+	EntityQuery.ForEachEntityChunk(EntityManager, Context, [this, World](FMassExecutionContext& Context)
 	{
 		const int32 NumEntities = Context.GetNumEntities();
 
diff --git a/Engine/Plugins/AI/MassAI/Source/MassZoneGraphNavigation/Private/MassZoneGraphNavigationTrait.cpp b/Engine/Plugins/AI/MassAI/Source/MassZoneGraphNavigation/Private/MassZoneGraphNavigationTrait.cpp
--- a/Engine/Plugins/AI/MassAI/Source/MassZoneGraphNavigation/Private/MassZoneGraphNavigationTrait.cpp
+++ b/Engine/Plugins/AI/MassAI/Source/MassZoneGraphNavigation/Private/MassZoneGraphNavigationTrait.cpp
@@ -6,11 +6,12 @@
 #include "MassNavigationFragments.h"
 #include "MassZoneGraphNavigationFragments.h"
 #include "Engine/World.h"
+#include "MassEntityUtils.h"
+
 
 void UMassZoneGraphNavigationTrait::BuildTemplate(FMassEntityTemplateBuildContext& BuildContext, const UWorld& World) const
 {
-	UMassEntitySubsystem* EntitySubsystem = UWorld::GetSubsystem<UMassEntitySubsystem>(&World);
-	check(EntitySubsystem);
+	FMassEntityManager& EntityManager = UE::Mass::Utils::GetEntityManagerChecked(World);
 
 	BuildContext.RequireFragment<FAgentRadiusFragment>();
 	BuildContext.RequireFragment<FTransformFragment>();
@@ -22,6 +23,6 @@ void UMassZoneGraphNavigationTrait::BuildTemplate(FMassEntityTemplateBuildContex
 	BuildContext.AddFragment<FMassZoneGraphShortPathFragment>();
 	BuildContext.AddFragment<FMassZoneGraphCachedLaneFragment>();
 
-	const FConstSharedStruct ZGMovementParamsFragment = EntitySubsystem->GetOrCreateConstSharedFragment(UE::StructUtils::GetStructCrc32(FConstStructView::Make(NavigationParameters)), NavigationParameters);
+	const FConstSharedStruct ZGMovementParamsFragment = EntityManager.GetOrCreateConstSharedFragment(UE::StructUtils::GetStructCrc32(FConstStructView::Make(NavigationParameters)), NavigationParameters);
 	BuildContext.AddConstSharedFragment(ZGMovementParamsFragment);
 }
diff --git a/Engine/Plugins/AI/MassAI/Source/MassZoneGraphNavigation/Public/MassZoneGraphNavigationProcessors.h b/Engine/Plugins/AI/MassAI/Source/MassZoneGraphNavigation/Public/MassZoneGraphNavigationProcessors.h
--- a/Engine/Plugins/AI/MassAI/Source/MassZoneGraphNavigation/Public/MassZoneGraphNavigationProcessors.h
+++ b/Engine/Plugins/AI/MassAI/Source/MassZoneGraphNavigation/Public/MassZoneGraphNavigationProcessors.h
@@ -22,7 +22,7 @@ public:
 
 protected:
 	virtual void ConfigureQueries() override;
-	virtual void Execute(UMassEntitySubsystem& EntitySubsystem, FMassExecutionContext& Context) override;
+	virtual void Execute(FMassEntityManager& EntityManager, FMassExecutionContext& Context) override;
 
 	FMassEntityQuery EntityQuery;
 };
@@ -40,7 +40,7 @@ protected:
 	
 	virtual void Initialize(UObject& Owner) override;
 	virtual void ConfigureQueries() override;
-	virtual void Execute(UMassEntitySubsystem& EntitySubsystem, FMassExecutionContext& Context) override;
+	virtual void Execute(FMassEntityManager& EntityManager, FMassExecutionContext& Context) override;
 
 	FMassEntityQuery EntityQuery_Conditional;
 
@@ -61,7 +61,7 @@ public:
 protected:
 	virtual void ConfigureQueries() override;
 	virtual void Initialize(UObject& Owner) override;
-	virtual void Execute(UMassEntitySubsystem& EntitySubsystem, FMassExecutionContext& Context) override;
+	virtual void Execute(FMassEntityManager& EntityManager, FMassExecutionContext& Context) override;
 
 private:
 	TWeakObjectPtr<UWorld> WeakWorld;
diff --git a/Engine/Plugins/AI/MassCrowd/Source/MassCrowd/Private/MassCrowdBubble.cpp b/Engine/Plugins/AI/MassCrowd/Source/MassCrowd/Private/MassCrowdBubble.cpp
--- a/Engine/Plugins/AI/MassCrowd/Source/MassCrowd/Private/MassCrowdBubble.cpp
+++ b/Engine/Plugins/AI/MassCrowd/Source/MassCrowd/Private/MassCrowdBubble.cpp
@@ -17,12 +17,12 @@ namespace UE::Mass::Crowd
 	// @todo provide a better way of selecting agents to debug
 	constexpr int32 MaxAgentsDraw = 300;
 
-	void DebugDrawReplicatedAgent(FMassEntityHandle Entity, const UMassEntitySubsystem& EntitySystem)
+	void DebugDrawReplicatedAgent(FMassEntityHandle Entity, const FMassEntityManager& EntityManager)
 	{
 		static const FVector DebugCylinderHeight = FVector(0.f, 0.f, 200.f);
 		static constexpr float DebugCylinderRadius = 50.f;
 
-		const FMassEntityView EntityView(EntitySystem, Entity);
+		const FMassEntityView EntityView(EntityManager, Entity);
 
 		const FTransformFragment& TransformFragment = EntityView.GetFragmentData<FTransformFragment>();
 		const FMassNetworkIDFragment& NetworkIDFragment = EntityView.GetFragmentData<FMassNetworkIDFragment>();
@@ -50,7 +50,7 @@ namespace UE::Mass::Crowd
 			DebugCylinderColor = FColor(InitialColor / 256 % 256, InitialColor / 256 / 256 % 256, InitialColor % 256);
 		}
 
-		const UWorld* World = EntitySystem.GetWorld();
+		const UWorld* World = EntityManager.GetWorld();
 		if (World != nullptr && World->GetNetMode() == NM_Client)
 		{
 			DrawDebugCylinder(World, Pos, Pos + 0.5f * DebugCylinderHeight, DebugCylinderRadius, /*segments = */24, DebugCylinderColor);
@@ -70,8 +70,7 @@ void FMassCrowdClientBubbleHandler::DebugValidateBubbleOnServer()
 
 	if (UE::Mass::Crowd::bDebugReplicationPositions)
 	{
-		UMassEntitySubsystem* EntitySystem = Serializer->GetEntitySystem();
-		check(EntitySystem);
+		const FMassEntityManager& EntityManager = Serializer->GetEntityManagerChecked();
 
 		// @todo cap at MaxAgentsDraw for now
 		const int32 MaxAgentsDraw = FMath::Min(UE::Mass::Crowd::MaxAgentsDraw, (*Agents).Num());
@@ -84,7 +83,7 @@ void FMassCrowdClientBubbleHandler::DebugValidateBubbleOnServer()
 
 			check(LookupData.Entity.IsSet());
 
-			UE::Mass::Crowd::DebugDrawReplicatedAgent(LookupData.Entity, *EntitySystem);
+			UE::Mass::Crowd::DebugDrawReplicatedAgent(LookupData.Entity, EntityManager);
 		}
 	}
 }
@@ -97,8 +96,7 @@ void FMassCrowdClientBubbleHandler::DebugValidateBubbleOnClient()
 
 	if (UE::Mass::Crowd::bDebugReplicationPositions)
 	{
-		UMassEntitySubsystem* EntitySystem = Serializer->GetEntitySystem();
-		check(EntitySystem);
+		const FMassEntityManager& EntityManager = Serializer->GetEntityManagerChecked();
 
 		UMassReplicationSubsystem* ReplicationSubsystem = Serializer->GetReplicationSubsystem();
 		check(ReplicationSubsystem);
@@ -114,7 +112,7 @@ void FMassCrowdClientBubbleHandler::DebugValidateBubbleOnClient()
 
 			check(EntityInfo->Entity.IsSet());
 
-			UE::Mass::Crowd::DebugDrawReplicatedAgent(EntityInfo->Entity, *EntitySystem);
+			UE::Mass::Crowd::DebugDrawReplicatedAgent(EntityInfo->Entity, EntityManager);
 		}
 	}
 }
diff --git a/Engine/Plugins/AI/MassCrowd/Source/MassCrowd/Private/MassCrowdNavigationProcessor.cpp b/Engine/Plugins/AI/MassCrowd/Source/MassCrowd/Private/MassCrowdNavigationProcessor.cpp
--- a/Engine/Plugins/AI/MassCrowd/Source/MassCrowd/Private/MassCrowdNavigationProcessor.cpp
+++ b/Engine/Plugins/AI/MassCrowd/Source/MassCrowd/Private/MassCrowdNavigationProcessor.cpp
@@ -38,9 +38,9 @@ void UMassCrowdLaneTrackingSignalProcessor::Initialize(UObject& Owner)
 	SubscribeToSignal(*SignalSubsystem, UE::Mass::Signals::CurrentLaneChanged);
 }
 
-void UMassCrowdLaneTrackingSignalProcessor::SignalEntities(UMassEntitySubsystem& EntitySubsystem, FMassExecutionContext& Context, FMassSignalNameLookup& EntitySignals)
+void UMassCrowdLaneTrackingSignalProcessor::SignalEntities(FMassEntityManager& EntityManager, FMassExecutionContext& Context, FMassSignalNameLookup& EntitySignals)
 {
-	EntityQuery.ForEachEntityChunk(EntitySubsystem, Context, [this, World = EntitySubsystem.GetWorld()](FMassExecutionContext& Context)
+	EntityQuery.ForEachEntityChunk(EntityManager, Context, [this, World = EntityManager.GetWorld()](FMassExecutionContext& Context)
 	{
 		UMassCrowdSubsystem& MassCrowdSubsystem = Context.GetMutableSubsystemChecked<UMassCrowdSubsystem>(World);
 		const int32 NumEntities = Context.GetNumEntities();
@@ -78,9 +78,9 @@ void UMassCrowdLaneTrackingDestructor::ConfigureQueries()
 	EntityQuery.AddSubsystemRequirement<UMassCrowdSubsystem>(EMassFragmentAccess::ReadWrite);
 }
 
-void UMassCrowdLaneTrackingDestructor::Execute(UMassEntitySubsystem& EntitySubsystem, FMassExecutionContext& Context)
+void UMassCrowdLaneTrackingDestructor::Execute(FMassEntityManager& EntityManager, FMassExecutionContext& Context)
 {
-	EntityQuery.ForEachEntityChunk(EntitySubsystem, Context, [this, World = EntitySubsystem.GetWorld()](FMassExecutionContext& Context)
+	EntityQuery.ForEachEntityChunk(EntityManager, Context, [this, World = EntityManager.GetWorld()](FMassExecutionContext& Context)
 	{
 		UMassCrowdSubsystem& MassCrowdSubsystem = Context.GetMutableSubsystemChecked<UMassCrowdSubsystem>(World);
 		const int32 NumEntities = Context.GetNumEntities();
@@ -128,13 +128,13 @@ void UMassCrowdDynamicObstacleProcessor::ConfigureQueries()
 	EntityQuery_Conditional.SetChunkFilter(&FMassSimulationVariableTickChunkFragment::ShouldTickChunkThisFrame);
 }
 
-void UMassCrowdDynamicObstacleProcessor::Execute(UMassEntitySubsystem& EntitySubsystem, FMassExecutionContext& Context)
+void UMassCrowdDynamicObstacleProcessor::Execute(FMassEntityManager& EntityManager, FMassExecutionContext& Context)
 {
-	UWorld* World = EntitySubsystem.GetWorld();
+	UWorld* World = EntityManager.GetWorld();
 	const UMassCrowdSettings* CrowdSettings = GetDefault<UMassCrowdSettings>();
 	checkf(CrowdSettings, TEXT("Settings default object is always expected to be valid"));
 
-	EntityQuery_Conditional.ForEachEntityChunk(EntitySubsystem, Context, [this, World, CrowdSettings](FMassExecutionContext& Context)
+	EntityQuery_Conditional.ForEachEntityChunk(EntityManager, Context, [this, World, CrowdSettings](FMassExecutionContext& Context)
 	{
 		const int32 NumEntities = Context.GetNumEntities();
 
@@ -239,11 +239,11 @@ void UMassCrowdDynamicObstacleInitializer::ConfigureQueries()
 	EntityQuery.AddRequirement<FMassCrowdObstacleFragment>(EMassFragmentAccess::ReadWrite);
 }
 
-void UMassCrowdDynamicObstacleInitializer::Execute(UMassEntitySubsystem& EntitySubsystem, FMassExecutionContext& Context)
+void UMassCrowdDynamicObstacleInitializer::Execute(FMassEntityManager& EntityManager, FMassExecutionContext& Context)
 {
-	UWorld* World = EntitySubsystem.GetWorld();
+	UWorld* World = EntityManager.GetWorld();
 	
-	EntityQuery.ForEachEntityChunk(EntitySubsystem, Context, [World](FMassExecutionContext& Context)
+	EntityQuery.ForEachEntityChunk(EntityManager, Context, [World](FMassExecutionContext& Context)
 	{
 		const int32 NumEntities = Context.GetNumEntities();
 		const TConstArrayView<FTransformFragment> LocationList = Context.GetFragmentView<FTransformFragment>();
@@ -286,9 +286,9 @@ void UMassCrowdDynamicObstacleDeinitializer::ConfigureQueries()
 	EntityQuery.AddRequirement<FMassCrowdObstacleFragment>(EMassFragmentAccess::ReadWrite);
 }
 
-void UMassCrowdDynamicObstacleDeinitializer::Execute(UMassEntitySubsystem& EntitySubsystem, FMassExecutionContext& Context)
+void UMassCrowdDynamicObstacleDeinitializer::Execute(FMassEntityManager& EntityManager, FMassExecutionContext& Context)
 {
-	EntityQuery.ForEachEntityChunk(EntitySubsystem, Context, [this](FMassExecutionContext& Context)
+	EntityQuery.ForEachEntityChunk(EntityManager, Context, [this](FMassExecutionContext& Context)
 	{
 		const int32 NumEntities = Context.GetNumEntities();
 		const TArrayView<FMassCrowdObstacleFragment> ObstacleDataList = Context.GetMutableFragmentView<FMassCrowdObstacleFragment>();
diff --git a/Engine/Plugins/AI/MassCrowd/Source/MassCrowd/Private/MassCrowdRepresentationActorManagement.cpp b/Engine/Plugins/AI/MassCrowd/Source/MassCrowd/Private/MassCrowdRepresentationActorManagement.cpp
--- a/Engine/Plugins/AI/MassCrowd/Source/MassCrowd/Private/MassCrowdRepresentationActorManagement.cpp
+++ b/Engine/Plugins/AI/MassCrowd/Source/MassCrowd/Private/MassCrowdRepresentationActorManagement.cpp
@@ -64,7 +64,7 @@ void UMassCrowdRepresentationActorManagement::SetActorEnabled(const EMassActorEn
 	}
 }
 
-AActor* UMassCrowdRepresentationActorManagement::GetOrSpawnActor(UMassRepresentationSubsystem& RepresentationSubsystem, UMassEntitySubsystem& EntitySubsystem, const FMassEntityHandle MassAgent, FMassActorFragment& ActorInfo, const FTransform& Transform, const int16 TemplateActorIndex, FMassActorSpawnRequestHandle& SpawnRequestHandle, const float Priority) const
+AActor* UMassCrowdRepresentationActorManagement::GetOrSpawnActor(UMassRepresentationSubsystem& RepresentationSubsystem, FMassEntityManager& EntityManager, const FMassEntityHandle MassAgent, FMassActorFragment& ActorInfo, const FTransform& Transform, const int16 TemplateActorIndex, FMassActorSpawnRequestHandle& SpawnRequestHandle, const float Priority) const
 {
 	FTransform RootTransform = Transform;
 	
@@ -76,7 +76,7 @@ AActor* UMassCrowdRepresentationActorManagement::GetOrSpawnActor(UMassRepresenta
 		}
 	}
 
-	return Super::GetOrSpawnActor(RepresentationSubsystem, EntitySubsystem, MassAgent, ActorInfo, RootTransform, TemplateActorIndex, SpawnRequestHandle, Priority);
+	return Super::GetOrSpawnActor(RepresentationSubsystem, EntityManager, MassAgent, ActorInfo, RootTransform, TemplateActorIndex, SpawnRequestHandle, Priority);
 }
 
 void UMassCrowdRepresentationActorManagement::TeleportActor(const FTransform& Transform, AActor& Actor, FMassCommandBuffer& CommandBuffer) const
diff --git a/Engine/Plugins/AI/MassCrowd/Source/MassCrowd/Private/MassCrowdServerRepresentationLODProcessor.cpp b/Engine/Plugins/AI/MassCrowd/Source/MassCrowd/Private/MassCrowdServerRepresentationLODProcessor.cpp
--- a/Engine/Plugins/AI/MassCrowd/Source/MassCrowd/Private/MassCrowdServerRepresentationLODProcessor.cpp
+++ b/Engine/Plugins/AI/MassCrowd/Source/MassCrowd/Private/MassCrowdServerRepresentationLODProcessor.cpp
@@ -46,18 +46,18 @@ void UMassCrowdServerRepresentationLODProcessor::Initialize(UObject& InOwner)
 	Super::Initialize(InOwner);
 }
 
-void UMassCrowdServerRepresentationLODProcessor::Execute(UMassEntitySubsystem& EntitySubsystem, FMassExecutionContext& Context)
+void UMassCrowdServerRepresentationLODProcessor::Execute(FMassEntityManager& EntityManager, FMassExecutionContext& Context)
 {
 	TRACE_CPUPROFILER_EVENT_SCOPE(TEXT("CrowdServerRepresentationLOD"))
 
-	const UMassLODSubsystem& LODSubsystem = Context.GetSubsystemChecked<UMassLODSubsystem>(EntitySubsystem.GetWorld());
+	const UMassLODSubsystem& LODSubsystem = Context.GetSubsystemChecked<UMassLODSubsystem>(EntityManager.GetWorld());
 	const TArray<FViewerInfo>& Viewers = LODSubsystem.GetViewers();
 	LODCalculator.PrepareExecution(Viewers);
 	
 	{
 		TRACE_CPUPROFILER_EVENT_SCOPE(TEXT("CalculateLOD"))
 		
-		EntityQuery.ForEachEntityChunk(EntitySubsystem, Context, [this](FMassExecutionContext& Context)
+		EntityQuery.ForEachEntityChunk(EntityManager, Context, [this](FMassExecutionContext& Context)
 		{
 			const TConstArrayView<FMassViewerInfoFragment> ViewersInfoList = Context.GetFragmentView<FMassViewerInfoFragment>();
 			const TArrayView<FMassRepresentationLODFragment> RepresentationLODFragments = Context.GetMutableFragmentView<FMassRepresentationLODFragment>();
@@ -70,7 +70,7 @@ void UMassCrowdServerRepresentationLODProcessor::Execute(UMassEntitySubsystem& E
 		
 		if (LODCalculator.AdjustDistancesFromCount())
 		{
-			EntityQuery.ForEachEntityChunk(EntitySubsystem, Context, [this](FMassExecutionContext& Context)
+			EntityQuery.ForEachEntityChunk(EntityManager, Context, [this](FMassExecutionContext& Context)
 			{
 				const TConstArrayView<FMassViewerInfoFragment> ViewersInfoList = Context.GetFragmentView<FMassViewerInfoFragment>();
 				const TArrayView<FMassRepresentationLODFragment> RepresentationLODFragments = Context.GetMutableFragmentView<FMassRepresentationLODFragment>();
@@ -83,8 +83,8 @@ void UMassCrowdServerRepresentationLODProcessor::Execute(UMassEntitySubsystem& E
 	if (UE::MassCrowd::bDebugCrowdServerRepresentationLOD)
 	{
 		TRACE_CPUPROFILER_EVENT_SCOPE(TEXT("DebugDisplayLOD"))
-		UWorld* World = EntitySubsystem.GetWorld();
-		EntityQuery.ForEachEntityChunk(EntitySubsystem, Context, [this, World](FMassExecutionContext& Context)
+		UWorld* World = EntityManager.GetWorld();
+		EntityQuery.ForEachEntityChunk(EntityManager, Context, [this, World](FMassExecutionContext& Context)
 		{
 			const TConstArrayView<FTransformFragment> LocationList = Context.GetFragmentView<FTransformFragment>();
 			const TConstArrayView<FMassRepresentationLODFragment> RepresentationLODFragments = Context.GetFragmentView<FMassRepresentationLODFragment>();
diff --git a/Engine/Plugins/AI/MassCrowd/Source/MassCrowd/Private/MassCrowdServerRepresentationTrait.cpp b/Engine/Plugins/AI/MassCrowd/Source/MassCrowd/Private/MassCrowdServerRepresentationTrait.cpp
--- a/Engine/Plugins/AI/MassCrowd/Source/MassCrowd/Private/MassCrowdServerRepresentationTrait.cpp
+++ b/Engine/Plugins/AI/MassCrowd/Source/MassCrowd/Private/MassCrowdServerRepresentationTrait.cpp
@@ -9,6 +9,8 @@
 #include "Engine/World.h"
 #include "MassCrowdRepresentationActorManagement.h"
 #include "MassActorSubsystem.h"
+#include "MassEntityUtils.h"
+
 
 UMassCrowdServerRepresentationTrait::UMassCrowdServerRepresentationTrait()
 {
@@ -36,8 +38,7 @@ void UMassCrowdServerRepresentationTrait::BuildTemplate(FMassEntityTemplateBuild
 	BuildContext.RequireFragment<FTransformFragment>();
 	BuildContext.RequireFragment<FMassActorFragment>();
 
-	UMassEntitySubsystem* EntitySubsystem = UWorld::GetSubsystem<UMassEntitySubsystem>(&World);
-	check(EntitySubsystem);
+	FMassEntityManager& EntityManager = UE::Mass::Utils::GetEntityManagerChecked(World);
 
 	UMassCrowdRepresentationSubsystem* RepresentationSubsystem = World.GetSubsystem<UMassCrowdRepresentationSubsystem>();
 	check(RepresentationSubsystem);
@@ -45,11 +46,11 @@ void UMassCrowdServerRepresentationTrait::BuildTemplate(FMassEntityTemplateBuild
 	FMassRepresentationSubsystemSharedFragment Subsystem;
 	Subsystem.RepresentationSubsystem = RepresentationSubsystem;
 	uint32 SubsystemHash = UE::StructUtils::GetStructCrc32(FConstStructView::Make(Subsystem));
-	FSharedStruct SubsystemFragment = EntitySubsystem->GetOrCreateSharedFragment<FMassRepresentationSubsystemSharedFragment>(SubsystemHash, Subsystem);
+	FSharedStruct SubsystemFragment = EntityManager.GetOrCreateSharedFragment<FMassRepresentationSubsystemSharedFragment>(SubsystemHash, Subsystem);
 	BuildContext.AddSharedFragment(SubsystemFragment);
 
 	uint32 ParamsHash = UE::StructUtils::GetStructCrc32(FConstStructView::Make(Params));
-	FConstSharedStruct ParamsFragment = EntitySubsystem->GetOrCreateConstSharedFragment<FMassRepresentationParameters>(ParamsHash, Params);
+	FConstSharedStruct ParamsFragment = EntityManager.GetOrCreateConstSharedFragment<FMassRepresentationParameters>(ParamsHash, Params);
 	ParamsFragment.Get<FMassRepresentationParameters>().ComputeCachedValues();
 	BuildContext.AddConstSharedFragment(ParamsFragment);
 
diff --git a/Engine/Plugins/AI/MassCrowd/Source/MassCrowd/Private/MassCrowdVisualizationLODProcessor.cpp b/Engine/Plugins/AI/MassCrowd/Source/MassCrowd/Private/MassCrowdVisualizationLODProcessor.cpp
--- a/Engine/Plugins/AI/MassCrowd/Source/MassCrowd/Private/MassCrowdVisualizationLODProcessor.cpp
+++ b/Engine/Plugins/AI/MassCrowd/Source/MassCrowd/Private/MassCrowdVisualizationLODProcessor.cpp
@@ -46,21 +46,21 @@ void UMassCrowdVisualizationLODProcessor::ConfigureQueries()
 	FilterTag = FMassCrowdTag::StaticStruct();
 }
 
-void UMassCrowdVisualizationLODProcessor::Execute(UMassEntitySubsystem& EntitySubsystem, FMassExecutionContext& Context)
+void UMassCrowdVisualizationLODProcessor::Execute(FMassEntityManager& EntityManager, FMassExecutionContext& Context)
 {
-	UWorld* World = EntitySubsystem.GetWorld();
+	UWorld* World = EntityManager.GetWorld();
 
 	ForceOffLOD((bool)UE::MassCrowd::GCrowdTurnOffVisualization);
 
 	TRACE_CPUPROFILER_EVENT_SCOPE(TEXT("CrowdVisualizationLOD"))
 
-	Super::Execute(EntitySubsystem, Context);
+	Super::Execute(EntityManager, Context);
 	
 	if (UE::MassCrowd::bDebugCrowdVisualizationLOD)
 	{
 		TRACE_CPUPROFILER_EVENT_SCOPE(TEXT("DebugDisplayLOD"))
 
-		DebugEntityQuery.ForEachEntityChunk(EntitySubsystem, Context, [World](FMassExecutionContext& Context)
+		DebugEntityQuery.ForEachEntityChunk(EntityManager, Context, [World](FMassExecutionContext& Context)
 		{
 			FMassVisualizationLODSharedFragment& LODSharedFragment = Context.GetMutableSharedFragment<FMassVisualizationLODSharedFragment>();
 			const TConstArrayView<FTransformFragment> LocationList = Context.GetFragmentView<FTransformFragment>();
@@ -73,7 +73,7 @@ void UMassCrowdVisualizationLODProcessor::Execute(UMassEntitySubsystem& EntitySu
 	{
 		TRACE_CPUPROFILER_EVENT_SCOPE(TEXT("ShowISMUnderSpecifiedRange"))
 
-		DebugEntityQuery.ForEachEntityChunk(EntitySubsystem, Context, [World](const FMassExecutionContext& Context)
+		DebugEntityQuery.ForEachEntityChunk(EntityManager, Context, [World](const FMassExecutionContext& Context)
 		{
 			const TConstArrayView<FTransformFragment> LocationList = Context.GetFragmentView<FTransformFragment>();
 			const TConstArrayView<FMassRepresentationFragment> RepresentationFragmentList = Context.GetFragmentView<FMassRepresentationFragment>();
