Cleared out Mass's ini-declared redirects. All of these were required before 5.0 when a lot of type moving around was taking place.

diff --git a/Engine/Plugins/AI/MassAI/Config/BaseMassAI.ini b/Engine/Plugins/AI/MassAI/Config/BaseMassAI.ini
--- a/Engine/Plugins/AI/MassAI/Config/BaseMassAI.ini
+++ b/Engine/Plugins/AI/MassAI/Config/BaseMassAI.ini
@@ -1,43 +1,2 @@
 [CoreRedirects]
-+PackageRedirects=(OldName="/Script/MassVisualization",NewName="/Script/MassRepresentation")
-+PackageRedirects=(OldName="/Script/MassBehavior",NewName="/Script/MassAIBehavior")
-
-+ClassRedirects=(OldName="/Script/MassMovement.MassAvoidanceObstacleTrait",NewName="/Script/MassAIMovement.MassAvoidanceObstacleTrait")
-+ClassRedirects=(OldName="/Script/MassMovement.MassDynamicObstacleTrait",NewName="/Script/MassAIMovement.MassAvoidanceObstacleTrait")
-+ClassRedirects=(OldName="/Script/MassMovement.MassAvoidanceTrait",NewName="/Script/MassAIMovement.MassAvoidanceTrait")
-+ClassRedirects=(OldName="/Script/MassMovement.MassZoneGraphMovementTrait",NewName="/Script/MassAIMovement.MassZoneGraphMovementTrait")
-+ClassRedirects=(OldName="/Script/MassAIMovement.MassMovementSettings",NewName="/Script/MassMovement.MassMovementSettings")
-
-; these are potentially not needed anymore
-+StructRedirects=(OldName="MassLaneHandleFragment",NewName="MassLaneNavigationFragment")
-+StructRedirects=(OldName="MassAgentConfig",NewName="MassEntityConfig")
-+StructRedirects=(OldName="MassSpawnedAgentType",NewName="MassSpawnedEntityType")
-+StructRedirects=(OldName="MassSpawnedAgentType",NewName="MassSpawnedEntityType")
-+StructRedirects=(OldName="MassStateTreeClaimSmartObjectTask",NewName="MassMoveAndClaimSmartObjectTask")
-+StructRedirects=(OldName="MassStateTreeUseSmartObjectTask",NewName="MassUseSmartObjectTask")
-+StructRedirects=(OldName="MassZoneGraphFindFleeTarget",NewName="MassZoneGraphFindEscapeTarget")
-+StructRedirects=(OldName="MassDynamicObstacleGridCellLocationFragment",NewName="MassAvoidanceObstacleGridCellLocationFragment")
-
-+ClassRedirects=(OldName="MassProcessor_SmartObjectTestBehavior",NewName="MassProcessor_SmartObjectTimedBehavior")
-+ClassRedirects=(OldName="MassAgentConfigAsset",NewName="MassEntityConfigAsset")
-+ClassRedirects=(OldName="MassAgentFeatureBase",NewName="MassEntityTraitBase")
-+ClassRedirects=(OldName="MassStateTreeFeature",NewName="MassStateTreeTrait")
-+ClassRedirects=(OldName="MassAssortedFragmentsFeature",NewName="MassAssortedFragmentsTrait")
-+ClassRedirects=(OldName="MassDebugVisualizationFeature",NewName="MassDebugVisualizationTrait")
-+ClassRedirects=(OldName="MassReplicationFeature",NewName="MassReplicationTrait")
-+ClassRedirects=(OldName="DynamicObstacleProcessor",NewName="MassAvoidanceDynamicObstacleProcessor")
-+ClassRedirects=(OldName="MassAvoidanceDynamicObstacleProcessor",NewName="MassAvoidanceObstacleProcessor")
-+ClassRedirects=(OldName="MassDynamicObstacleTrait",NewName="MassAvoidanceObstacleTrait")
-+ClassRedirects=(OldName="MassFragmentDestructor_DynamicObstacleRemover",NewName="MassAvoidanceObstacleRemoverFragmentDestructor")
-
-+PropertyRedirects=(OldName="MassAgentComponent.bUseAgentConfig",NewName="bUseEntityConfig")
-+PropertyRedirects=(OldName="MassAgentComponent.AgentConfig",NewName="EntityConfig")
-+PropertyRedirects=(OldName="MassSpawnedEntityType.AgentConfig",NewName="EntityConfig")
-+PropertyRedirects=(OldName="MassSpawner.AgentTypes",NewName="EntityTypes")
-+PropertyRedirects=(OldName="MassEntityConfig.Features",NewName="Traits")
-
-+PropertyRedirects=(OldName="MassZoneGraphFindEscapeTarget.FleeTag",NewName="DisturbanceAnnotationTag")
-
-;MassAIBehavior
-+StructRedirects=(OldName="/Script/MassAIBehavior.MassZoneGraphAnnotationTagsFragment",NewName="/Script/MassAIBehavior.MassZoneGraphAnnotationFragment")
-+StructRedirects=(OldName="/Script/MassAIBehavior.MassStateTreeActivated",NewName="/Script/MassAIBehavior.MassStateTreeActivatedTag")
+; no changes past UE 5.0
diff --git a/Engine/Plugins/AI/MassAI/Source/MassNavigation/Private/MassNavigationModule.cpp b/Engine/Plugins/AI/MassAI/Source/MassNavigation/Private/MassNavigationModule.cpp
--- a/Engine/Plugins/AI/MassAI/Source/MassNavigation/Private/MassNavigationModule.cpp
+++ b/Engine/Plugins/AI/MassAI/Source/MassNavigation/Private/MassNavigationModule.cpp
@@ -18,11 +18,6 @@ IMPLEMENT_MODULE(FMassNavigationModule, MassNavigation)
 void FMassNavigationModule::StartupModule()
 {
 	// This code will execute after your module is loaded into memory (but after global variables are initialized, of course.)
-
-	TArray<FCoreRedirect> Redirects;
-//	Redirects.Emplace(ECoreRedirectFlags::Type_Struct, TEXT("AvoidanceComponent"), TEXT("MassAvoidanceFragment"));
-	
-	FCoreRedirects::AddRedirectList(Redirects, TEXT("MassNavigation"));
 }
 
 
diff --git a/Engine/Plugins/AI/MassCrowd/Config/BaseMassCrowd.ini b/Engine/Plugins/AI/MassCrowd/Config/BaseMassCrowd.ini
--- a/Engine/Plugins/AI/MassCrowd/Config/BaseMassCrowd.ini
+++ b/Engine/Plugins/AI/MassCrowd/Config/BaseMassCrowd.ini
@@ -1,7 +1,2 @@
 [CoreRedirects]
-+PackageRedirects=(OldName="/Script/MassCrowdAI5", NewName="/Script/MassCrowd")
-+StructRedirects=(OldName="/Script/MassAIBehavior.MassZoneGraphFindWanderTarget",NewName="/Script/MassCrowd.MassZoneGraphFindWanderTarget")
-
-+StructRedirects=(OldName="/Script/MassCrowd.TagFragment_MassCrowd",NewName="/Script/MassCrowd.MassCrowdTag")
-
-+PropertyRedirects=(OldName="MassCrowdServerRepresentationTrait.Config", NewName="Params")
+; no changes past UE 5.0
\ No newline at end of file
diff --git a/Engine/Plugins/Runtime/MassEntity/Config/BaseMassEntity.ini b/Engine/Plugins/Runtime/MassEntity/Config/BaseMassEntity.ini
--- a/Engine/Plugins/Runtime/MassEntity/Config/BaseMassEntity.ini
+++ b/Engine/Plugins/Runtime/MassEntity/Config/BaseMassEntity.ini
@@ -1,4 +1,2 @@
 [CoreRedirects]
-+PackageRedirects=(OldName="/Script/AggregateTicking",NewName="/Script/MassEntity")
-+PackageRedirects=(OldName="/Script/Pipe",NewName="/Script/MassEntity")
-+StructRedirects=(OldName="LWEntity",NewName="/Script/MassEntity.MassEntityHandle")
+; no changes past UE 5.0
diff --git a/Engine/Plugins/Runtime/MassGameplay/Config/BaseMassGameplay.ini b/Engine/Plugins/Runtime/MassGameplay/Config/BaseMassGameplay.ini
--- a/Engine/Plugins/Runtime/MassGameplay/Config/BaseMassGameplay.ini
+++ b/Engine/Plugins/Runtime/MassGameplay/Config/BaseMassGameplay.ini
@@ -1,42 +1,2 @@
 [CoreRedirects]
-+PackageRedirects=(OldName="/Script/Mass",NewName="/Script/MassGameplay")
-
-+StructRedirects=(OldName="MassSpawnPointGenerator", NewName="MassSpawnDataGenerator")
-+StructRedirects=(OldName="/Script/MassAIMovement.MassMovementStyleRef", NewName="/Script/MassMovement.MassMovementStyleRef")
-+ClassRedirects=(OldName="/Script/MassAIMovement.MassMovementSettings", NewName="/Script/MassMovement.MassMovementSettings")
-
-+PropertyRedirects=(OldName="/Script/MassSpawner.SpawnPointsGenerators", NewName="SpawnDataGenerators")
-
-;MassCommon
-+StructRedirects=(OldName="DataFragment_AgentRadius", NewName="AgentRadiusFragment")
-+StructRedirects=(OldName="DataFragment_ObjectWrapper", NewName="ObjectWrapperFragment")
-+StructRedirects=(OldName="DataFragment_Transform", NewName="TransformFragment")
-+PropertyRedirects=(OldName="/Script/MassCommon.DataFragment_AgentRadius", NewName="/Script/MassCommon.AgentRadiusFragment")
-+PropertyRedirects=(OldName="/Script/MassCommon.DataFragment_ObjectWrapper", NewName="/Script/MassCommon.ObjectWrapperFragment")
-+PropertyRedirects=(OldName="/Script/MassCommon.DataFragment_Transform", NewName="/Script/MassCommon.TransformFragment")
-
-;MassActors
-+StructRedirects=(OldName="/Script/MassActors.DataFragment_Actor", NewName="/Script/MassActors.MassActorFragment")
-+StructRedirects=(OldName="/Script/MassActors.DataFragment_CapsuleComponentWrapper", NewName="/Script/MassActors.CapsuleComponentWrapperFragment")
-+StructRedirects=(OldName="/Script/MassActors.DataFragment_CharacterMovementComponentWrapper", NewName="/Script/MassActors.CharacterMovementComponentWrapperFragment")
-
-;MassGameplayDebug
-+StructRedirects=(OldName="/Script/MassGameplayDebug.SimDebugVisComponent", NewName="/Script/MassGameplayDebug.SimDebugVisFragment")
-
-;MassSpawner
-+StructRedirects=(OldName="/Script/MassSpawner.DataFragment_ReplicationTemplateID", NewName="/Script/MassSpawner.ReplicationTemplateIDFragment")
-
-;MassRepresentation
-+StructRedirects=(OldName="/Script/MassRepresentation.MassRepresentationConfig", NewName="/Script/MassRepresentation.MassRepresentationParameters")
-+StructRedirects=(OldName="/Script/MassRepresentation.MassRepresentationSubsystemFragment", NewName="/Script/MassRepresentation.MassRepresentationSubsystemSharedFragment")
-+PropertyRedirects=(OldName="MassVisualizationTrait.Config", NewName="Params")
-
-;MassLOD
-+StructRedirects=(OldName="/Script/MassLOD.MassSimulationVariableTickConfig", NewName="/Script/MassLOD.MassSimulationVariableTickParameters")
-+StructRedirects=(OldName="/Script/MassLOD.MassSimulationLODConfig", NewName="/Script/MassLOD.MassSimulationLODParameters")
-+PropertyRedirects=(OldName="MassSimulationLODTrait.Config", NewName="Params")
-+PropertyRedirects=(OldName="MassSimulationLODTrait.VariableTickConfig", NewName="VariableTickParams")
-
-;MassReplication
-+ClassRedirects=(OldName="/Script/MassReplication.MassReplicationManager", NewName="/Script/MassReplication.MassReplicationSubsystem")
-+ClassRedirects=(OldName="MassReplicationManager", NewName="MassReplicationSubsystem")
+; no changes past UE 5.0
diff --git a/Engine/Plugins/Runtime/MassGameplay/Source/MassLOD/Private/MassLODModule.cpp b/Engine/Plugins/Runtime/MassGameplay/Source/MassLOD/Private/MassLODModule.cpp
--- a/Engine/Plugins/Runtime/MassGameplay/Source/MassLOD/Private/MassLODModule.cpp
+++ b/Engine/Plugins/Runtime/MassGameplay/Source/MassLOD/Private/MassLODModule.cpp
@@ -19,23 +19,6 @@ IMPLEMENT_MODULE(FMassLODModule, MassLOD)
 void FMassLODModule::StartupModule()
 {
 	// This code will execute after your module is loaded into memory (but after global variables are initialized, of course.)
-
-	TArray<FCoreRedirect> Redirects;
-	Redirects.Emplace(ECoreRedirectFlags::Type_Struct, TEXT("/Script/MassCommon.DataFragment_MassSimulationLODInfo"), TEXT("/Script/MassLOD.DataFragment_MassSimulationLODInfo"));
-	Redirects.Emplace(ECoreRedirectFlags::Type_Struct, TEXT("/Script/MassCommon.DataFragment_MassSimulationLOD"), TEXT("/Script/MassLOD.DataFragment_MassSimulationLOD"));
-	Redirects.Emplace(ECoreRedirectFlags::Type_Struct, TEXT("/Script/MassCommon.DataFragment_MassViewerLOD"), TEXT("/Script/MassLOD.DataFragment_MassViewerLOD"));
-
-	Redirects.Emplace(ECoreRedirectFlags::Type_Class, TEXT("/Script/MassCommon.MassProcessor_LODBase"), TEXT("/Script/MassLOD.MassLODProcessorBase"));
-	Redirects.Emplace(ECoreRedirectFlags::Type_Class, TEXT("/Script/MassLOD.MassProcessor_LODBase"), TEXT("/Script/MassLOD.MassLODProcessorBase"));
-
-	Redirects.Emplace(ECoreRedirectFlags::Type_Class, TEXT("/Script/MassCommon.MassProcessor_MassSimulationLODViewersInfo"), TEXT("/Script/MassLOD.MassProcessor_MassSimulationLODViewersInfo"));
-	Redirects.Emplace(ECoreRedirectFlags::Type_Class, TEXT("/Script/MassCommon.MassProcessor_MassSimulationLOD"), TEXT("/Script/MassLOD.MassProcessor_MassSimulationLOD"));
-
-	Redirects.Emplace(ECoreRedirectFlags::Type_Class, TEXT("/Script/MassCommon.MassProcessor_MassViewerLOD"), TEXT("/Script/MassLOD.MassProcessor_MassViewerLOD"));
-
-	Redirects.Emplace(ECoreRedirectFlags::Type_Class, TEXT("/Script/MassLOD.MassLODTrait"), TEXT("/Script/MassLOD.MassSimulationLODTrait"));
-	
-	FCoreRedirects::AddRedirectList(Redirects, TEXT("MassMovement"));
 }
 
 
diff --git a/Engine/Plugins/Runtime/MassGameplay/Source/MassMovement/Private/MassMovementModule.cpp b/Engine/Plugins/Runtime/MassGameplay/Source/MassMovement/Private/MassMovementModule.cpp
--- a/Engine/Plugins/Runtime/MassGameplay/Source/MassMovement/Private/MassMovementModule.cpp
+++ b/Engine/Plugins/Runtime/MassGameplay/Source/MassMovement/Private/MassMovementModule.cpp
@@ -19,11 +19,6 @@ IMPLEMENT_MODULE(FMassMovementModule, MassMovement)
 void FMassMovementModule::StartupModule()
 {
 	// This code will execute after your module is loaded into memory (but after global variables are initialized, of course.)
-
-	TArray<FCoreRedirect> Redirects;
-	Redirects.Emplace(ECoreRedirectFlags::Type_Struct, TEXT("VelocityComponent"), TEXT("MassVelocityFragment"));
-	
-	FCoreRedirects::AddRedirectList(Redirects, TEXT("MassMovement"));
 }
 
 
diff --git a/Engine/Plugins/Runtime/MassGameplay/Source/MassRepresentation/Private/MassRepresentationModule.cpp b/Engine/Plugins/Runtime/MassGameplay/Source/MassRepresentation/Private/MassRepresentationModule.cpp
--- a/Engine/Plugins/Runtime/MassGameplay/Source/MassRepresentation/Private/MassRepresentationModule.cpp
+++ b/Engine/Plugins/Runtime/MassGameplay/Source/MassRepresentation/Private/MassRepresentationModule.cpp
@@ -17,22 +17,6 @@ IMPLEMENT_MODULE(FMassRepresentationModule, MassRepresentation)
 
 void FMassRepresentationModule::StartupModule()
 {
-	TArray<FCoreRedirect> Redirects;
-
-	Redirects.Emplace(ECoreRedirectFlags::Type_Struct, TEXT("/Script/MassVisualization.DataFragment_Visualization"), TEXT("/Script/MassRepresentation.MassRepresentationFragment"));
-	Redirects.Emplace(ECoreRedirectFlags::Type_Class, TEXT("/Script/MassVisualization.MassProcessor_Visualization"), TEXT("/Script/MassRepresentation.MassRepresentationProcessor"));
-	Redirects.Emplace(ECoreRedirectFlags::Type_Class, TEXT("/Script/MassVisualization.MassFragmentDestructor_Visualization"), TEXT("/Script/MassRepresentation.MassRepresentationFragmentDestructor"));
-	Redirects.Emplace(ECoreRedirectFlags::Type_Class, TEXT("/Script/MassVisualization.MassVisualizationLODProcessor"), TEXT("/Script/MassRepresentation.MassVisualizationLODProcessor"));
-	Redirects.Emplace(ECoreRedirectFlags::Type_Class, TEXT("/Script/MassRepresentation.MassRepresentationLODProcessor"), TEXT("/Script/MassRepresentation.MassVisualizationLODProcessor"));
-	Redirects.Emplace(ECoreRedirectFlags::Type_Class, TEXT("/Script/MassVisualization.MassVisualizationFeature"), TEXT("/Script/MassRepresentation.MassVisualizationTrait"));
-	Redirects.Emplace(ECoreRedirectFlags::Type_Property, TEXT("/Script/MassRepresentation.MassVisualizationTrait.TemplateActor"), TEXT("/Script/MassRepresentation.MassVisualizationTrait.HighResTemplateActor"));
-	Redirects.Emplace(ECoreRedirectFlags::Type_Enum, TEXT("/Script/MassRepresentation.RepresentationType"), TEXT("/Script/MassRepresentation.MassRepresentationType"));
-	Redirects.Emplace(ECoreRedirectFlags::Type_Enum, TEXT("RepresentationType"),                            TEXT("/Script/MassRepresentation.MassRepresentationType"));
-	Redirects.Emplace(ECoreRedirectFlags::Type_Enum, TEXT("/Script/MassRepresentation.ActorEnabledType"), TEXT("/Script/MassRepresentation.MassActorEnabledType"));
-	Redirects.Emplace(ECoreRedirectFlags::Type_Enum, TEXT("ActorEnabledType"),                            TEXT("/Script/MassRepresentation.MassActorEnabledType"));
-
-
-	FCoreRedirects::AddRedirectList(Redirects, TEXT("MassRepresentation"));
 }
 
 
diff --git a/Engine/Plugins/Runtime/MassGameplay/Source/MassSmartObjects/Private/MassSmartObjectsModule.cpp b/Engine/Plugins/Runtime/MassGameplay/Source/MassSmartObjects/Private/MassSmartObjectsModule.cpp
--- a/Engine/Plugins/Runtime/MassGameplay/Source/MassSmartObjects/Private/MassSmartObjectsModule.cpp
+++ b/Engine/Plugins/Runtime/MassGameplay/Source/MassSmartObjects/Private/MassSmartObjectsModule.cpp
@@ -18,13 +18,6 @@ IMPLEMENT_MODULE(FMassSmartObjectsModule, MassSmartObjects)
 void FMassSmartObjectsModule::StartupModule()
 {
 	// This code will execute after your module is loaded into memory (but after global variables are initialized, of course.)
-
-	TArray<FCoreRedirect> Redirects;
-	Redirects.Emplace(ECoreRedirectFlags::Type_Class, TEXT("SmartObjectMassBehaviorConfig"), TEXT("SmartObjectMassBehaviorDefinition"));
-	Redirects.Emplace(ECoreRedirectFlags::Type_Class, TEXT("MassProcessor_SmartObjectTimedBehavior"), TEXT("MassSmartObjectTimedBehaviorProcessor"));
-	Redirects.Emplace(ECoreRedirectFlags::Type_Class, TEXT("MassProcessor_SmartObjectCandidatesFinder"), TEXT("MassSmartObjectCandidatesFinderProcessor"));
-
-	FCoreRedirects::AddRedirectList(Redirects, TEXT("MassSmartObject"));
 }
 
 
