Fix bad merge of MassStateTreeProcessor

diff --git a/Engine/Plugins/AI/MassAI/Source/MassAIBehavior/Private/MassStateTreeProcessors.cpp b/Engine/Plugins/AI/MassAI/Source/MassAIBehavior/Private/MassStateTreeProcessors.cpp
--- a/Engine/Plugins/AI/MassAI/Source/MassAIBehavior/Private/MassStateTreeProcessors.cpp
+++ b/Engine/Plugins/AI/MassAI/Source/MassAIBehavior/Private/MassStateTreeProcessors.cpp
@@ -171,7 +171,7 @@ void ForEachEntityInChunk(FCachedExternalData& CachedExternalData, FMassExecutio
 			// Gather all required fragments.
 			{
 				CSV_SCOPED_TIMING_STAT_EXCLUSIVE(StateTreeProcessorExternalFragments);
-				if (!ensureMsgf(UE::MassBehavior::SetExternalFragments(StateTreeContext, StateTreeContext.GetEntitySubsystem()), TEXT("StateTree will not execute due to missing required fragments.")))
+				if (!ensureMsgf(UE::MassBehavior::SetExternalFragments(StateTreeContext, StateTreeContext.GetEntityManager()), TEXT("StateTree will not execute due to missing required fragments.")))
 				{
 					break;
 				}
@@ -225,12 +225,11 @@ void UMassStateTreeFragmentDestructor::Execute(FMassEntityManager& EntityManager
 	{
 		return;
 	}
-	FMassStateTreeExecutionContext StateTreeContext(EntityManager, *SignalSubsystem, Context);
 
 	UE::MassBehavior::FCachedExternalData CachedExternalData;
 
 	EntityQuery.ForEachEntityChunk(EntityManager, Context,
-		[this, &EntityManager, &CachedExternalData, World = EntityManager.GetWorld()](FMassExecutionContext& Context)
+		[&EntityManager, &CachedExternalData, SignalSubsystem = SignalSubsystem, World = EntityManager.GetWorld()](FMassExecutionContext& Context)
 		{
 			UMassStateTreeSubsystem& MassStateTreeSubsystem = Context.GetMutableSubsystemChecked<UMassStateTreeSubsystem>(World);
 			const TArrayView<FMassStateTreeInstanceFragment> StateTreeInstanceList = Context.GetMutableFragmentView<FMassStateTreeInstanceFragment>();
@@ -286,10 +285,9 @@ void UMassStateTreeActivationProcessor::Execute(FMassEntityManager& EntityManage
 
 	const UMassBehaviorSettings* BehaviorSettings = GetDefault<UMassBehaviorSettings>();
 	check(BehaviorSettings);
-
+ 
 	// StateTree processor relies on signals to be ticked but we need an 'initial tick' to set the tree in the proper state.
 	// The initializer provides that by sending a signal to all new entities that use StateTree.
-
 	UE::MassBehavior::FCachedExternalData CachedExternalData;
 	const float TimeInSeconds = EntityManager.GetWorld()->GetTimeSeconds();
 
@@ -297,7 +295,7 @@ void UMassStateTreeActivationProcessor::Execute(FMassEntityManager& EntityManage
 	int32 ActivationCounts[EMassLOD::Max] {0,0,0,0};
 	
 	EntityQuery.ForEachEntityChunk(EntityManager, Context,
-		[&EntitiesToSignal, &ActivationCounts, &CachedExternalData, SignalSubsystem = SignalSubsystem, &EntityManager, MaxActivationsPerLOD = BehaviorSettings->MaxActivationsPerLOD, TimeInSeconds, World = EntityManager.GetWorld()](FMassExecutionContext& Context)
+		[&EntitiesToSignal, &ActivationCounts, &CachedExternalData, SignalSubsystem = &SignalSubsystem, &EntityManager, MaxActivationsPerLOD = BehaviorSettings->MaxActivationsPerLOD, TimeInSeconds, World = EntityManager.GetWorld()](FMassExecutionContext& Context)
 		{
 			UMassStateTreeSubsystem& MassStateTreeSubsystem = Context.GetMutableSubsystemChecked<UMassStateTreeSubsystem>(World);
 			const int32 NumEntities = Context.GetNumEntities();
@@ -320,7 +318,7 @@ void UMassStateTreeActivationProcessor::Execute(FMassEntityManager& EntityManage
 			}
 			
 			// Start StateTree. This may do substantial amount of work, as we select and enter the first state.
-			UE::MassBehavior::ForEachEntityInChunk(CachedExternalData, Context, EntitySubsystem, *SignalSubsystem, MassStateTreeSubsystem,
+			UE::MassBehavior::ForEachEntityInChunk(CachedExternalData, Context, EntityManager, *SignalSubsystem, MassStateTreeSubsystem,
 				[TimeInSeconds](FMassStateTreeExecutionContext& StateTreeExecutionContext, FMassStateTreeInstanceFragment& StateTreeFragment)
 				{
 					// Start the tree instance
@@ -418,15 +416,15 @@ void UMassStateTreeProcessor::SignalEntities(FMassEntityManager& EntityManager,
 	UE::MassBehavior::FCachedExternalData CachedExternalData;
 	TArray<FMassEntityHandle> EntitiesToSignal;
 
-	EntityQuery.ForEachEntityChunk(EntitySubsystem, Context,
-		[this, &TimeInSeconds, &EntitiesToSignal, &CachedExternalData, &EntityManager, SignalSubsystem = SignalSubsystem, World, &EntitySignals](FMassExecutionContext& Context)
+	EntityQuery.ForEachEntityChunk(EntityManager, Context,
+		[this, &TimeInSeconds, &EntitiesToSignal, &CachedExternalData, &EntityManager, SignalSubsystem = &SignalSubsystem, World, &EntitySignals](FMassExecutionContext& Context)
 		{
 			// Keep stats regarding the amount of tree instances ticked per frame
 			CSV_CUSTOM_STAT(StateTreeProcessor, NumTickedStateTree, Context.GetNumEntities(), ECsvCustomStatOp::Accumulate);
 
 			UMassStateTreeSubsystem& MassStateTreeSubsystem = Context.GetMutableSubsystemChecked<UMassStateTreeSubsystem>(World);
 
-			UE::MassBehavior::ForEachEntityInChunk(CachedExternalData, Context, EntitySubsystem, *SignalSubsystem, MassStateTreeSubsystem,
+			UE::MassBehavior::ForEachEntityInChunk(CachedExternalData, Context, EntityManager, *SignalSubsystem, MassStateTreeSubsystem,
 				[TimeInSeconds, &EntitiesToSignal, &EntitySignals, &MassStateTreeSubsystem]
 				(FMassStateTreeExecutionContext& StateTreeExecutionContext, FMassStateTreeInstanceFragment& StateTreeFragment)
 				{
