MassAIBehavior: Fix StateTree asset filtering for StateTree trait.

diff --git a/Engine/Plugins/AI/MassAI/Source/MassAIBehavior/Public/MassStateTreeTrait.h b/Engine/Plugins/AI/MassAI/Source/MassAIBehavior/Public/MassStateTreeTrait.h
--- a/Engine/Plugins/AI/MassAI/Source/MassAIBehavior/Public/MassStateTreeTrait.h
+++ b/Engine/Plugins/AI/MassAI/Source/MassAIBehavior/Public/MassStateTreeTrait.h
@@ -21,6 +21,6 @@ protected:
 
 	virtual void ValidateTemplate(FMassEntityTemplateBuildContext& BuildContext, const UWorld& World) const override;
 
-	UPROPERTY(Category="StateTree", EditAnywhere, meta=(RequiredAssetDataTags="Schema=/Script/MassAIBehavior/MassStateTreeSchema"))
+	UPROPERTY(Category="StateTree", EditAnywhere, meta=(RequiredAssetDataTags="Schema=/Script/MassAIBehavior.MassStateTreeSchema"))
 	TObjectPtr<UStateTree> StateTree;
 };
