Mass: Remove unnecessary (and broken) private include paths.

diff --git a/Engine/Plugins/AI/MassAI/Source/MassNavigation/MassNavigation.Build.cs b/Engine/Plugins/AI/MassAI/Source/MassNavigation/MassNavigation.Build.cs
--- a/Engine/Plugins/AI/MassAI/Source/MassNavigation/MassNavigation.Build.cs
+++ b/Engine/Plugins/AI/MassAI/Source/MassNavigation/MassNavigation.Build.cs
@@ -8,13 +8,6 @@ namespace UnrealBuildTool.Rules
 		{
 			PCHUsage = ModuleRules.PCHUsageMode.UseExplicitOrSharedPCHs;
 
-			PublicIncludePaths.AddRange(
-				new string[] {
-				"Runtime/AIModule/Public",
-				ModuleDirectory + "/Public",
-				}
-			);
-
 			PublicDependencyModuleNames.AddRange(
 				new string[] {
 					"AIModule",
@@ -33,7 +26,7 @@ namespace UnrealBuildTool.Rules
 				}
 			);
 
-			if (Target.bBuildEditor == true)
+			if (Target.bCompileAgainstEditor)
 			{
 				PrivateDependencyModuleNames.Add("EditorFramework");
 				PrivateDependencyModuleNames.Add("UnrealEd");
