MassNavigation: Disable unity build to work around intermittent vs2022 compiler crash

#rnx

diff --git a/Engine/Plugins/AI/MassAI/Source/MassNavigation/MassNavigation.Build.cs b/Engine/Plugins/AI/MassAI/Source/MassNavigation/MassNavigation.Build.cs
--- a/Engine/Plugins/AI/MassAI/Source/MassNavigation/MassNavigation.Build.cs
+++ b/Engine/Plugins/AI/MassAI/Source/MassNavigation/MassNavigation.Build.cs
@@ -8,6 +8,9 @@ namespace UnrealBuildTool.Rules
 		{
 			PCHUsage = ModuleRules.PCHUsageMode.UseExplicitOrSharedPCHs;
 
+			// Work-around for intermittent VS2022 msvc crash
+			bUseUnity = false;
+
 			PublicDependencyModuleNames.AddRange(
 				new string[] {
 					"AIModule",
