Silence static analysis warning in MassAvoidanceProcessors

diff --git a/Engine/Plugins/AI/MassAI/Source/MassNavigation/Private/MassAvoidanceProcessors.cpp b/Engine/Plugins/AI/MassAI/Source/MassNavigation/Private/MassAvoidanceProcessors.cpp
--- a/Engine/Plugins/AI/MassAI/Source/MassNavigation/Private/MassAvoidanceProcessors.cpp
+++ b/Engine/Plugins/AI/MassAI/Source/MassNavigation/Private/MassAvoidanceProcessors.cpp
@@ -853,7 +853,7 @@ void UMassMovingAvoidanceProcessor::Execute(FMassEntityManager& EntityManager, F
 					
 					// Find the maximum distance from edges that are too close.
 					float MaxDist = -1.f;
-					FVector ClosestPoint;
+					FVector ClosestPoint = FVector::ZeroVector;
 					for (const FNavigationAvoidanceEdge& Edge : NavEdges.AvoidanceEdges)
 					{
 						const FVector Point = FMath::ClosestPointOnSegment(Collider.Location, Edge.Start, Edge.End);
