Addressing numeric conversion warnings in Mass.

diff --git a/Engine/Plugins/AI/MassAI/Source/MassAIBehavior/Private/Evaluators/MassComponentHitEvaluator.cpp b/Engine/Plugins/AI/MassAI/Source/MassAIBehavior/Private/Evaluators/MassComponentHitEvaluator.cpp
--- a/Engine/Plugins/AI/MassAI/Source/MassAIBehavior/Private/Evaluators/MassComponentHitEvaluator.cpp
+++ b/Engine/Plugins/AI/MassAI/Source/MassAIBehavior/Private/Evaluators/MassComponentHitEvaluator.cpp
@@ -32,9 +32,9 @@ void FMassComponentHitEvaluator::Tick(FStateTreeExecutionContext &Context, const
 		// @todo: This is a bit of a kludge to expose an event to StateTree.
 		const UWorld* World = Context.GetWorld();
 		check(World);
-		const float CurrentTime = World->GetTimeSeconds();
-		const float TimeSinceHit = CurrentTime - HitResult->HitTime;
-		constexpr float HitEventDuration = 0.1f;
+		const double CurrentTime = World->GetTimeSeconds();
+		const double TimeSinceHit = CurrentTime - HitResult->HitTime;
+		constexpr double HitEventDuration = 0.1;
 		if (TimeSinceHit < HitEventDuration)
 		{
 			MASSBEHAVIOR_LOG(VeryVerbose, TEXT("Got hit"));
diff --git a/Engine/Plugins/AI/MassAI/Source/MassAIBehavior/Private/MassComponentHitSubsystem.cpp b/Engine/Plugins/AI/MassAI/Source/MassAIBehavior/Private/MassComponentHitSubsystem.cpp
--- a/Engine/Plugins/AI/MassAI/Source/MassAIBehavior/Private/MassComponentHitSubsystem.cpp
+++ b/Engine/Plugins/AI/MassAI/Source/MassAIBehavior/Private/MassComponentHitSubsystem.cpp
@@ -91,16 +91,16 @@ void UMassComponentHitSubsystem::OnHitCallback(UPrimitiveComponent* HitComp, AAc
 		bProcessHit = (HitActorAsPawn != nullptr && HitActorAsPawn->IsPlayerControlled()) || (OtherAsPawn != nullptr && OtherAsPawn->IsPlayerControlled());
 	}
 
-	const float CurrentTime = World->GetTimeSeconds();
+	const double CurrentTime = World->GetTimeSeconds();
 
 	// If new hit result comes during this duration, it will be merged to existing one.
-	constexpr float HitResultMergeDuration = 1.0f;
+	constexpr double HitResultMergeDuration = 1.;
 	if (bProcessHit)
 	{
 		FMassHitResult* ExistingHitResult = HitResults.Find(Entity);
 		if (ExistingHitResult)
 		{
-			const float TimeSinceLastHit = CurrentTime - ExistingHitResult->LastFilteredHitTime;
+			const double TimeSinceLastHit = CurrentTime - ExistingHitResult->LastFilteredHitTime;
 			if (TimeSinceLastHit < HitResultMergeDuration)
 			{
 				ExistingHitResult->LastFilteredHitTime = CurrentTime;
@@ -128,13 +128,13 @@ void UMassComponentHitSubsystem::Tick(float DeltaTime)
 	const UWorld* World = GetWorld();
 	check(World);
 
-	const float CurrentTime = World->GetTimeSeconds();
-	constexpr float HitResultDecayDuration = 1.0f;
+	const double CurrentTime = World->GetTimeSeconds();
+	constexpr double HitResultDecayDuration = 1.;
 	
 	for (auto Iter = HitResults.CreateIterator(); Iter; ++Iter)
 	{
 		const FMassHitResult& HitResult = Iter.Value();
-		const float ElapsedTime = CurrentTime - HitResult.LastFilteredHitTime;
+		const double ElapsedTime = CurrentTime - HitResult.LastFilteredHitTime;
 		if (ElapsedTime > HitResultDecayDuration)
 		{
 			Iter.RemoveCurrent();
diff --git a/Engine/Plugins/AI/MassAI/Source/MassAIBehavior/Private/MassLookAtFragments.cpp b/Engine/Plugins/AI/MassAI/Source/MassAIBehavior/Private/MassLookAtFragments.cpp
--- a/Engine/Plugins/AI/MassAI/Source/MassAIBehavior/Private/MassLookAtFragments.cpp
+++ b/Engine/Plugins/AI/MassAI/Source/MassAIBehavior/Private/MassLookAtFragments.cpp
@@ -64,7 +64,7 @@ FVector FMassLookAtTrajectoryFragment::GetPointAtDistanceExtrapolated(const floa
 	const float T = FMath::Clamp((DistanceAlongPath - SegStartDistance) * InvSegLength, 0.0f, 1.0f);
 
 	// 1/3 third is used to create smooth bezier curve. On linear segments 1/3 will result linear interpolation.
-	const float TangentDistance = FVector::Dist(StartPoint.Position, EndPoint.Position) / 3.0f;
+	const FVector::FReal TangentDistance = FVector::Dist(StartPoint.Position, EndPoint.Position) / 3.;
 	const FVector StartControlPoint = StartPoint.Position + StartPoint.Tangent.GetVector() * TangentDistance;
 	const FVector EndControlPoint = EndPoint.Position - EndPoint.Tangent.GetVector() * TangentDistance;
 
diff --git a/Engine/Plugins/AI/MassAI/Source/MassAIBehavior/Private/MassLookAtProcessors.cpp b/Engine/Plugins/AI/MassAI/Source/MassAIBehavior/Private/MassLookAtProcessors.cpp
--- a/Engine/Plugins/AI/MassAI/Source/MassAIBehavior/Private/MassLookAtProcessors.cpp
+++ b/Engine/Plugins/AI/MassAI/Source/MassAIBehavior/Private/MassLookAtProcessors.cpp
@@ -39,14 +39,14 @@ namespace UE::MassBehavior
 	// Clamps direction vector to a cone specified by the cone angle along X-axis
 	FVector ClampDirectionToXAxisCone(const FVector Direction, const float ConeAngle)
 	{
-		float ConeSin = 0.0f, ConeCos = 0.0f;
+		FVector::FReal ConeSin = 0., ConeCos = 0.;
 		FMath::SinCos(&ConeSin, &ConeCos, ConeAngle);
 		
-		const float AngleCos = Direction.X; // Same as FVector::DotProduct(FVector::ForwardVector, Direction);
+		const FVector::FReal AngleCos = Direction.X; // Same as FVector::DotProduct(FVector::ForwardVector, Direction);
 		if (AngleCos < ConeCos)
 		{
-			const float DistToRimSq = FMath::Square(Direction.Y) + FMath::Square(Direction.Z);
-			const float InvDistToRim = DistToRimSq > KINDA_SMALL_NUMBER ? (1.0f / FMath::Sqrt(DistToRimSq)) : 0.0f;
+			const FVector::FReal DistToRimSq = FMath::Square(Direction.Y) + FMath::Square(Direction.Z);
+			const FVector::FReal InvDistToRim = DistToRimSq > KINDA_SMALL_NUMBER ? (1. / FMath::Sqrt(DistToRimSq)) : 0.;
 			return FVector(ConeCos, Direction.Y * InvDistToRim * ConeSin, Direction.Z * InvDistToRim * ConeSin);
 		}
 		
@@ -118,7 +118,7 @@ void UMassLookAtProcessor::Execute(FMassEntityManager& EntityManager, FMassExecu
 {
 	QUICK_SCOPE_CYCLE_COUNTER(LookAtProcessor_Run);
 
-	const float CurrentTime = GetWorld()->GetTimeSeconds();
+	const double CurrentTime = GetWorld()->GetTimeSeconds();
 
 	EntityQuery_Conditional.ForEachEntityChunk(EntityManager, Context, [this, &EntityManager, CurrentTime, World = EntityManager.GetWorld()](FMassExecutionContext& Context)
 		{
@@ -151,7 +151,7 @@ void UMassLookAtProcessor::Execute(FMassEntityManager& EntityManager, FMassExecu
 				// Update gaze target when current cycle is finished.
 				if (LookAt.RandomGazeMode != EMassLookAtGazeMode::None)
 				{
-					const float TimeSinceUpdate = CurrentTime - LookAt.GazeStartTime;
+					const double TimeSinceUpdate = CurrentTime - LookAt.GazeStartTime;
 					if (TimeSinceUpdate >= LookAt.GazeDuration)
 					{
 						FindNewGazeTarget(MassNavSystem, EntityManager, CurrentTime, TransformFragment.GetTransform(), LookAt);
@@ -197,7 +197,7 @@ void UMassLookAtProcessor::Execute(FMassEntityManager& EntityManager, FMassExecu
 				// Apply gaze
 				if (LookAt.RandomGazeMode != EMassLookAtGazeMode::None)
 				{
-					const float TimeSinceUpdate = CurrentTime - LookAt.GazeStartTime;
+					const float TimeSinceUpdate = FloatCastChecked<float>(CurrentTime - LookAt.GazeStartTime, /* Precision */ 1./64.);
 					const float GazeStrength = UE::MassBehavior::GazeEnvelope(TimeSinceUpdate, LookAt.GazeDuration, LookAt.RandomGazeMode);
 
 					if (GazeStrength > KINDA_SMALL_NUMBER)
@@ -233,7 +233,7 @@ void UMassLookAtProcessor::Execute(FMassEntityManager& EntityManager, FMassExecu
 		});
 }
 
-void UMassLookAtProcessor::FindNewGazeTarget(const UMassNavigationSubsystem& MassNavSystem, const FMassEntityManager& EntityManager, const float CurrentTime, const FTransform& Transform, FMassLookAtFragment& LookAt) const
+void UMassLookAtProcessor::FindNewGazeTarget(const UMassNavigationSubsystem& MassNavSystem, const FMassEntityManager& EntityManager, const double CurrentTime, const FTransform& Transform, FMassLookAtFragment& LookAt) const
 {
 	const FNavigationObstacleHashGrid2D& ObstacleGrid = MassNavSystem.GetObstacleGrid();
 	const FMassEntityHandle LastTrackedEntity = LookAt.GazeTrackedEntity;
@@ -492,7 +492,7 @@ void UMassLookAtProcessor::BuildTrajectory(const UZoneGraphSubsystem& ZoneGraphS
 			const FVector StartForward = CurrPoint.Tangent.GetVector();
 			const FVector EndPoint = NextPoint.Position;
 			const FVector EndForward = NextPoint.Tangent.GetVector();
-			const float TangentDistance = FVector::Dist(StartPoint, EndPoint) / 3.0f;
+			const FVector::FReal TangentDistance = FVector::Dist(StartPoint, EndPoint) / 3.;
 			const FVector StartControlPoint = StartPoint + StartForward * TangentDistance;
 			const FVector EndControlPoint = EndPoint - EndForward * TangentDistance;
 
diff --git a/Engine/Plugins/AI/MassAI/Source/MassAIBehavior/Private/MassStateTreeProcessors.cpp b/Engine/Plugins/AI/MassAI/Source/MassAIBehavior/Private/MassStateTreeProcessors.cpp
--- a/Engine/Plugins/AI/MassAI/Source/MassAIBehavior/Private/MassStateTreeProcessors.cpp
+++ b/Engine/Plugins/AI/MassAI/Source/MassAIBehavior/Private/MassStateTreeProcessors.cpp
@@ -289,7 +289,7 @@ void UMassStateTreeActivationProcessor::Execute(FMassEntityManager& EntityManage
 	// StateTree processor relies on signals to be ticked but we need an 'initial tick' to set the tree in the proper state.
 	// The initializer provides that by sending a signal to all new entities that use StateTree.
 	UE::MassBehavior::FCachedExternalData CachedExternalData;
-	const float TimeInSeconds = EntityManager.GetWorld()->GetTimeSeconds();
+	const double TimeInSeconds = EntityManager.GetWorld()->GetTimeSeconds();
 
 	TArray<FMassEntityHandle> EntitiesToSignal;
 	int32 ActivationCounts[EMassLOD::Max] {0,0,0,0};
@@ -411,7 +411,7 @@ void UMassStateTreeProcessor::SignalEntities(FMassEntityManager& EntityManager,
 	QUICK_SCOPE_CYCLE_COUNTER(StateTreeProcessor_Run);
 	CSV_SCOPED_TIMING_STAT_EXCLUSIVE(StateTreeProcessorExecute);
 
-	const float TimeInSeconds = EntityManager.GetWorld()->GetTimeSeconds();
+	const double TimeInSeconds = EntityManager.GetWorld()->GetTimeSeconds();
 
 	UE::MassBehavior::FCachedExternalData CachedExternalData;
 	TArray<FMassEntityHandle> EntitiesToSignal;
@@ -429,7 +429,7 @@ void UMassStateTreeProcessor::SignalEntities(FMassEntityManager& EntityManager,
 				(FMassStateTreeExecutionContext& StateTreeExecutionContext, FMassStateTreeInstanceFragment& StateTreeFragment)
 				{
 					// Compute adjusted delta time
-					const float AdjustedDeltaTime = TimeInSeconds - StateTreeFragment.LastUpdateTimeInSeconds;
+					const float AdjustedDeltaTime = FloatCastChecked<float>(TimeInSeconds - StateTreeFragment.LastUpdateTimeInSeconds, /* Precision */ 1./256.);
 					StateTreeFragment.LastUpdateTimeInSeconds = TimeInSeconds;
 
 #if WITH_MASSGAMEPLAY_DEBUG
diff --git a/Engine/Plugins/AI/MassAI/Source/MassAIBehavior/Public/MassComponentHitTypes.h b/Engine/Plugins/AI/MassAI/Source/MassAIBehavior/Public/MassComponentHitTypes.h
--- a/Engine/Plugins/AI/MassAI/Source/MassAIBehavior/Public/MassComponentHitTypes.h
+++ b/Engine/Plugins/AI/MassAI/Source/MassAIBehavior/Public/MassComponentHitTypes.h
@@ -20,7 +20,7 @@ struct FMassHitResult
 
 	FMassHitResult() = default;
 
-	FMassHitResult(const FMassEntityHandle OtherEntity, const float Time)
+	FMassHitResult(const FMassEntityHandle OtherEntity, const double Time)
 		: OtherEntity(OtherEntity)
 		, HitTime(Time)
 		, LastFilteredHitTime(Time)
@@ -32,8 +32,8 @@ struct FMassHitResult
 	FMassEntityHandle OtherEntity;
 
 	/** Time when first hit was received. */
-	float HitTime = 0.f;
+	double HitTime = 0.;
 
 	/** Time used for filtering frequent hits. */
-	float LastFilteredHitTime = 0.f;
+	double LastFilteredHitTime = 0.;
 };
diff --git a/Engine/Plugins/AI/MassAI/Source/MassAIBehavior/Public/MassLookAtFragments.h b/Engine/Plugins/AI/MassAI/Source/MassAIBehavior/Public/MassLookAtFragments.h
--- a/Engine/Plugins/AI/MassAI/Source/MassAIBehavior/Public/MassLookAtFragments.h
+++ b/Engine/Plugins/AI/MassAI/Source/MassAIBehavior/Public/MassLookAtFragments.h
@@ -79,7 +79,7 @@ struct MASSAIBEHAVIOR_API FMassLookAtFragment : public FMassFragment
 
 	/** Start time of the current gaze. */
 	UPROPERTY(Transient)
-	float GazeStartTime = 0.0f;
+	double GazeStartTime = 0.;
 
 	/** Duration of the current gaze. */
 	UPROPERTY(Transient)
diff --git a/Engine/Plugins/AI/MassAI/Source/MassAIBehavior/Public/MassLookAtProcessors.h b/Engine/Plugins/AI/MassAI/Source/MassAIBehavior/Public/MassLookAtProcessors.h
--- a/Engine/Plugins/AI/MassAI/Source/MassAIBehavior/Public/MassLookAtProcessors.h
+++ b/Engine/Plugins/AI/MassAI/Source/MassAIBehavior/Public/MassLookAtProcessors.h
@@ -31,7 +31,7 @@ protected:
 	virtual void Execute(FMassEntityManager& EntityManager, FMassExecutionContext& Context) override;
 
 	/** Selects a nearby target if possible or use a random fixed direction */
-	void FindNewGazeTarget(const UMassNavigationSubsystem& MassNavSystem, const FMassEntityManager& EntitySubsystem, const float CurrentTime, const FTransform& Transform, FMassLookAtFragment& LookAt) const;
+	void FindNewGazeTarget(const UMassNavigationSubsystem& MassNavSystem, const FMassEntityManager& EntitySubsystem, const double CurrentTime, const FTransform& Transform, FMassLookAtFragment& LookAt) const;
 
 	/** Updates look direction based on look at trajectory. */
 	void UpdateLookAtTrajectory(const FTransform& Transform, const FMassZoneGraphLaneLocationFragment& ZoneGraphLocation,
diff --git a/Engine/Plugins/AI/MassAI/Source/MassAIBehavior/Public/MassStateTreeFragments.h b/Engine/Plugins/AI/MassAI/Source/MassAIBehavior/Public/MassStateTreeFragments.h
--- a/Engine/Plugins/AI/MassAI/Source/MassAIBehavior/Public/MassStateTreeFragments.h
+++ b/Engine/Plugins/AI/MassAI/Source/MassAIBehavior/Public/MassStateTreeFragments.h
@@ -18,7 +18,7 @@ struct MASSAIBEHAVIOR_API FMassStateTreeInstanceFragment : public FMassFragment
 	FMassStateTreeInstanceHandle InstanceHandle;
 
 	/** The last update time use to calculate ticking delta time. */
-	float LastUpdateTimeInSeconds = 0.0f;
+	double LastUpdateTimeInSeconds = 0.;
 };
 
 
diff --git a/Engine/Plugins/AI/MassAI/Source/MassAIBehavior/Public/Tasks/MassFindSmartObjectTask.h b/Engine/Plugins/AI/MassAI/Source/MassAIBehavior/Public/Tasks/MassFindSmartObjectTask.h
--- a/Engine/Plugins/AI/MassAI/Source/MassAIBehavior/Public/Tasks/MassFindSmartObjectTask.h
+++ b/Engine/Plugins/AI/MassAI/Source/MassAIBehavior/Public/Tasks/MassFindSmartObjectTask.h
@@ -1,4 +1,4 @@
-﻿// Copyright Epic Games, Inc. All Rights Reserved.
+// Copyright Epic Games, Inc. All Rights Reserved.
 
 #pragma once
 #include "MassStateTreeTypes.h"
@@ -29,7 +29,7 @@ struct MASSAIBEHAVIOR_API FMassFindSmartObjectTaskInstanceData
 
 	/** Next update time; task will not do anything when Tick gets called before that time */
 	UPROPERTY()
-	float NextUpdate = 0.f;
+	double NextUpdate = 0.;
 
 	/** Last lane where the smart objects were searched. */
 	UPROPERTY()
diff --git a/Engine/Plugins/AI/MassAI/Source/MassAIDebug/Private/GameplayDebuggerCategory_Mass.cpp b/Engine/Plugins/AI/MassAI/Source/MassAIDebug/Private/GameplayDebuggerCategory_Mass.cpp
--- a/Engine/Plugins/AI/MassAI/Source/MassAIDebug/Private/GameplayDebuggerCategory_Mass.cpp
+++ b/Engine/Plugins/AI/MassAI/Source/MassAIDebug/Private/GameplayDebuggerCategory_Mass.cpp
@@ -48,12 +48,12 @@ namespace UE::Mass::Debug
 	FMassEntityHandle GetBestEntity(const FVector ViewLocation, const FVector ViewDirection, const TConstArrayView<FMassEntityHandle> Entities, const TConstArrayView<FVector> Locations, const bool bLimitAngle)
 	{
 		// Reusing similar algorithm as UGameplayDebuggerLocalController for now 
-		constexpr float MaxScanDistanceSq = 25000.0f * 25000.0f;
-		constexpr float MinViewDirDot = 0.707f; // 45 degrees
+		constexpr FVector::FReal MaxScanDistanceSq = 25000. * 25000.;
+		constexpr FVector::FReal MinViewDirDot = 0.707; // 45 degrees
 
 		checkf(Entities.Num() == Locations.Num(), TEXT("Both Entities and Locations lists are expected to be of the same size: %d vs %d"), Entities.Num(), Locations.Num());
 		
-		float BestScore = bLimitAngle ? MinViewDirDot : (-1.f - KINDA_SMALL_NUMBER);	
+		FVector::FReal BestScore = bLimitAngle ? MinViewDirDot : (-1. - KINDA_SMALL_NUMBER);
 		FMassEntityHandle BestEntity;
 
 		for (int i = 0; i < Entities.Num(); ++i)
@@ -64,14 +64,14 @@ namespace UE::Mass::Debug
 			}
 			
 			const FVector DirToEntity = (Locations[i] - ViewLocation);
-			const float DistToEntitySq = DirToEntity.SizeSquared();
+			const FVector::FReal DistToEntitySq = DirToEntity.SizeSquared();
 			if (DistToEntitySq > MaxScanDistanceSq)
 			{
 				continue;
 			}
 
 			const FVector DirToEntityNormal = (FMath::IsNearlyZero(DistToEntitySq)) ? ViewDirection : (DirToEntity / FMath::Sqrt(DistToEntitySq));
-			const float ViewDot = FVector::DotProduct(ViewDirection, DirToEntityNormal);
+			const FVector::FReal ViewDot = FVector::DotProduct(ViewDirection, DirToEntityNormal);
 			if (ViewDot > BestScore)
 			{
 				BestScore = ViewDot;
@@ -381,7 +381,7 @@ void FGameplayDebuggerCategory_Mass::CollectData(APlayerController* OwnerPC, AAc
 		EntityQuery.AddRequirement<FMassZoneGraphShortPathFragment>(EMassFragmentAccess::ReadOnly);
 		EntityQuery.AddRequirement<FMassSmartObjectUserFragment>(EMassFragmentAccess::ReadOnly, EMassFragmentPresence::Optional);
 
-		const float CurrentTime = World->GetTimeSeconds();
+		const double CurrentTime = World->GetTimeSeconds();
 		
 		UMassStateTreeSubsystem* MassStateTreeSubsystem = World->GetSubsystem<UMassStateTreeSubsystem>();
 		UMassSignalSubsystem* SignalSubsystem = World->GetSubsystem<UMassSignalSubsystem>();
@@ -416,8 +416,8 @@ void FGameplayDebuggerCategory_Mass::CollectData(APlayerController* OwnerPC, AAc
 				const bool bHasSOUser = (SOUserList.Num() > 0);
 
 				const UGameplayDebuggerUserSettings* Settings = GetDefault<UGameplayDebuggerUserSettings>();
-				const float MaxViewDistance = Settings->MaxViewDistance;
-				const float MinViewDirDot = FMath::Cos(FMath::DegreesToRadians(Settings->MaxViewAngle));
+				const FVector::FReal MaxViewDistance = Settings->MaxViewDistance;
+				const FVector::FReal MinViewDirDot = FMath::Cos(FMath::DegreesToRadians(Settings->MaxViewAngle));
 
 				const UStateTree* StateTree = SharedStateTree.StateTree;
 
@@ -428,12 +428,12 @@ void FGameplayDebuggerCategory_Mass::CollectData(APlayerController* OwnerPC, AAc
 					
 					// Cull entities
 					const FVector DirToEntity = EntityLocation - ViewLocation;
-					const float DistanceToEntitySq = DirToEntity.SquaredLength();
+					const FVector::FReal DistanceToEntitySq = DirToEntity.SquaredLength();
 					if (DistanceToEntitySq > FMath::Square(MaxViewDistance))
 					{
 						continue;
 					}
-					const float ViewDot = FVector::DotProduct(DirToEntity.GetSafeNormal(), ViewDirection);
+					const FVector::FReal ViewDot = FVector::DotProduct(DirToEntity.GetSafeNormal(), ViewDirection);
 					if (ViewDot < MinViewDirDot)
 					{
 						continue;
@@ -471,7 +471,7 @@ void FGameplayDebuggerCategory_Mass::CollectData(APlayerController* OwnerPC, AAc
 					AddShape(FGameplayDebuggerShape::MakeArrow(MoveBasePos - MoveTarget.Forward * Radius.Radius, MoveBasePos + MoveTarget.Forward * Radius.Radius, 10.0f, 2.0f, FColorList::MediumVioletRed));
 
 					// Look at
-					constexpr float LookArrowLength = 100.0f;
+					constexpr FVector::FReal LookArrowLength = 100.;
 					BasePos = EntityLocation + FVector(0,0,EyeHeight);
 
 					if (bHasLookAt)
@@ -487,7 +487,7 @@ void FGameplayDebuggerCategory_Mass::CollectData(APlayerController* OwnerPC, AAc
 								TargetPosition.Z = BasePos.Z;
 								AddShape(FGameplayDebuggerShape::MakeCircle(TargetPosition, FVector::UpVector, Radius.Radius, FColor::Red));
 
-								const float TargetDistance = FMath::Max(LookArrowLength, FVector::DotProduct(WorldLookDirection, TargetPosition - BasePos));
+								const FVector::FReal TargetDistance = FMath::Max(LookArrowLength, FVector::DotProduct(WorldLookDirection, TargetPosition - BasePos));
 								AddShape(FGameplayDebuggerShape::MakeSegment(BasePos, BasePos + WorldLookDirection * TargetDistance, FColorList::LightGrey));
 								bLookArrowDrawn = true;
 							}
@@ -613,7 +613,7 @@ void FGameplayDebuggerCategory_Mass::CollectData(APlayerController* OwnerPC, AAc
 						if (bHasLookAt)
 						{
 							const FMassLookAtFragment& LookAt = LookAtList[EntityIndex];
-							const float RemainingTime = LookAt.GazeDuration - (CurrentTime - LookAt.GazeStartTime);
+							const double RemainingTime = LookAt.GazeDuration - (CurrentTime - LookAt.GazeStartTime);
 							Status += FString::Printf(TEXT("{turquoise}%s/%s {lightgrey}%.1f\n"),
 								*UEnum::GetDisplayValueAsText(LookAt.LookAtMode).ToString(), *UEnum::GetDisplayValueAsText(LookAt.RandomGazeMode).ToString(), RemainingTime);
 						}
@@ -621,9 +621,9 @@ void FGameplayDebuggerCategory_Mass::CollectData(APlayerController* OwnerPC, AAc
 						if (!Status.IsEmpty())
 						{
 							BasePos += FVector(0,0,50);
-							constexpr float ViewWeight = 0.6f; // Higher the number the more the view angle affects the score.
-							const float ViewScale = 1.f - (ViewDot / MinViewDirDot); // Zero at center of screen
-							NearEntityDescriptions.Emplace(DistanceToEntitySq * ((1.0f - ViewWeight) + ViewScale * ViewWeight), BasePos, Status);
+							constexpr FVector::FReal ViewWeight = 0.6f; // Higher the number the more the view angle affects the score.
+							const FVector::FReal ViewScale = 1. - (ViewDot / MinViewDirDot); // Zero at center of screen
+							NearEntityDescriptions.Emplace(FloatCastChecked<float>(DistanceToEntitySq * ((1. - ViewWeight) + ViewScale * ViewWeight), UE::LWC::DefaultFloatPrecision), BasePos, Status);
 						}
 					}
 				}
@@ -732,24 +732,24 @@ void FGameplayDebuggerCategory_Mass::DrawData(APlayerController* OwnerPC, FGamep
 			Rect.Alpha = 0.0f;
 
 			// Calculate transparency based on how much more important rects are overlapping the new rect.
-			const float Area = FMath::Max(0.0f, Rect.Max.X - Rect.Min.X) * FMath::Max(0.0f, Rect.Max.Y - Rect.Min.Y);
-			const float InvArea = Area > KINDA_SMALL_NUMBER ? 1.0f / Area : 0.0f;
-			float Coverage = 0.0;
+			const FVector::FReal Area = FMath::Max(0.0, Rect.Max.X - Rect.Min.X) * FMath::Max(0.0, Rect.Max.Y - Rect.Min.Y);
+			const FVector::FReal InvArea = Area > KINDA_SMALL_NUMBER ? 1.0 / Area : 0.0;
+			FVector::FReal Coverage = 0.0;
 
 			for (const FEntityLayoutRect& Other : Layout)
 			{
 				// Calculate rect intersection
-				const float MinX = FMath::Max(Rect.Min.X, Other.Min.X);
-				const float MinY = FMath::Max(Rect.Min.Y, Other.Min.Y);
-				const float MaxX = FMath::Min(Rect.Max.X, Other.Max.X);
-				const float MaxY = FMath::Min(Rect.Max.Y, Other.Max.Y);
+				const FVector::FReal MinX = FMath::Max(Rect.Min.X, Other.Min.X);
+				const FVector::FReal MinY = FMath::Max(Rect.Min.Y, Other.Min.Y);
+				const FVector::FReal MaxX = FMath::Min(Rect.Max.X, Other.Max.X);
+				const FVector::FReal MaxY = FMath::Min(Rect.Max.Y, Other.Max.Y);
 
 				// return zero area if not overlapping
-				const float IntersectingArea = FMath::Max(0.0f, MaxX - MinX) * FMath::Max(0.0f, MaxY - MinY);
+				const FVector::FReal IntersectingArea = FMath::Max(0.0, MaxX - MinX) * FMath::Max(0.0, MaxY - MinY);
 				Coverage += (IntersectingArea * InvArea) * Other.Alpha;
 			}
 
-			Rect.Alpha = FMath::Square(1.0f - FMath::Min(Coverage, 1.0f));
+			Rect.Alpha = FloatCastChecked<float>(FMath::Square(1.0 - FMath::Min(Coverage, 1.0)), UE::LWC::DefaultFloatPrecision);
 			
 			if (Rect.Alpha > KINDA_SMALL_NUMBER)
 			{
@@ -768,9 +768,13 @@ void FGameplayDebuggerCategory_Mass::DrawData(APlayerController* OwnerPC, FGamep
 		const FVector2D BackgroundPosition(Rect.Min - Padding);
 		FCanvasTileItem Background(Rect.Min - Padding, Rect.Max - Rect.Min + Padding * 2.0f, FLinearColor(0.0f, 0.0f, 0.0f, 0.35f * Rect.Alpha));
 		Background.BlendMode = SE_BLEND_TranslucentAlphaOnly;
-		CanvasContext.DrawItem(Background, BackgroundPosition.X, BackgroundPosition.Y);
-		
-		CanvasContext.PrintAt(Rect.Min.X, Rect.Min.Y, FColor::White, Rect.Alpha, Desc.Description);
+		CanvasContext.DrawItem(Background
+			, FloatCastChecked<float>(BackgroundPosition.X, UE::LWC::DefaultFloatPrecision)
+			, FloatCastChecked<float>(BackgroundPosition.Y, UE::LWC::DefaultFloatPrecision));
+
+		CanvasContext.PrintAt(FloatCastChecked<float>(Rect.Min.X, UE::LWC::DefaultFloatPrecision)
+			, FloatCastChecked<float>(Rect.Min.Y, UE::LWC::DefaultFloatPrecision)
+			, FColor::White, Rect.Alpha, Desc.Description);
 	}
 
 	FGameplayDebuggerCategory::DrawData(OwnerPC, CanvasContext);
diff --git a/Engine/Plugins/AI/MassAI/Source/MassAIReplication/Public/MassReplicationPathHandlers.h b/Engine/Plugins/AI/MassAI/Source/MassAIReplication/Public/MassReplicationPathHandlers.h
--- a/Engine/Plugins/AI/MassAI/Source/MassAIReplication/Public/MassReplicationPathHandlers.h
+++ b/Engine/Plugins/AI/MassAI/Source/MassAIReplication/Public/MassReplicationPathHandlers.h
@@ -46,7 +46,7 @@ struct MASSAIREPLICATION_API FReplicatedAgentPathData
 
 	/** Server time in seconds when the action started. */
 	UPROPERTY(Transient)
-	float ActionServerStartTime = 0.0f;
+	double ActionServerStartTime = 0.;
 
 	/** Distance along current lane. */
 	UPROPERTY(Transient)
diff --git a/Engine/Plugins/AI/MassAI/Source/MassNavigation/Private/MassAvoidanceProcessors.cpp b/Engine/Plugins/AI/MassAI/Source/MassNavigation/Private/MassAvoidanceProcessors.cpp
--- a/Engine/Plugins/AI/MassAI/Source/MassNavigation/Private/MassAvoidanceProcessors.cpp
+++ b/Engine/Plugins/AI/MassAI/Source/MassNavigation/Private/MassAvoidanceProcessors.cpp
@@ -44,11 +44,11 @@ namespace UE::MassAvoidance
 	constexpr int32 MinTouchingCellCount = 4;
 	constexpr int32 MaxObstacleResults = MaxExpectedAgentsPerCell * MinTouchingCellCount;
 
-	static void FindCloseObstacles(const FVector& Center, const float SearchRadius, const FNavigationObstacleHashGrid2D& AvoidanceObstacleGrid,
+	static void FindCloseObstacles(const FVector& Center, const FVector::FReal SearchRadius, const FNavigationObstacleHashGrid2D& AvoidanceObstacleGrid,
 									TArray<FMassNavigationObstacleItem, TFixedAllocator<MaxObstacleResults>>& OutCloseEntities, const int32 MaxResults)
 	{
 		OutCloseEntities.Reset();
-		const FVector Extent(SearchRadius, SearchRadius, 0.f);
+		const FVector Extent(SearchRadius, SearchRadius, 0.);
 		const FBox QueryBox = FBox(Center - Extent, Center + Extent);
 
 		struct FSortingCell
@@ -56,24 +56,24 @@ namespace UE::MassAvoidance
 			int32 X;
 			int32 Y;
 			int32 Level;
-			float SqDist;
+			FVector::FReal SqDist;
 		};
 		TArray<FSortingCell, TInlineAllocator<64>> Cells;
 		const FVector QueryCenter = QueryBox.GetCenter();
 		
 		for (int32 Level = 0; Level < AvoidanceObstacleGrid.NumLevels; Level++)
 		{
-			const float CellSize = AvoidanceObstacleGrid.GetCellSize(Level);
+			const FVector::FReal CellSize = AvoidanceObstacleGrid.GetCellSize(Level);
 			const FNavigationObstacleHashGrid2D::FCellRect Rect = AvoidanceObstacleGrid.CalcQueryBounds(QueryBox, Level);
 			for (int32 Y = Rect.MinY; Y <= Rect.MaxY; Y++)
 			{
 				for (int32 X = Rect.MinX; X <= Rect.MaxX; X++)
 				{
-					const float CenterX = (X + 0.5f) * CellSize;
-					const float CenterY = (Y + 0.5f) * CellSize;
-					const float DX = CenterX - QueryCenter.X;
-					const float DY = CenterY - QueryCenter.Y;
-					const float SqDist = DX * DX + DY * DY;
+					const FVector::FReal CenterX = (X + 0.5) * CellSize;
+					const FVector::FReal CenterY = (Y + 0.5) * CellSize;
+					const FVector::FReal DX = CenterX - QueryCenter.X;
+					const FVector::FReal DY = CenterY - QueryCenter.Y;
+					const FVector::FReal SqDist = DX * DX + DY * DY;
 					FSortingCell SortCell;
 					SortCell.X = X;
 					SortCell.Y = Y;
@@ -104,50 +104,50 @@ namespace UE::MassAvoidance
 	}
 
 	// Adapted from ray-capsule intersection: https://iquilezles.org/www/articles/intersectors/intersectors.htm
-	static float ComputeClosestPointOfApproach(const FVector2D Pos, const FVector2D Vel, const float Rad, const FVector2D SegStart, const FVector2D SegEnd, const float TimeHoriz)
+	static FVector::FReal ComputeClosestPointOfApproach(const FVector2D Pos, const FVector2D Vel, const FVector::FReal Rad, const FVector2D SegStart, const FVector2D SegEnd, const FVector::FReal TimeHoriz)
 	{
 		const FVector2D SegDir = SegEnd - SegStart;
 		const FVector2D RelPos = Pos - SegStart;
-		const float VelSq = FVector2D::DotProduct(Vel, Vel);
-		const float SegDirSq = FVector2D::DotProduct(SegDir, SegDir);
-		const float DirVelSq = FVector2D::DotProduct(SegDir, Vel);
-		const float DirRelPosSq = FVector2D::DotProduct(SegDir, RelPos);
-		const float VelRelPosSq = FVector2D::DotProduct(Vel, RelPos);
-		const float RelPosSq = FVector2D::DotProduct(RelPos, RelPos);
-		const float A = SegDirSq * VelSq - DirVelSq * DirVelSq;
-		const float B = SegDirSq * VelRelPosSq - DirRelPosSq * DirVelSq;
-		const float C = SegDirSq * RelPosSq - DirRelPosSq * DirRelPosSq - FMath::Square(Rad) * SegDirSq;
-		const float H = FMath::Max<float>(0.f, B*B - A*C); // b^2 - ac, Using max for closest point of arrival result when no hit.
-		const float T = FMath::Abs(A) > SMALL_NUMBER ? (-B - FMath::Sqrt(H)) / A : 0.f;
-		const float Y = DirRelPosSq + T * DirVelSq;
+		const FVector::FReal VelSq = FVector2D::DotProduct(Vel, Vel);
+		const FVector::FReal SegDirSq = FVector2D::DotProduct(SegDir, SegDir);
+		const FVector::FReal DirVelSq = FVector2D::DotProduct(SegDir, Vel);
+		const FVector::FReal DirRelPosSq = FVector2D::DotProduct(SegDir, RelPos);
+		const FVector::FReal VelRelPosSq = FVector2D::DotProduct(Vel, RelPos);
+		const FVector::FReal RelPosSq = FVector2D::DotProduct(RelPos, RelPos);
+		const FVector::FReal A = SegDirSq * VelSq - DirVelSq * DirVelSq;
+		const FVector::FReal B = SegDirSq * VelRelPosSq - DirRelPosSq * DirVelSq;
+		const FVector::FReal C = SegDirSq * RelPosSq - DirRelPosSq * DirRelPosSq - FMath::Square(Rad) * SegDirSq;
+		const FVector::FReal H = FMath::Max(0., B*B - A*C); // b^2 - ac, Using max for closest point of arrival result when no hit.
+		const FVector::FReal T = FMath::Abs(A) > SMALL_NUMBER ? (-B - FMath::Sqrt(H)) / A : 0.;
+		const FVector::FReal Y = DirRelPosSq + T * DirVelSq;
 		
-		if (Y > 0.f && Y < SegDirSq) 
+		if (Y > 0. && Y < SegDirSq) 
 		{
-			return FMath::Clamp(T, 0.f, TimeHoriz);
+			return FMath::Clamp(T, 0., TimeHoriz);
 		}
 		else 
 		{
 			// caps
-			const FVector2D CapRelPos = (Y <= 0.f) ? RelPos : Pos - SegEnd;
-			const float Cb = FVector2D::DotProduct(Vel, CapRelPos);
-			const float Cc = FVector2D::DotProduct(CapRelPos, CapRelPos) - FMath::Square(Rad);
-			const float Ch = FMath::Max<float>(0.0f, Cb * Cb - VelSq * Cc);
-			const float T1 = VelSq > SMALL_NUMBER ? (-Cb - FMath::Sqrt(Ch)) / VelSq : 0.f;
-			return FMath::Clamp(T1, 0.f, TimeHoriz);
+			const FVector2D CapRelPos = (Y <= 0.) ? RelPos : Pos - SegEnd;
+			const FVector::FReal Cb = FVector2D::DotProduct(Vel, CapRelPos);
+			const FVector::FReal Cc = FVector2D::DotProduct(CapRelPos, CapRelPos) - FMath::Square(Rad);
+			const FVector::FReal Ch = FMath::Max(0., Cb * Cb - VelSq * Cc);
+			const FVector::FReal T1 = VelSq > SMALL_NUMBER ? (-Cb - FMath::Sqrt(Ch)) / VelSq : 0.;
+			return FMath::Clamp(T1, 0., TimeHoriz);
 		}
 	}
 
-	static float ComputeClosestPointOfApproach(const FVector RelPos, const FVector RelVel, const float TotalRadius, const float TimeHoriz)
+	static FVector::FReal ComputeClosestPointOfApproach(const FVector RelPos, const FVector RelVel, const FVector::FReal TotalRadius, const FVector::FReal TimeHoriz)
 	{
 		// Calculate time of impact based on relative agent positions and velocities.
-		const float A = FVector::DotProduct(RelVel, RelVel);
-		const float Inv2A = A > SMALL_NUMBER ? 1.f / (2.f * A) : 0.f;
-		const float B = FMath::Min(0.f, 2.f * FVector::DotProduct(RelVel, RelPos));
-		const float C = FVector::DotProduct(RelPos, RelPos) - FMath::Square(TotalRadius);
+		const FVector::FReal A = FVector::DotProduct(RelVel, RelVel);
+		const FVector::FReal Inv2A = A > SMALL_NUMBER ? 1. / (2. * A) : 0.;
+		const FVector::FReal B = FMath::Min(0., 2. * FVector::DotProduct(RelVel, RelPos));
+		const FVector::FReal C = FVector::DotProduct(RelPos, RelPos) - FMath::Square(TotalRadius);
 		// Using max() here gives us CPA (closest point on arrival) when there is no hit.
-		const float Discr = FMath::Sqrt(FMath::Max(0.f, B * B - 4.f * A * C));
-		const float T = (-B - Discr) * Inv2A;
-		return FMath::Clamp(T, 0.f, TimeHoriz);
+		const FVector::FReal Discr = FMath::Sqrt(FMath::Max(0., B * B - 4. * A * C));
+		const FVector::FReal T = (-B - Discr) * Inv2A;
+		return FMath::Clamp(T, 0., TimeHoriz);
 	}
 
 	static bool UseDrawDebugHelper()
@@ -178,8 +178,8 @@ namespace UE::MassAvoidance
 	static const FColor ObstacleAvoidForceColor = FColor::Magenta;
 	static const FColor ObstacleSeparationForceColor = FColor(255, 66, 66);	// Bright red
 	
-	static const FVector DebugAgentHeightOffset = FVector(0.f, 0.f, 185.f);
-	static const FVector DebugLowCylinderOffset = FVector(0.f, 0.f, 20.f);
+	static const FVector DebugAgentHeightOffset = FVector(0., 0., 185.);
+	static const FVector DebugLowCylinderOffset = FVector(0., 0., 20.);
 
 	//----------------------------------------------------------------------//
 	// Begin MassDebugUtils
@@ -228,7 +228,7 @@ namespace UE::MassAvoidance
 			return;
 		}
 
-		const float Pointyness = 1.8f;
+		const FVector::FReal Pointyness = 1.8;
 		const FVector Line = End - Start;
 		const FVector UnitV = Line.GetSafeNormal();
 		const FVector Perp = FVector::CrossProduct(UnitV, FVector::UpVector);
@@ -305,13 +305,13 @@ namespace UE::MassAvoidance
 		}
 
 		const float Thickness = 3.f;
-		const float Pointyness = 1.8f;
+		const FVector::FReal Pointyness = 1.8;
 		const FVector Line = End - Start;
 		const FVector UnitV = Line.GetSafeNormal();
 		const FVector Perp = FVector::CrossProduct(UnitV, FVector::UpVector);
 		const FVector Left = Perp - (Pointyness * UnitV);
 		const FVector Right = -Perp - (Pointyness * UnitV);
-		const float HeadSize = 0.08f * Line.Size();
+		const FVector::FReal HeadSize = 0.08 * Line.Size();
 		UE_VLOG_SEGMENT_THICK(Context.LogOwner, Context.Category, Log, Start, End, Color, (int16)Thickness, TEXT(""));
 		UE_VLOG_SEGMENT_THICK(Context.LogOwner, Context.Category, Log, End, End + HeadSize * Left, Color, (int16)Thickness, TEXT(""));
 		UE_VLOG_SEGMENT_THICK(Context.LogOwner, Context.Category, Log, End, End + HeadSize * Right, Color, (int16)Thickness, TEXT(""));
@@ -333,7 +333,7 @@ namespace UE::MassAvoidance
 
 	static void DebugDrawSummedForce(const FDebugContext& Context, const FVector& Start, const FVector& End, const FColor& Color)
 	{
-		DebugDrawArrow(Context, Start + FVector(0.f,0.f,1.f), End + FVector(0.f, 0.f, 1.f), Color, /*HeadSize*/8.f, /*Thickness*/6.f);
+		DebugDrawArrow(Context, Start + FVector(0.,0.,1.), End + FVector(0., 0., 1.), Color, /*HeadSize*/8.f, /*Thickness*/6.f);
 	}
 
 #endif // WITH_MASSGAMEPLAY_DEBUG
@@ -388,7 +388,7 @@ void UMassMovingAvoidanceProcessor::Execute(FMassEntityManager& EntityManager, F
 	EntityQuery.ForEachEntityChunk(EntityManager, Context, [this, &EntityManager](FMassExecutionContext& Context)
 	{
 		const float DeltaTime = Context.GetDeltaTimeSeconds();
-		const float CurrentTime = World->GetTimeSeconds();
+		const double CurrentTime = World->GetTimeSeconds();
 		const int32 NumEntities = Context.GetNumEntities();
 		
 		const TArrayView<FMassForceFragment> ForceList = Context.GetMutableFragmentView<FMassForceFragment>();
@@ -400,7 +400,7 @@ void UMassMovingAvoidanceProcessor::Execute(FMassEntityManager& EntityManager, F
 		const FMassMovingAvoidanceParameters& MovingAvoidanceParams = Context.GetConstSharedFragment<FMassMovingAvoidanceParameters>();
 		const FMassMovementParameters& MovementParams = Context.GetConstSharedFragment<FMassMovementParameters>();
 
-		const float InvPredictiveAvoidanceTime = 1.0f / MovingAvoidanceParams.PredictiveAvoidanceTime;
+		const FVector::FReal InvPredictiveAvoidanceTime = 1. / MovingAvoidanceParams.PredictiveAvoidanceTime;
 
 		// Arrays used to store close obstacles
 		TArray<FMassNavigationObstacleItem, TFixedAllocator<UE::MassAvoidance::MaxObstacleResults>> CloseEntities;
@@ -411,7 +411,7 @@ void UMassMovingAvoidanceProcessor::Execute(FMassEntityManager& EntityManager, F
 			FVector LocationCached;
 			FVector Forward;
 			FMassNavigationObstacleItem ObstacleItem;
-			float SqDist;
+			FVector::FReal SqDist;
 		};
 		TArray<FSortedObstacle, TFixedAllocator<UE::MassAvoidance::MaxObstacleResults>> ClosestObstacles;
 
@@ -420,7 +420,7 @@ void UMassMovingAvoidanceProcessor::Execute(FMassEntityManager& EntityManager, F
 		{
 			FVector Position = FVector::ZeroVector;
 			FVector Normal = FVector::ZeroVector;
-			float Distance = 0.f;
+			FVector::FReal Distance = 0.;
 		};
 		TArray<FEnvironmentContact, TInlineAllocator<16>> Contacts;
 
@@ -453,39 +453,39 @@ void UMassMovingAvoidanceProcessor::Execute(FMassEntityManager& EntityManager, F
 
 			// Smaller steering max accel makes the steering more "calm" but less opportunistic, may not find solution, or gets stuck.
 			// Max contact accel should be quite a big bigger than steering so that collision response is firm. 
-			const float MaxSteerAccel = MovementParams.MaxAcceleration;
-			const float MaximumSpeed = MovementParams.MaxSpeed;
+			const FVector::FReal MaxSteerAccel = MovementParams.MaxAcceleration;
+			const FVector::FReal MaximumSpeed = MovementParams.MaxSpeed;
 
 			const FVector AgentLocation = Location.GetTransform().GetTranslation();
-			const FVector AgentVelocity = FVector(Velocity.Value.X, Velocity.Value.Y, 0.0f);
+			const FVector AgentVelocity = FVector(Velocity.Value.X, Velocity.Value.Y, 0.);
 			
-			const float AgentRadius = Radius.Radius;
-			const float SeparationAgentRadius = Radius.Radius * MovingAvoidanceParams.SeparationRadiusScale;
-			const float PredictiveAvoidanceAgentRadius = Radius.Radius * MovingAvoidanceParams.PredictiveAvoidanceRadiusScale;
+			const FVector::FReal AgentRadius = Radius.Radius;
+			const FVector::FReal SeparationAgentRadius = Radius.Radius * MovingAvoidanceParams.SeparationRadiusScale;
+			const FVector::FReal PredictiveAvoidanceAgentRadius = Radius.Radius * MovingAvoidanceParams.PredictiveAvoidanceRadiusScale;
 			
 			FVector SteeringForce = Force.Value;
 
 			// Near start and end fades are used to subdue the avoidance at the start and end of the path.
-			float NearStartFade = 1.0f;
-			float NearEndFade = 1.0f;
+			FVector::FReal NearStartFade = 1.;
+			FVector::FReal NearEndFade = 1.;
 
 			if (MoveTarget.GetPreviousAction() != EMassMovementAction::Move)
 			{
 				// Fade in avoidance when transitioning from other than move action.
 				// I.e. the standing behavior may move the agents so close to each,
 				// and that causes the separation to push them out quickly when avoidance is activated. 
-				NearStartFade = FMath::Min((CurrentTime - MoveTarget.GetCurrentActionStartTime()) / MovingAvoidanceParams.StartOfPathDuration, 1.0f);
+				NearStartFade = FMath::Min((CurrentTime - MoveTarget.GetCurrentActionStartTime()) / MovingAvoidanceParams.StartOfPathDuration, 1.);
 			}
 
 			if (MoveTarget.IntentAtGoal == EMassMovementAction::Stand)
 			{
 				// Estimate approach based on current desired speed.
-				const float ApproachDistance = FMath::Max(1.0f, MovingAvoidanceParams.EndOfPathDuration * MoveTarget.DesiredSpeed.Get());
-				NearEndFade = FMath::Clamp(MoveTarget.DistanceToGoal / ApproachDistance, 0.f, 1.f);
+				const FVector::FReal ApproachDistance = FMath::Max<FVector::FReal>(1., MovingAvoidanceParams.EndOfPathDuration * MoveTarget.DesiredSpeed.Get());
+				NearEndFade = FMath::Clamp(MoveTarget.DistanceToGoal / ApproachDistance, 0., 1.);
 			}
 			
-			const float NearStartScaling = FMath::Lerp(MovingAvoidanceParams.StartOfPathAvoidanceScale, 1.0f, NearStartFade);
-			const float NearEndScaling = FMath::Lerp(MovingAvoidanceParams.EndOfPathAvoidanceScale, 1.0f, NearEndFade);
+			const FVector::FReal NearStartScaling = FMath::Lerp<FVector::FReal>(MovingAvoidanceParams.StartOfPathAvoidanceScale, 1., NearStartFade);
+			const FVector::FReal NearEndScaling = FMath::Lerp<FVector::FReal>(MovingAvoidanceParams.EndOfPathAvoidanceScale, 1., NearEndFade);
 			
 #if WITH_MASSGAMEPLAY_DEBUG && UNSAFE_FOR_MT
 			const UE::MassAvoidance::FDebugContext BaseDebugContext(this, LogAvoidance, World, Entity);
@@ -497,7 +497,7 @@ void UMassMovingAvoidanceProcessor::Execute(FMassEntityManager& EntityManager, F
 			{
 				// Draw agent
 				const FString Text = FString::Printf(TEXT("%i"), Entity.Index);
-				DebugDrawCylinder(BaseDebugContext, AgentLocation, AgentLocation + UE::MassAvoidance::DebugAgentHeightOffset, AgentRadius+1.f, UE::MassAvoidance::CurrentAgentColor, Text);
+				DebugDrawCylinder(BaseDebugContext, AgentLocation, AgentLocation + UE::MassAvoidance::DebugAgentHeightOffset, static_cast<float>(AgentRadius+1.), UE::MassAvoidance::CurrentAgentColor, Text);
 
 				DebugDrawSphere(BaseDebugContext, AgentLocation, 10.f, UE::MassAvoidance::CurrentAgentColor);
 
@@ -539,22 +539,22 @@ void UMassMovingAvoidanceProcessor::Execute(FMassEntityManager& EntityManager, F
 				{
 					const FVector EdgeDiff = Edge.End - Edge.Start;
 					FVector EdgeDir = FVector::ZeroVector;
-					float EdgeLength = 0.0f;
+					FVector::FReal EdgeLength = 0.;
 					EdgeDiff.ToDirectionAndLength(EdgeDir, EdgeLength);
 
 					const FVector AgentToEdgeStart = AgentLocation - Edge.Start;
-					const float DistAlongEdge = FVector::DotProduct(EdgeDir, AgentToEdgeStart);
-					const float DistAwayFromEdge = FVector::DotProduct(Edge.LeftDir, AgentToEdgeStart);
+					const FVector::FReal DistAlongEdge = FVector::DotProduct(EdgeDir, AgentToEdgeStart);
+					const FVector::FReal DistAwayFromEdge = FVector::DotProduct(Edge.LeftDir, AgentToEdgeStart);
 
-					float ConDist = 0.0f;
+					FVector::FReal ConDist = 0.;
 					FVector ConNorm = FVector::ForwardVector;
 					FVector ConPos = FVector::ZeroVector;
 					bool bDirectlyBehindEdge = false;
 					
-					if (DistAwayFromEdge < 0.0f)
+					if (DistAwayFromEdge < 0.)
 					{
 						// Inside or behind the edge
-						if (DistAlongEdge < 0.0f)
+						if (DistAlongEdge < 0.)
 						{
 							ConPos = Edge.Start;
 							ConNorm = -EdgeDir;
@@ -570,13 +570,13 @@ void UMassMovingAvoidanceProcessor::Execute(FMassEntityManager& EntityManager, F
 						{
 							ConPos = Edge.Start + EdgeDir * DistAlongEdge;
 							ConNorm = Edge.LeftDir;
-							ConDist = 0.0f;
+							ConDist = 0.;
 							bDirectlyBehindEdge = true;
 						}
 					}
 					else
 					{
-						if (DistAlongEdge < 0.0f)
+						if (DistAlongEdge < 0.)
 						{
 							// Start Corner
 							ConPos = Edge.Start;
@@ -601,7 +601,7 @@ void UMassMovingAvoidanceProcessor::Execute(FMassEntityManager& EntityManager, F
 					bool bAdd = true;
 					for (int ContactIndex = 0; ContactIndex < Contacts.Num(); ContactIndex++)
 					{
-						if (FVector::DotProduct(Contacts[ContactIndex].Normal, ConNorm) > 0.f && FMath::Abs(FVector::DotProduct(ConNorm, Contacts[ContactIndex].Position - ConPos)) < (10.f/*cm*/))
+						if (FVector::DotProduct(Contacts[ContactIndex].Normal, ConNorm) > 0. && FMath::Abs(FVector::DotProduct(ConNorm, Contacts[ContactIndex].Position - ConPos)) < (10./*cm*/))
 						{
 							// Contacts are on same place, merge
 							if (ConDist < Contacts[ContactIndex].Distance)
@@ -630,33 +630,33 @@ void UMassMovingAvoidanceProcessor::Execute(FMassEntityManager& EntityManager, F
 					if (!bDirectlyBehindEdge)
 					{
 						// Avoid edges
-						const float CPA = UE::MassAvoidance::ComputeClosestPointOfApproach(FVector2D(AgentLocation), FVector2D(DesiredVelocity), AgentRadius,
+						const FVector::FReal CPA = UE::MassAvoidance::ComputeClosestPointOfApproach(FVector2D(AgentLocation), FVector2D(DesiredVelocity), AgentRadius,
 							FVector2D(Edge.Start), FVector2D(Edge.End), MovingAvoidanceParams.PredictiveAvoidanceTime);
 						const FVector HitAgentPos = AgentLocation + DesiredVelocity * CPA;
-						const float EdgeT = UE::MassNavigation::ProjectPtSeg(FVector2D(HitAgentPos), FVector2D(Edge.Start), FVector2D(Edge.End));
+						const FVector::FReal EdgeT = UE::MassNavigation::ProjectPtSeg(FVector2D(HitAgentPos), FVector2D(Edge.Start), FVector2D(Edge.End));
 						const FVector HitObPos = FMath::Lerp(Edge.Start, Edge.End, EdgeT);
 
 						// Calculate penetration at CPA
 						FVector AvoidRelPos = HitAgentPos - HitObPos;
-						AvoidRelPos.Z = 0.f;	// @todo AT: ignore the z component for now until we clamp the height of obstacles
-						const float AvoidDist = AvoidRelPos.Size();
-						const FVector AvoidNormal = AvoidDist > 0.0f ? (AvoidRelPos / AvoidDist) : FVector::ForwardVector;
+						AvoidRelPos.Z = 0.;	// @todo AT: ignore the z component for now until we clamp the height of obstacles
+						const FVector::FReal AvoidDist = AvoidRelPos.Size();
+						const FVector AvoidNormal = AvoidDist > 0. ? (AvoidRelPos / AvoidDist) : FVector::ForwardVector;
 
-						const float AvoidPen = (PredictiveAvoidanceAgentRadius + MovingAvoidanceParams.PredictiveAvoidanceDistance) - AvoidDist;
-						const float AvoidMag = FMath::Square(FMath::Clamp(AvoidPen / MovingAvoidanceParams.PredictiveAvoidanceDistance, 0.f, 1.f));
-						const float AvoidMagDist = 1.f + FMath::Square(1.f - CPA * InvPredictiveAvoidanceTime);
+						const FVector::FReal AvoidPen = (PredictiveAvoidanceAgentRadius + MovingAvoidanceParams.PredictiveAvoidanceDistance) - AvoidDist;
+						const FVector::FReal AvoidMag = FMath::Square(FMath::Clamp(AvoidPen / MovingAvoidanceParams.PredictiveAvoidanceDistance, 0., 1.));
+						const FVector::FReal AvoidMagDist = 1. + FMath::Square(1. - CPA * InvPredictiveAvoidanceTime);
 						const FVector AvoidForce = AvoidNormal * AvoidMag * AvoidMagDist * MovingAvoidanceParams.EnvironmentPredictiveAvoidanceStiffness * NearEndScaling; // Predictive avoidance against environment is tuned down towards the end of the path
 
 						SteeringForce += AvoidForce;
 
 #if WITH_MASSGAMEPLAY_DEBUG && UNSAFE_FOR_MT
 						// Draw contact normal
-						UE::MassAvoidance::DebugDrawArrow(ObstacleDebugContext, ConPos, ConPos + 50.f * ConNorm, UE::MassAvoidance::ObstacleContactNormalColor, /*HeadSize=*/ 5.f);
+						UE::MassAvoidance::DebugDrawArrow(ObstacleDebugContext, ConPos, ConPos + 50. * ConNorm, UE::MassAvoidance::ObstacleContactNormalColor, /*HeadSize=*/ 5.f);
 						UE::MassAvoidance::DebugDrawSphere(ObstacleDebugContext, ConPos, 2.5f, UE::MassAvoidance::ObstacleContactNormalColor);
 
 						// Draw hit pos with edge
 						UE::MassAvoidance::DebugDrawLine(ObstacleDebugContext, AgentLocation, HitAgentPos, UE::MassAvoidance::ObstacleAvoidForceColor);
-						UE::MassAvoidance::DebugDrawCylinder(ObstacleDebugContext, HitAgentPos, HitAgentPos + UE::MassAvoidance::DebugAgentHeightOffset, AgentRadius, UE::MassAvoidance::ObstacleAvoidForceColor);
+						UE::MassAvoidance::DebugDrawCylinder(ObstacleDebugContext, HitAgentPos, HitAgentPos + UE::MassAvoidance::DebugAgentHeightOffset, static_cast<float>(AgentRadius), UE::MassAvoidance::ObstacleAvoidForceColor);
 
 						// Draw avoid obstacle force
 						UE::MassAvoidance::DebugDrawForce(ObstacleDebugContext, HitObPos, HitObPos + AvoidForce, UE::MassAvoidance::ObstacleAvoidForceColor);
@@ -678,11 +678,11 @@ void UMassMovingAvoidanceProcessor::Execute(FMassEntityManager& EntityManager, F
 				for (int ContactIndex = 0; ContactIndex < Contacts.Num(); ContactIndex++) 
 				{
 					const FVector ConNorm = Contacts[ContactIndex].Normal.GetSafeNormal();
-					const float ContactDist = Contacts[ContactIndex].Distance;
+					const FVector::FReal ContactDist = Contacts[ContactIndex].Distance;
 
 					// Separation force (stay away from obstacles if possible)
-					const float SeparationPenalty = (SeparationAgentRadius + MovingAvoidanceParams.EnvironmentSeparationDistance) - ContactDist;
-					const float SeparationMag = UE::MassNavigation::Smooth(FMath::Clamp(SeparationPenalty / MovingAvoidanceParams.EnvironmentSeparationDistance, 0.f, 1.f));
+					const FVector::FReal SeparationPenalty = (SeparationAgentRadius + MovingAvoidanceParams.EnvironmentSeparationDistance) - ContactDist;
+					const FVector::FReal SeparationMag = UE::MassNavigation::Smooth(FMath::Clamp(SeparationPenalty / MovingAvoidanceParams.EnvironmentSeparationDistance, 0., 1.));
 					const FVector SeparationForce = ConNorm * MovingAvoidanceParams.EnvironmentSeparationStiffness * SeparationMag;
 
 					SteeringForce += SeparationForce;
@@ -710,7 +710,7 @@ void UMassMovingAvoidanceProcessor::Execute(FMassEntityManager& EntityManager, F
 						DebugDrawLine(ObstacleDebugContext, UE::MassAvoidance::DebugAgentHeightOffset + Edge.Start,
 							UE::MassAvoidance::DebugAgentHeightOffset + Edge.End, UE::MassAvoidance::ObstacleColor, /*Thickness=*/2.f);
 						const FVector Middle = UE::MassAvoidance::DebugAgentHeightOffset + 0.5f * (Edge.Start + Edge.End);
-						DebugDrawArrow(ObstacleDebugContext, Middle, Middle + 10.f * FVector::CrossProduct((Edge.End - Edge.Start), FVector::UpVector).GetSafeNormal(), UE::MassAvoidance::ObstacleColor, /*HeadSize=*/2.f);
+						DebugDrawArrow(ObstacleDebugContext, Middle, Middle + 10. * FVector::CrossProduct((Edge.End - Edge.Start), FVector::UpVector).GetSafeNormal(), UE::MassAvoidance::ObstacleColor, /*HeadSize=*/2.f);
 					}
 				}
 #endif // WITH_MASSGAMEPLAY_DEBUG
@@ -728,7 +728,7 @@ void UMassMovingAvoidanceProcessor::Execute(FMassEntityManager& EntityManager, F
 			UE::MassAvoidance::FindCloseObstacles(AgentLocation, MovingAvoidanceParams.ObstacleDetectionDistance, AvoidanceObstacleGrid, CloseEntities, UE::MassAvoidance::MaxObstacleResults);
 
 			// Remove unwanted and find the closests in the CloseEntities
-			const float DistanceCutOffSqr = FMath::Square(MovingAvoidanceParams.ObstacleDetectionDistance);
+			const FVector::FReal DistanceCutOffSqr = FMath::Square(MovingAvoidanceParams.ObstacleDetectionDistance);
 			ClosestObstacles.Reset();
 			for (const FNavigationObstacleHashGrid2D::ItemIDType OtherEntity : CloseEntities)
 			{
@@ -749,7 +749,7 @@ void UMassMovingAvoidanceProcessor::Execute(FMassEntityManager& EntityManager, F
 				const FTransform& Transform = EntityManager.GetFragmentDataChecked<FTransformFragment>(OtherEntity.Entity).GetTransform();
 				const FVector OtherLocation = Transform.GetLocation();
 				
-				const float SqDist = FVector::DistSquared(AgentLocation, OtherLocation);
+				const FVector::FReal SqDist = FVector::DistSquared(AgentLocation, OtherLocation);
 				if (SqDist > DistanceCutOffSqr)
 				{
 					continue;
@@ -849,22 +849,22 @@ void UMassMovingAvoidanceProcessor::Execute(FMassEntityManager& EntityManager, F
 				{
 					// If the other obstacle cannot avoid us, try to avoid the local minima they create between the wall and their collider.
 					// If the space between edge and collider is less than MinClearance, make the agent to avoid the gap.
-					const float MinClearance = 2.f * AgentRadius * MovingAvoidanceParams. StaticObstacleClearanceScale;
+					const FVector::FReal MinClearance = 2. * AgentRadius * MovingAvoidanceParams.StaticObstacleClearanceScale;
 					
 					// Find the maximum distance from edges that are too close.
-					float MaxDist = -1.f;
+					FVector::FReal MaxDist = -1.;
 					FVector ClosestPoint = FVector::ZeroVector;
 					for (const FNavigationAvoidanceEdge& Edge : NavEdges.AvoidanceEdges)
 					{
 						const FVector Point = FMath::ClosestPointOnSegment(Collider.Location, Edge.Start, Edge.End);
 						const FVector Offset = Collider.Location - Point;
-						if (FVector::DotProduct(Offset, Edge.LeftDir) < 0.f)
+						if (FVector::DotProduct(Offset, Edge.LeftDir) < 0.)
 						{
 							// Behind the edge, ignore.
 							continue;
 						}
 
-						const float OffsetLength = Offset.Length();
+						const FVector::FReal OffsetLength = Offset.Length();
 						const bool bTooNarrow = (OffsetLength - Collider.Radius) < MinClearance; 
 						if (bTooNarrow)
 						{
@@ -876,7 +876,7 @@ void UMassMovingAvoidanceProcessor::Execute(FMassEntityManager& EntityManager, F
 						}
 					}
 
-					if (MaxDist != -1.f)
+					if (MaxDist != -1.)
 					{
 						// Set up forced normal to avoid the gap between collider and edge.
 						ForcedNormal = (Collider.Location - ClosestPoint).GetSafeNormal();
@@ -885,25 +885,25 @@ void UMassMovingAvoidanceProcessor::Execute(FMassEntityManager& EntityManager, F
 				}
 
 				FVector RelPos = AgentLocation - Collider.Location;
-				RelPos.Z = 0.f; // we assume we work on a flat plane for now
+				RelPos.Z = 0.; // we assume we work on a flat plane for now
 				const FVector RelVel = DesVel - Collider.Velocity;
-				const float ConDist = RelPos.Size();
-				const FVector ConNorm = ConDist > 0.f ? RelPos / ConDist : FVector::ForwardVector;
+				const FVector::FReal ConDist = RelPos.Size();
+				const FVector ConNorm = ConDist > 0. ? RelPos / ConDist : FVector::ForwardVector;
 
 				FVector SeparationNormal = ConNorm;
 				if (bHasForcedNormal)
 				{
 					// The more head on the collisions is, the more we should avoid towards the forced direction.
 					const FVector RelVelNorm = RelVel.GetSafeNormal();
-					const float Blend = FMath::Max(0.0f, -FVector::DotProduct(ConNorm, RelVelNorm));
+					const FVector::FReal Blend = FMath::Max(0., -FVector::DotProduct(ConNorm, RelVelNorm));
 					SeparationNormal = FMath::Lerp(ConNorm, ForcedNormal, Blend).GetSafeNormal();
 				}
 				
-				const float StandingScaling = Collider.bIsMoving ? 1.0f : MovingAvoidanceParams.StandingObstacleAvoidanceScale; // Care less about standing agents so that we can push through standing crowd.
+				const FVector::FReal StandingScaling = Collider.bIsMoving ? 1. : MovingAvoidanceParams.StandingObstacleAvoidanceScale; // Care less about standing agents so that we can push through standing crowd.
 				
 				// Separation force (stay away from agents if possible)
-				const float PenSep = (SeparationAgentRadius + Collider.Radius + MovingAvoidanceParams.ObstacleSeparationDistance) - ConDist;
-				const float SeparationMag = FMath::Square(FMath::Clamp(PenSep / MovingAvoidanceParams.ObstacleSeparationDistance, 0.f, 1.f));
+				const FVector::FReal PenSep = (SeparationAgentRadius + Collider.Radius + MovingAvoidanceParams.ObstacleSeparationDistance) - ConDist;
+				const FVector::FReal SeparationMag = FMath::Square(FMath::Clamp(PenSep / MovingAvoidanceParams.ObstacleSeparationDistance, 0., 1.));
 				const FVector SepForce = SeparationNormal * MovingAvoidanceParams.ObstacleSeparationStiffness;
 				const FVector SeparationForce = SepForce * SeparationMag * StandingScaling;
 
@@ -911,25 +911,25 @@ void UMassMovingAvoidanceProcessor::Execute(FMassEntityManager& EntityManager, F
 				TotalAgentSeparationForce += SeparationForce;
 
 				// Calculate closest point of approach based on relative agent positions and velocities.
-				const float CPA = UE::MassAvoidance::ComputeClosestPointOfApproach(RelPos, RelVel, PredictiveAvoidanceAgentRadius + Collider.Radius, MovingAvoidanceParams.PredictiveAvoidanceTime);
+				const FVector::FReal CPA = UE::MassAvoidance::ComputeClosestPointOfApproach(RelPos, RelVel, PredictiveAvoidanceAgentRadius + Collider.Radius, MovingAvoidanceParams.PredictiveAvoidanceTime);
 
 				// Calculate penetration at CPA
 				const FVector AvoidRelPos = RelPos + RelVel * CPA;
-				const float AvoidDist = AvoidRelPos.Size();
-				const FVector AvoidConNormal = AvoidDist > 0.0f ? (AvoidRelPos / AvoidDist) : FVector::ForwardVector;
+				const FVector::FReal AvoidDist = AvoidRelPos.Size();
+				const FVector AvoidConNormal = AvoidDist > 0. ? (AvoidRelPos / AvoidDist) : FVector::ForwardVector;
 
 				FVector AvoidNormal = AvoidConNormal;
 				if (bHasForcedNormal)
 				{
 					// The more head on the predicted collisions is, the more we should avoid towards the forced direction.
 					const FVector RelVelNorm = RelVel.GetSafeNormal();
-					const float Blend = FMath::Max(0.0f, -FVector::DotProduct(AvoidConNormal, RelVelNorm));
+					const FVector::FReal Blend = FMath::Max(0., -FVector::DotProduct(AvoidConNormal, RelVelNorm));
 					AvoidNormal = FMath::Lerp(AvoidConNormal, ForcedNormal, Blend).GetSafeNormal();
 				}
 				
-				const float AvoidPenetration = (PredictiveAvoidanceAgentRadius + Collider.Radius + MovingAvoidanceParams.PredictiveAvoidanceDistance) - AvoidDist; // Based on future agents distance
-				const float AvoidMag = FMath::Square(FMath::Clamp(AvoidPenetration / MovingAvoidanceParams.PredictiveAvoidanceDistance, 0.f, 1.f));
-				const float AvoidMagDist = (1.f - (CPA * InvPredictiveAvoidanceTime)); // No clamp, CPA is between 0 and PredictiveAvoidanceTime
+				const FVector::FReal AvoidPenetration = (PredictiveAvoidanceAgentRadius + Collider.Radius + MovingAvoidanceParams.PredictiveAvoidanceDistance) - AvoidDist; // Based on future agents distance
+				const FVector::FReal AvoidMag = FMath::Square(FMath::Clamp(AvoidPenetration / MovingAvoidanceParams.PredictiveAvoidanceDistance, 0., 1.));
+				const FVector::FReal AvoidMagDist = (1. - (CPA * InvPredictiveAvoidanceTime)); // No clamp, CPA is between 0 and PredictiveAvoidanceTime
 				const FVector AvoidForce = AvoidNormal * AvoidMag * AvoidMagDist * MovingAvoidanceParams.ObstaclePredictiveAvoidanceStiffness * StandingScaling;
 
 				SteeringForce += AvoidForce;
@@ -949,14 +949,14 @@ void UMassMovingAvoidanceProcessor::Execute(FMassEntityManager& EntityManager, F
 					Collider.Location + UE::MassAvoidance::DebugAgentHeightOffset + SeparationForce,
 					UE::MassAvoidance::AgentSeparationForceColor); 
 				
-				if (AvoidForce.Size() > 0.f)
+				if (AvoidForce.Size() > 0.)
 				{
 					// Draw agent vs agent hit positions
 					const FVector HitPosition = AgentLocation + (DesVel * CPA);
 					const FVector LeftOffset = PredictiveAvoidanceAgentRadius * UE::MassNavigation::GetLeftDirection(DesVel.GetSafeNormal(), FVector::UpVector);
 					UE::MassAvoidance::DebugDrawLine(AgentDebugContext, AgentLocation + UE::MassAvoidance::DebugAgentHeightOffset + LeftOffset, HitPosition + UE::MassAvoidance::DebugAgentHeightOffset + LeftOffset, UE::MassAvoidance::CurrentAgentColor, 1.5f);
 					UE::MassAvoidance::DebugDrawLine(AgentDebugContext, AgentLocation + UE::MassAvoidance::DebugAgentHeightOffset - LeftOffset, HitPosition + UE::MassAvoidance::DebugAgentHeightOffset - LeftOffset, UE::MassAvoidance::CurrentAgentColor, 1.5f);
-					UE::MassAvoidance::DebugDrawCylinder(AgentDebugContext, HitPosition, HitPosition + UE::MassAvoidance::DebugAgentHeightOffset, PredictiveAvoidanceAgentRadius, UE::MassAvoidance::CurrentAgentColor);
+					UE::MassAvoidance::DebugDrawCylinder(AgentDebugContext, HitPosition, HitPosition + UE::MassAvoidance::DebugAgentHeightOffset, static_cast<float>(PredictiveAvoidanceAgentRadius), UE::MassAvoidance::CurrentAgentColor);
 
 					const FVector OtherHitPosition = Collider.Location + (Collider.Velocity * CPA);
 					const FVector OtherLeftOffset = Collider.Radius * UE::MassNavigation::GetLeftDirection(Collider.Velocity.GetSafeNormal(), FVector::UpVector);
@@ -964,8 +964,8 @@ void UMassMovingAvoidanceProcessor::Execute(FMassEntityManager& EntityManager, F
 					const FVector Right = UE::MassAvoidance::DebugAgentHeightOffset - OtherLeftOffset;
 					UE::MassAvoidance::DebugDrawLine(AgentDebugContext, Collider.Location + Left, OtherHitPosition + Left, UE::MassAvoidance::AgentsColor, 1.5f);
 					UE::MassAvoidance::DebugDrawLine(AgentDebugContext, Collider.Location + Right, OtherHitPosition + Right, UE::MassAvoidance::AgentsColor, 1.5f);
-					UE::MassAvoidance::DebugDrawCylinder(AgentDebugContext, Collider.Location, Collider.Location + UE::MassAvoidance::DebugAgentHeightOffset, AgentRadius, UE::MassAvoidance::AgentsColor);
-					UE::MassAvoidance::DebugDrawCylinder(AgentDebugContext, OtherHitPosition, OtherHitPosition + UE::MassAvoidance::DebugAgentHeightOffset, AgentRadius, UE::MassAvoidance::AgentsColor);
+					UE::MassAvoidance::DebugDrawCylinder(AgentDebugContext, Collider.Location, Collider.Location + UE::MassAvoidance::DebugAgentHeightOffset, static_cast<float>(AgentRadius), UE::MassAvoidance::AgentsColor);
+					UE::MassAvoidance::DebugDrawCylinder(AgentDebugContext, OtherHitPosition, OtherHitPosition + UE::MassAvoidance::DebugAgentHeightOffset, static_cast<float>(AgentRadius), UE::MassAvoidance::AgentsColor);
 
 					// Draw agent avoid force
 					UE::MassAvoidance::DebugDrawForce(AgentDebugContext,
@@ -1060,11 +1060,11 @@ void UMassStandingAvoidanceProcessor::Execute(FMassEntityManager& EntityManager,
 		const TConstArrayView<FMassMoveTargetFragment> MoveTargetList = Context.GetFragmentView<FMassMoveTargetFragment>();
 		const FMassStandingAvoidanceParameters& StandingParams = Context.GetConstSharedFragment<FMassStandingAvoidanceParameters>();
 
-		const float GhostSeparationDistance = StandingParams.GhostSeparationDistance;
-		const float GhostSeparationStiffness = StandingParams.GhostSeparationStiffness;
+		const FVector::FReal GhostSeparationDistance = StandingParams.GhostSeparationDistance;
+		const FVector::FReal GhostSeparationStiffness = StandingParams.GhostSeparationStiffness;
 
-		const float MovingSeparationDistance = StandingParams.GhostSeparationDistance * StandingParams.MovingObstacleAvoidanceScale;
-		const float MovingSeparationStiffness = StandingParams.GhostSeparationStiffness * StandingParams.MovingObstacleAvoidanceScale;
+		const FVector::FReal MovingSeparationDistance = StandingParams.GhostSeparationDistance * StandingParams.MovingObstacleAvoidanceScale;
+		const FVector::FReal MovingSeparationStiffness = StandingParams.GhostSeparationStiffness * StandingParams.MovingObstacleAvoidanceScale;
 
 		// Arrays used to store close agents
 		TArray<FMassNavigationObstacleItem, TFixedAllocator<UE::MassAvoidance::MaxObstacleResults>> CloseEntities;
@@ -1072,12 +1072,12 @@ void UMassStandingAvoidanceProcessor::Execute(FMassEntityManager& EntityManager,
 		struct FSortedObstacle
 		{
 			FSortedObstacle() = default;
-			FSortedObstacle(const FMassEntityHandle InEntity, const FVector InLocation, const FVector InForward, const float InDistSq) : Entity(InEntity), Location(InLocation), Forward(InForward), DistSq(InDistSq) {}
+			FSortedObstacle(const FMassEntityHandle InEntity, const FVector InLocation, const FVector InForward, const FVector::FReal InDistSq) : Entity(InEntity), Location(InLocation), Forward(InForward), DistSq(InDistSq) {}
 			
 			FMassEntityHandle Entity;
 			FVector Location = FVector::ZeroVector;
 			FVector Forward = FVector::ForwardVector;
-			float DistSq = 0.0f;
+			FVector::FReal DistSq = 0.;
 		};
 		TArray<FSortedObstacle, TFixedAllocator<UE::MassAvoidance::MaxObstacleResults>> ClosestObstacles;
 
@@ -1102,21 +1102,21 @@ void UMassStandingAvoidanceProcessor::Execute(FMassEntityManager& EntityManager,
 
 			FMassEntityHandle Entity = Context.GetEntity(EntityIndex);
 			const FVector AgentLocation = Location.GetTransform().GetTranslation();
-			const float AgentRadius = Radius.Radius;
+			const FVector::FReal AgentRadius = Radius.Radius;
 
 			// Steer ghost to move target.
-			const float SteerK = 1.f / StandingParams.GhostSteeringReactionTime;
-			constexpr float SteeringMinDistance = 1.0f; // Do not bother to steer if the distance is less than this.
+			const FVector::FReal SteerK = 1. / StandingParams.GhostSteeringReactionTime;
+			constexpr FVector::FReal SteeringMinDistance = 1.; // Do not bother to steer if the distance is less than this.
 
 			FVector SteerDirection = FVector::ZeroVector;
 			FVector Delta = MoveTarget.Center - Ghost.Location;
-			Delta.Z = 0.0f;
-			const float Distance = Delta.Size();
-			float SpeedFade = 0.0;
+			Delta.Z = 0.;
+			const FVector::FReal Distance = Delta.Size();
+			FVector::FReal SpeedFade = 0.;
 			if (Distance > SteeringMinDistance)
 			{
 				SteerDirection = Delta / Distance;
-				SpeedFade = FMath::Clamp(Distance / FMath::Max(KINDA_SMALL_NUMBER, StandingParams.GhostStandSlowdownRadius), 0.0f, 1.0f);
+				SpeedFade = FMath::Clamp(Distance / FMath::Max(KINDA_SMALL_NUMBER, StandingParams.GhostStandSlowdownRadius), 0., 1.);
 			}
 
 			const FVector GhostDesiredVelocity = SteerDirection * StandingParams.GhostMaxSpeed * SpeedFade;
@@ -1128,7 +1128,7 @@ void UMassStandingAvoidanceProcessor::Execute(FMassEntityManager& EntityManager,
 			UE::MassAvoidance::FindCloseObstacles(AgentLocation, StandingParams.GhostObstacleDetectionDistance, ObstacleGrid, CloseEntities, UE::MassAvoidance::MaxObstacleResults);
 
 			// Remove unwanted and find the closest in the CloseEntities
-			const float DistanceCutOffSqr = FMath::Square(StandingParams.GhostObstacleDetectionDistance);
+			const FVector::FReal DistanceCutOffSqr = FMath::Square(StandingParams.GhostObstacleDetectionDistance);
 			ClosestObstacles.Reset();
 			for (const FNavigationObstacleHashGrid2D::ItemIDType OtherEntity : CloseEntities)
 			{
@@ -1148,7 +1148,7 @@ void UMassStandingAvoidanceProcessor::Execute(FMassEntityManager& EntityManager,
 				// Skip too far
 				const FTransformFragment& OtherTransform = EntityManager.GetFragmentDataChecked<FTransformFragment>(OtherEntity.Entity);
 				const FVector OtherLocation = OtherTransform.GetTransform().GetLocation();
-				const float DistSq = FVector::DistSquared(AgentLocation, OtherLocation);
+				const FVector::FReal DistSq = FVector::DistSquared(AgentLocation, OtherLocation);
 				if (DistSq > DistanceCutOffSqr)
 				{
 					continue;
@@ -1158,7 +1158,7 @@ void UMassStandingAvoidanceProcessor::Execute(FMassEntityManager& EntityManager,
 			}
 			ClosestObstacles.Sort([](const FSortedObstacle& A, const FSortedObstacle& B) { return A.DistSq < B.DistSq; });
 
-			const float GhostRadius = AgentRadius * StandingParams.GhostSeparationRadiusScale;
+			const FVector::FReal GhostRadius = AgentRadius * StandingParams.GhostSeparationRadiusScale;
 			
 			// Compute forces
 			constexpr int32 MaxCloseObstacleTreated = 6;
@@ -1168,8 +1168,8 @@ void UMassStandingAvoidanceProcessor::Execute(FMassEntityManager& EntityManager,
 				FSortedObstacle& OtherAgent = ClosestObstacles[Index];
 				FMassEntityView OtherEntityView(EntityManager, OtherAgent.Entity);
 
-				const float OtherRadius = OtherEntityView.GetFragmentData<FAgentRadiusFragment>().Radius;
-				const float TotalRadius = GhostRadius + OtherRadius;
+				const FVector::FReal OtherRadius = OtherEntityView.GetFragmentData<FAgentRadiusFragment>().Radius;
+				const FVector::FReal TotalRadius = GhostRadius + OtherRadius;
 
 				// @todo: this is heavy fragment to access, see if we could handle this differently.
 				const FMassMoveTargetFragment* OtherMoveTarget = OtherEntityView.GetFragmentDataPtr<FMassMoveTargetFragment>();
@@ -1183,19 +1183,19 @@ void UMassStandingAvoidanceProcessor::Execute(FMassEntityManager& EntityManager,
 				if (bOtherHasGhost)
 				{
 					// Avoid the other agent more, when it is further away from it's goal location.
-					const float OtherDistanceToGoal = FVector::Distance(OtherGhost->Location, OtherMoveTarget->Center);
-					const float OtherSteerFade = FMath::Clamp(OtherDistanceToGoal / StandingParams.GhostToTargetMaxDeviation, 0.0f, 1.0f);
-					const float SeparationStiffness = FMath::Lerp(GhostSeparationStiffness, MovingSeparationStiffness, OtherSteerFade);
+					const FVector::FReal OtherDistanceToGoal = FVector::Distance(OtherGhost->Location, OtherMoveTarget->Center);
+					const FVector::FReal OtherSteerFade = FMath::Clamp(OtherDistanceToGoal / StandingParams.GhostToTargetMaxDeviation, 0., 1.);
+					const FVector::FReal SeparationStiffness = FMath::Lerp(GhostSeparationStiffness, MovingSeparationStiffness, OtherSteerFade);
 
 					// Ghost separation
 					FVector RelPos = Ghost.Location - OtherGhost->Location;
-					RelPos.Z = 0.f; // we assume we work on a flat plane for now
-					const float ConDist = RelPos.Size();
-					const FVector ConNorm = ConDist > 0.f ? RelPos / ConDist : FVector::ForwardVector;
+					RelPos.Z = 0.; // we assume we work on a flat plane for now
+					const FVector::FReal ConDist = RelPos.Size();
+					const FVector ConNorm = ConDist > 0. ? RelPos / ConDist : FVector::ForwardVector;
 
 					// Separation force (stay away from obstacles if possible)
-					const float PenSep = (TotalRadius + GhostSeparationDistance) - ConDist;
-					const float SeparationMag = UE::MassNavigation::Smooth(FMath::Clamp(PenSep / GhostSeparationDistance, 0.f, 1.f));
+					const FVector::FReal PenSep = (TotalRadius + GhostSeparationDistance) - ConDist;
+					const FVector::FReal SeparationMag = UE::MassNavigation::Smooth(FMath::Clamp(PenSep / GhostSeparationDistance, 0., 1.));
 					const FVector SeparationForce = ConNorm * SeparationStiffness * SeparationMag;
 
 					GhostSteeringForce += SeparationForce;
@@ -1204,8 +1204,8 @@ void UMassStandingAvoidanceProcessor::Execute(FMassEntityManager& EntityManager,
 				{
 					// Avoid more when the avoidance other is in front,
 					const FVector DirToOther = (OtherAgent.Location - Ghost.Location).GetSafeNormal();
-					const float DirectionalFade = FMath::Square(FMath::Max(0.0f, FVector::DotProduct(MoveTarget.Forward, DirToOther)));
-					const float DirectionScale = FMath::Lerp(StandingParams.MovingObstacleDirectionalScale, 1.0f, DirectionalFade);
+					const FVector::FReal DirectionalFade = FMath::Square(FMath::Max(0., FVector::DotProduct(MoveTarget.Forward, DirToOther)));
+					const FVector::FReal DirectionScale = FMath::Lerp(StandingParams.MovingObstacleDirectionalScale, 1., DirectionalFade);
 
 					// Treat the other agent as a 2D capsule protruding towards forward.
  					const FVector OtherBasePosition = OtherAgent.Location;
@@ -1213,23 +1213,23 @@ void UMassStandingAvoidanceProcessor::Execute(FMassEntityManager& EntityManager,
 					const FVector OtherLocation = FMath::ClosestPointOnSegment(Ghost.Location, OtherBasePosition, OtherPersonalSpacePosition);
 
 					FVector RelPos = Ghost.Location - OtherLocation;
-					RelPos.Z = 0.f;
-					const float ConDist = RelPos.Size();
-					const FVector ConNorm = ConDist > 0.f ? RelPos / ConDist : FVector::ForwardVector;
+					RelPos.Z = 0.;
+					const FVector::FReal ConDist = RelPos.Size();
+					const FVector ConNorm = ConDist > 0. ? RelPos / ConDist : FVector::ForwardVector;
 
 					// Separation force (stay away from obstacles if possible)
-					const float PenSep = (TotalRadius + MovingSeparationDistance) - ConDist;
-					const float SeparationMag = UE::MassNavigation::Smooth(FMath::Clamp(PenSep / MovingSeparationDistance, 0.f, 1.f));
+					const FVector::FReal PenSep = (TotalRadius + MovingSeparationDistance) - ConDist;
+					const FVector::FReal SeparationMag = UE::MassNavigation::Smooth(FMath::Clamp(PenSep / MovingSeparationDistance, 0., 1.));
 					const FVector SeparationForce = ConNorm * MovingSeparationStiffness * SeparationMag;
 
 					GhostSteeringForce += SeparationForce;
 				}
 			}
 
-			GhostSteeringForce.Z = 0.0f;
+			GhostSteeringForce.Z = 0.;
 			GhostSteeringForce = UE::MassNavigation::ClampVector(GhostSteeringForce, StandingParams.GhostMaxAcceleration); // Assume unit mass
 			Ghost.Velocity += GhostSteeringForce * DeltaTime;
-			Ghost.Velocity.Z = 0.0f;
+			Ghost.Velocity.Z = 0.;
 			
 			// Damping
 			FMath::ExponentialSmoothingApprox(Ghost.Velocity, FVector::ZeroVector, DeltaTime, StandingParams.GhostVelocityDampingTime);
@@ -1238,12 +1238,11 @@ void UMassStandingAvoidanceProcessor::Execute(FMassEntityManager& EntityManager,
 
 			// Dont let the ghost location too far from move target center.
 			const FVector DirToCenter = Ghost.Location - MoveTarget.Center;
-			const float DistToCenter = DirToCenter.Length();
+			const FVector::FReal DistToCenter = DirToCenter.Length();
 			if (DistToCenter > StandingParams.GhostToTargetMaxDeviation)
 			{
 				Ghost.Location = MoveTarget.Center + DirToCenter * (StandingParams.GhostToTargetMaxDeviation / DistToCenter);
 			}
-
 		}
 	});
 	
diff --git a/Engine/Plugins/AI/MassAI/Source/MassNavigation/Private/MassNavigationFragments.cpp b/Engine/Plugins/AI/MassAI/Source/MassNavigation/Private/MassNavigationFragments.cpp
--- a/Engine/Plugins/AI/MassAI/Source/MassNavigation/Private/MassNavigationFragments.cpp
+++ b/Engine/Plugins/AI/MassAI/Source/MassNavigation/Private/MassNavigationFragments.cpp
@@ -17,7 +17,7 @@ void FMassMoveTargetFragment::CreateNewAction(const EMassMovementAction InAction
 	MarkNetDirty();
 }
 
-void FMassMoveTargetFragment::CreateReplicatedAction(const EMassMovementAction InAction, const uint16 InActionID, const float InWorldStartTime, const float InServerStartTime)
+void FMassMoveTargetFragment::CreateReplicatedAction(const EMassMovementAction InAction, const uint16 InActionID, const double InWorldStartTime, const double InServerStartTime)
 {
 	CurrentAction = InAction;
 	CurrentActionID = InActionID;
diff --git a/Engine/Plugins/AI/MassAI/Source/MassNavigation/Private/MassSmoothOrientationProcessors.cpp b/Engine/Plugins/AI/MassAI/Source/MassNavigation/Private/MassSmoothOrientationProcessors.cpp
--- a/Engine/Plugins/AI/MassAI/Source/MassNavigation/Private/MassSmoothOrientationProcessors.cpp
+++ b/Engine/Plugins/AI/MassAI/Source/MassNavigation/Private/MassSmoothOrientationProcessors.cpp
@@ -71,12 +71,12 @@ void UMassSmoothOrientationProcessor::Execute(FMassEntityManager& EntityManager,
 				const FMassVelocityFragment& CurrentVelocity = VelocityList[EntityIndex];
 				FTransform& CurrentTransform = LocationList[EntityIndex].GetMutableTransform();
 				const FVector CurrentForward = CurrentTransform.GetRotation().GetForwardVector();
-				const float CurrentHeading = UE::MassNavigation::GetYawFromDirection(CurrentForward);
+				const FVector::FReal CurrentHeading = UE::MassNavigation::GetYawFromDirection(CurrentForward);
 
 				const float EndOfPathAnticipationDistance = OrientationParams.EndOfPathDuration * MoveTarget.DesiredSpeed.Get();
 				
-				float MoveTargetWeight = 0.5f;
-				float VelocityWeight = 0.5f;
+				FVector::FReal MoveTargetWeight = 0.5;
+				FVector::FReal VelocityWeight = 0.5;
 				
 				if (MoveTarget.GetCurrentAction() == EMassMovementAction::Move)
 				{
@@ -100,13 +100,13 @@ void UMassSmoothOrientationProcessor::Execute(FMassEntityManager& EntityManager,
 					VelocityWeight = OrientationParams.Standing.VelocityWeight;
 				}
 				
-				const float VelocityHeading = UE::MassNavigation::GetYawFromDirection(CurrentVelocity.Value);
-				const float MovementHeading = UE::MassNavigation::GetYawFromDirection(MoveTarget.Forward);
+				const FVector::FReal VelocityHeading = UE::MassNavigation::GetYawFromDirection(CurrentVelocity.Value);
+				const FVector::FReal MovementHeading = UE::MassNavigation::GetYawFromDirection(MoveTarget.Forward);
 
-				const float Ratio = MoveTargetWeight / (MoveTargetWeight + VelocityWeight);
-				const float DesiredHeading = UE::MassNavigation::LerpAngle(VelocityHeading, MovementHeading,Ratio);
+				const FVector::FReal Ratio = MoveTargetWeight / (MoveTargetWeight + VelocityWeight);
+				const FVector::FReal DesiredHeading = UE::MassNavigation::LerpAngle(VelocityHeading, MovementHeading,Ratio);
 				
-				const float NewHeading = UE::MassNavigation::ExponentialSmoothingAngle(CurrentHeading, DesiredHeading, DeltaTime, OrientationParams.OrientationSmoothingTime);
+				const FVector::FReal NewHeading = UE::MassNavigation::ExponentialSmoothingAngle(CurrentHeading, DesiredHeading, DeltaTime, OrientationParams.OrientationSmoothingTime);
 
 				FQuat Rotation(FVector::UpVector, NewHeading);
 				CurrentTransform.SetRotation(Rotation);
diff --git a/Engine/Plugins/AI/MassAI/Source/MassNavigation/Private/MassSteeringProcessors.cpp b/Engine/Plugins/AI/MassAI/Source/MassNavigation/Private/MassSteeringProcessors.cpp
--- a/Engine/Plugins/AI/MassAI/Source/MassNavigation/Private/MassSteeringProcessors.cpp
+++ b/Engine/Plugins/AI/MassAI/Source/MassNavigation/Private/MassSteeringProcessors.cpp
@@ -18,29 +18,29 @@ namespace UE::MassNavigation
 	/*
 	* Calculates speed scale based on agent's forward direction and desired steering direction.
 	*/
-	static float CalcDirectionalSpeedScale(const FVector ForwardDirection, const FVector SteerDirection)
+	static FVector::FReal CalcDirectionalSpeedScale(const FVector ForwardDirection, const FVector SteerDirection)
 	{
 		// @todo: make these configurable
-		constexpr float ForwardSpeedScale = 1.0f;
-		constexpr float BackwardSpeedScale = 0.25f;
-		constexpr float SideSpeedScale = 0.5f;
+		constexpr FVector::FReal ForwardSpeedScale = 1.;
+		constexpr FVector::FReal BackwardSpeedScale = 0.25;
+		constexpr FVector::FReal SideSpeedScale = 0.5;
 
 		const FVector LeftDirection = FVector::CrossProduct(ForwardDirection, FVector::UpVector);
-		const float DirX = FVector::DotProduct(LeftDirection, SteerDirection);
-		const float DirY = FVector::DotProduct(ForwardDirection, SteerDirection);
+		const FVector::FReal DirX = FVector::DotProduct(LeftDirection, SteerDirection);
+		const FVector::FReal DirY = FVector::DotProduct(ForwardDirection, SteerDirection);
 
 		// Calculate intersection between a direction vector and ellipse, where A & B are the size of the ellipse.
 		// The direction vector is starting from the center of the ellipse.
-		constexpr float SideA = SideSpeedScale;
-		const float SideB = DirY > 0.0f ? ForwardSpeedScale : BackwardSpeedScale;
-		const float Disc = FMath::Square(SideA) * FMath::Square(DirY) + FMath::Square(SideB) * FMath::Square(DirX);
-		const float Speed = (Disc > SMALL_NUMBER) ? (SideA * SideB / FMath::Sqrt(Disc)) : 0.0f;;
+		constexpr FVector::FReal SideA = SideSpeedScale;
+		const FVector::FReal SideB = DirY > 0. ? ForwardSpeedScale : BackwardSpeedScale;
+		const FVector::FReal Disc = FMath::Square(SideA) * FMath::Square(DirY) + FMath::Square(SideB) * FMath::Square(DirX);
+		const FVector::FReal Speed = (Disc > SMALL_NUMBER) ? (SideA * SideB / FMath::Sqrt(Disc)) : 0.;
 
 		return Speed;
 	}
 
 	/** Speed envelope when approaching a point. NormalizedDistance in range [0..1] */
-	static float ArrivalSpeedEnvelope(const float NormalizedDistance)
+	static FVector::FReal ArrivalSpeedEnvelope(const FVector::FReal NormalizedDistance)
 	{
 		return FMath::Sqrt(NormalizedDistance);
 	}
@@ -91,7 +91,7 @@ void UMassSteerToMoveTargetProcessor::Execute(FMassEntityManager& EntityManager,
 		const FMassMovingSteeringParameters& MovingSteeringParams = Context.GetConstSharedFragment<FMassMovingSteeringParameters>();
 		const FMassStandingSteeringParameters& StandingSteeringParams = Context.GetConstSharedFragment<FMassStandingSteeringParameters>();
 
-		const float SteerK = 1.f / MovingSteeringParams.ReactionTime;
+		const FVector::FReal SteerK = 1. / MovingSteeringParams.ReactionTime;
 		const float DeltaTime = Context.GetDeltaTimeSeconds();
 
 		for (int32 EntityIndex = 0; EntityIndex < NumEntities; ++EntityIndex)
@@ -111,71 +111,71 @@ void UMassSteerToMoveTargetProcessor::Execute(FMassEntityManager& EntityManager,
 			const FVector CurrentLocation = Transform.GetLocation();
 			const FVector CurrentForward = Transform.GetRotation().GetForwardVector();
 
-			const float LookAheadDistance = FMath::Max(1.0f, MovingSteeringParams.LookAheadTime * MoveTarget.DesiredSpeed.Get());
+			const FVector::FReal LookAheadDistance = FMath::Max(1.0f, MovingSteeringParams.LookAheadTime * MoveTarget.DesiredSpeed.Get());
 
 			if (MoveTarget.GetCurrentAction() == EMassMovementAction::Move)
 			{
 				// Tune down avoidance and speed when arriving at goal.
-				float ArrivalFade = 1.0f;
+				FVector::FReal ArrivalFade = 1.;
 				if (MoveTarget.IntentAtGoal == EMassMovementAction::Stand)
 				{
-					ArrivalFade = FMath::Clamp(MoveTarget.DistanceToGoal / LookAheadDistance, 0.0f, 1.0f);
+					ArrivalFade = FMath::Clamp(MoveTarget.DistanceToGoal / LookAheadDistance, 0., 1.);
 				}
-				const float SteeringPredictionDistance = LookAheadDistance * ArrivalFade;
+				const FVector::FReal SteeringPredictionDistance = LookAheadDistance * ArrivalFade;
 
 				// Steer towards and along the move target.
 				const FVector TargetSide = FVector::CrossProduct(MoveTarget.Forward, FVector::UpVector);
 				const FVector Delta = CurrentLocation - MoveTarget.Center;
 
-				const float ForwardOffset = FVector::DotProduct(MoveTarget.Forward, Delta);
+				const FVector::FReal ForwardOffset = FVector::DotProduct(MoveTarget.Forward, Delta);
 
 				// Calculate steering direction. When far away from the line defined by TargetPosition and TargetTangent,
 				// the steering direction is towards the line, the close we get, the more it aligns with the line.
-				const float SidewaysOffset = FVector::DotProduct(TargetSide, Delta);
-				const float SteerForward = FMath::Sqrt(FMath::Max(0.0f, FMath::Square(SteeringPredictionDistance) - FMath::Square(SidewaysOffset)));
+				const FVector::FReal SidewaysOffset = FVector::DotProduct(TargetSide, Delta);
+				const FVector::FReal SteerForward = FMath::Sqrt(FMath::Max(0., FMath::Square(SteeringPredictionDistance) - FMath::Square(SidewaysOffset)));
 
 				// The Max() here makes the steering directions behind the TargetPosition to steer towards it directly.
-				FVector SteerTarget = MoveTarget.Center + MoveTarget.Forward * FMath::Clamp(ForwardOffset + SteerForward, 0.0f, SteeringPredictionDistance);
+				FVector SteerTarget = MoveTarget.Center + MoveTarget.Forward * FMath::Clamp(ForwardOffset + SteerForward, 0., SteeringPredictionDistance);
 
 				FVector SteerDirection = SteerTarget - CurrentLocation;
-				SteerDirection.Z = 0.0f;
-				const float DistanceToSteerTarget = SteerDirection.Length();
+				SteerDirection.Z = 0.;
+				const FVector::FReal DistanceToSteerTarget = SteerDirection.Length();
 				if (DistanceToSteerTarget > KINDA_SMALL_NUMBER)
 				{
-					SteerDirection *= 1.0f / DistanceToSteerTarget;
+					SteerDirection *= 1. / DistanceToSteerTarget;
 				}
 				
-				const float DirSpeedScale = UE::MassNavigation::CalcDirectionalSpeedScale(CurrentForward, SteerDirection);
-				float DesiredSpeed = MoveTarget.DesiredSpeed.Get() * DirSpeedScale;
+				const FVector::FReal DirSpeedScale = UE::MassNavigation::CalcDirectionalSpeedScale(CurrentForward, SteerDirection);
+				FVector::FReal DesiredSpeed = MoveTarget.DesiredSpeed.Get() * DirSpeedScale;
 
 				// Control speed based relation to the forward axis of the move target.
-				float CatchupDesiredSpeed = DesiredSpeed;
-				if (ForwardOffset < 0.0f)
+				FVector::FReal CatchupDesiredSpeed = DesiredSpeed;
+				if (ForwardOffset < 0.)
 				{
 					// Falling behind, catch up
-					const float T = FMath::Min(-ForwardOffset / LookAheadDistance, 1.0f);
+					const FVector::FReal T = FMath::Min(-ForwardOffset / LookAheadDistance, 1.);
 					CatchupDesiredSpeed = FMath::Lerp(DesiredSpeed, MovementParams.MaxSpeed, T);
 				}
-				else if (ForwardOffset > 0.0f)
+				else if (ForwardOffset > 0.)
 				{
 					// Ahead, slow down.
-					const float T = FMath::Min(ForwardOffset / LookAheadDistance, 1.0f);
-					CatchupDesiredSpeed = FMath::Lerp(DesiredSpeed, DesiredSpeed * 0.0f, 1.0f - FMath::Square(1.0f - T));
+					const FVector::FReal T = FMath::Min(ForwardOffset / LookAheadDistance, 1.);
+					CatchupDesiredSpeed = FMath::Lerp(DesiredSpeed, DesiredSpeed * 0., 1. - FMath::Square(1. - T));
 				}
 
 				// Control speed based on distance to move target. This allows to catch up even if speed above reaches zero.
-				const float DeviantSpeed = FMath::Min(FMath::Abs(SidewaysOffset) / LookAheadDistance, 1.0f) * DesiredSpeed;
+				const FVector::FReal DeviantSpeed = FMath::Min(FMath::Abs(SidewaysOffset) / LookAheadDistance, 1.) * DesiredSpeed;
 
 				DesiredSpeed = FMath::Max(CatchupDesiredSpeed, DeviantSpeed);
 
 				// Slow down towards the end of path.
 				if (MoveTarget.IntentAtGoal == EMassMovementAction::Stand)
 				{
-					const float NormalizedDistanceToSteerTarget = FMath::Clamp(DistanceToSteerTarget / LookAheadDistance, 0.0f, 1.0f);
+					const FVector::FReal NormalizedDistanceToSteerTarget = FMath::Clamp(DistanceToSteerTarget / LookAheadDistance, 0., 1.);
 					DesiredSpeed *= UE::MassNavigation::ArrivalSpeedEnvelope(FMath::Max(ArrivalFade, NormalizedDistanceToSteerTarget));
 				}
 
-				MoveTarget.bSteeringFallingBehind = ForwardOffset < -LookAheadDistance * 0.8f;
+				MoveTarget.bSteeringFallingBehind = ForwardOffset < -LookAheadDistance * 0.8;
 
 				// @todo: This current completely overrides steering, we probably should have one processor that resets the steering at the beginning of the frame.
 				Steering.DesiredVelocity = SteerDirection * DesiredSpeed;
@@ -184,8 +184,8 @@ void UMassSteerToMoveTargetProcessor::Execute(FMassEntityManager& EntityManager,
 			else if (MoveTarget.GetCurrentAction() == EMassMovementAction::Stand)
 			{
 				// Calculate unique target move threshold so that different agents react a bit differently.
-				const float PerEntityScale = UE::RandomSequence::FRand(Entity.Index);
-				const float TargetMoveThreshold = StandingSteeringParams.TargetMoveThreshold * (1.0f - StandingSteeringParams.TargetMoveThresholdVariance + PerEntityScale * StandingSteeringParams.TargetMoveThresholdVariance * 2.0f);
+				const FVector::FReal PerEntityScale = UE::RandomSequence::FRand(Entity.Index);
+				const FVector::FReal TargetMoveThreshold = StandingSteeringParams.TargetMoveThreshold * (1. - StandingSteeringParams.TargetMoveThresholdVariance + PerEntityScale * StandingSteeringParams.TargetMoveThresholdVariance * 2.);
 				
 				if (Ghost.LastSeenActionID != MoveTarget.GetCurrentActionID())
 				{
@@ -197,7 +197,7 @@ void UMassSteerToMoveTargetProcessor::Execute(FMassEntityManager& EntityManager,
 					StandingSteering.TargetLocation = MoveTarget.Center;
 					StandingSteering.TrackedTargetSpeed = 0.0f;
 					StandingSteering.bIsUpdatingTarget = false;
-					StandingSteering.TargetSelectionCooldown = StandingSteeringParams.TargetSelectionCooldown * FMath::RandRange(1.0f - StandingSteeringParams.TargetSelectionCooldownVariance, 1.0f + StandingSteeringParams.TargetSelectionCooldownVariance);
+					StandingSteering.TargetSelectionCooldown = StandingSteeringParams.TargetSelectionCooldown * FMath::RandRange(1.f - StandingSteeringParams.TargetSelectionCooldownVariance, 1.f + StandingSteeringParams.TargetSelectionCooldownVariance);
 					StandingSteering.bEnteredFromMoveAction = MoveTarget.GetPreviousAction() == EMassMovementAction::Move;
 				}
 
@@ -219,10 +219,12 @@ void UMassSteerToMoveTargetProcessor::Execute(FMassEntityManager& EntityManager,
 				{
 					// Updating target
 					StandingSteering.TargetLocation = Ghost.Location;
-					const float GhostSpeed = Ghost.Velocity.Length();
+					const FVector::FReal GhostSpeed = Ghost.Velocity.Length();
+
 					if (GhostSpeed > (StandingSteering.TrackedTargetSpeed * StandingSteeringParams.TargetSpeedHysteresisScale))
 					{
-						StandingSteering.TrackedTargetSpeed = FMath::Max(StandingSteering.TrackedTargetSpeed, GhostSpeed);
+						const FVector::FReal TrackedTargetSpeed = FMath::Max(StandingSteering.TrackedTargetSpeed, GhostSpeed);
+						StandingSteering.TrackedTargetSpeed = static_cast<float>(TrackedTargetSpeed);
 					}
 					else
 					{
@@ -234,25 +236,25 @@ void UMassSteerToMoveTargetProcessor::Execute(FMassEntityManager& EntityManager,
 				
 				// Move directly towards the move target when standing.
 				FVector SteerDirection = FVector::ZeroVector;
-				float DesiredSpeed = 0.0f;
+				FVector::FReal DesiredSpeed = 0.;
 
 				FVector Delta = StandingSteering.TargetLocation - CurrentLocation;
-				Delta.Z = 0.0f;
-				const float Distance = Delta.Size();
+				Delta.Z = 0.;
+				const FVector::FReal Distance = Delta.Size();
 				if (Distance > StandingSteeringParams.DeadZoneRadius)
 				{
 					SteerDirection = Delta / Distance;
 					if (StandingSteering.bEnteredFromMoveAction)
 					{
 						// If the current steering target is from approaching a move target, use the same speed logic as movement to ensure smooth transition.
-						const float Range = FMath::Max(1.0f, LookAheadDistance - StandingSteeringParams.DeadZoneRadius);
-						const float SpeedFade = FMath::Clamp((Distance - StandingSteeringParams.DeadZoneRadius) / Range, 0.0f, 1.0f);
+						const FVector::FReal Range = FMath::Max(1., LookAheadDistance - StandingSteeringParams.DeadZoneRadius);
+						const FVector::FReal SpeedFade = FMath::Clamp((Distance - StandingSteeringParams.DeadZoneRadius) / Range, 0., 1.);
 						DesiredSpeed = MoveTarget.DesiredSpeed.Get() * UE::MassNavigation::CalcDirectionalSpeedScale(CurrentForward, SteerDirection) * UE::MassNavigation::ArrivalSpeedEnvelope(SpeedFade);
 					}
 					else
 					{
-						const float Range = FMath::Max(1.0f, LookAheadDistance - StandingSteeringParams.DeadZoneRadius);
-						const float SpeedFade = FMath::Clamp((Distance - StandingSteeringParams.DeadZoneRadius) / Range, 0.0f, 1.0f);
+						const FVector::FReal Range = FMath::Max(1., LookAheadDistance - StandingSteeringParams.DeadZoneRadius);
+						const FVector::FReal SpeedFade = FMath::Clamp((Distance - StandingSteeringParams.DeadZoneRadius) / Range, 0., 1.);
 						// Not using the directional scaling so that the steps we take to avoid are done quickly, and the behavior is reactive.
 						DesiredSpeed = MoveTarget.DesiredSpeed.Get() * UE::MassNavigation::ArrivalSpeedEnvelope(SpeedFade);
 					}
diff --git a/Engine/Plugins/AI/MassAI/Source/MassNavigation/Public/MassNavigationFragments.h b/Engine/Plugins/AI/MassAI/Source/MassNavigation/Public/MassNavigationFragments.h
--- a/Engine/Plugins/AI/MassAI/Source/MassNavigation/Public/MassNavigationFragments.h
+++ b/Engine/Plugins/AI/MassAI/Source/MassNavigation/Public/MassNavigationFragments.h
@@ -22,7 +22,7 @@ struct MASSNAVIGATION_API FMassMoveTargetFragment : public FMassFragment
 	void CreateNewAction(const EMassMovementAction InAction, const UWorld& InWorld);
 
 	/** To setup current action from replicated data */
-	void CreateReplicatedAction(const EMassMovementAction InAction, const uint16 InActionID, const float InWorldStartTime, const float InServerStartTime);
+	void CreateReplicatedAction(const EMassMovementAction InAction, const uint16 InActionID, const double InWorldStartTime, const double InServerStartTime);
 
 	void MarkNetDirty() { bNetDirty = true; }
 	bool GetNetDirty() const { return bNetDirty; }
@@ -33,8 +33,8 @@ public:
 
 	EMassMovementAction GetPreviousAction() const { return PreviousAction; }
 	EMassMovementAction GetCurrentAction() const { return CurrentAction; }
-	float GetCurrentActionStartTime() const { return CurrentActionWorldStartTime; }
-	float GetCurrentActionServerStartTime() const { return CurrentActionServerStartTime; }
+	double GetCurrentActionStartTime() const { return CurrentActionWorldStartTime; }
+	double GetCurrentActionServerStartTime() const { return CurrentActionServerStartTime; }
 	uint16 GetCurrentActionID() const { return CurrentActionID; }
 
 	/** Center of the move target. */
@@ -51,10 +51,10 @@ public:
 
 private:
 	/** World time in seconds when the action started. */
-	float CurrentActionWorldStartTime = 0.0f;
+	double CurrentActionWorldStartTime = 0.0;
 
 	/** Server time in seconds when the action started. */
-	float CurrentActionServerStartTime = 0.0f;
+	double CurrentActionServerStartTime = 0.0;
 
 	/** Number incremented each time new action (i.e move, stand, animation) is started. */
 	uint16 CurrentActionID = 0;
diff --git a/Engine/Plugins/AI/MassAI/Source/MassNavigation/Public/MassNavigationUtils.h b/Engine/Plugins/AI/MassAI/Source/MassNavigation/Public/MassNavigationUtils.h
--- a/Engine/Plugins/AI/MassAI/Source/MassNavigation/Public/MassNavigationUtils.h
+++ b/Engine/Plugins/AI/MassAI/Source/MassNavigation/Public/MassNavigationUtils.h
@@ -1,4 +1,4 @@
-﻿// Copyright Epic Games, Inc. All Rights Reserved.
+// Copyright Epic Games, Inc. All Rights Reserved.
 
 #pragma once
 
@@ -8,51 +8,51 @@
 namespace UE::MassNavigation
 {
 	// Calculates yaw angle from direction vector.
-	inline float GetYawFromQuat(const FQuat Rotation)
+	inline FQuat::FReal GetYawFromQuat(const FQuat Rotation)
 	{
-		const float YawY = 2.f * (Rotation.W * Rotation.Z + Rotation.X * Rotation.Y);
-		const float YawX = (1.f - 2.f * (FMath::Square(Rotation.Y) + FMath::Square(Rotation.Z)));
+		const FQuat::FReal YawY = 2. * (Rotation.W * Rotation.Z + Rotation.X * Rotation.Y);
+		const FQuat::FReal YawX = (1. - 2. * (FMath::Square(Rotation.Y) + FMath::Square(Rotation.Z)));
 		return FMath::Atan2(YawY, YawX);
 	}
 
-	inline float GetYawFromDirection(const FVector Direction)
+	inline FVector::FReal GetYawFromDirection(const FVector Direction)
 	{
 		return FMath::Atan2(Direction.Y, Direction.X);
 	}
 
 	// Wraps and angle to range -PI..PI. Angle in radians.
-	inline float WrapAngle(const float Angle)
+	inline FVector::FReal WrapAngle(const FVector::FReal Angle)
 	{
-		float WrappedAngle = FMath::Fmod(Angle, PI*2.0f);
-		WrappedAngle = (WrappedAngle > PI) ? WrappedAngle - PI*2.0f : WrappedAngle;
-		WrappedAngle = (WrappedAngle < -PI) ? WrappedAngle + PI*2.0f : WrappedAngle;
+		FVector::FReal WrappedAngle = FMath::Fmod(Angle, UE_DOUBLE_PI*2.);
+		WrappedAngle = (WrappedAngle > UE_DOUBLE_PI) ? WrappedAngle - UE_DOUBLE_PI * 2. : WrappedAngle;
+		WrappedAngle = (WrappedAngle < -UE_DOUBLE_PI) ? WrappedAngle + UE_DOUBLE_PI * 2. : WrappedAngle;
 		return WrappedAngle;
 	}
 
 	// Linearly interpolates between two angles (in Radians).
-	inline float LerpAngle(const float AngleA, const float AngleB, const float T)
+	inline FVector::FReal LerpAngle(const FVector::FReal AngleA, const FVector::FReal AngleB, const FVector::FReal T)
 	{
-		const float DeltaAngle = WrapAngle(AngleB - AngleA);
+		const FVector::FReal DeltaAngle = WrapAngle(AngleB - AngleA);
 		return AngleA + DeltaAngle * T;
 	}
 
 	// Exponential smooth from current angle to target angle. Angles in radians.
-	inline float ExponentialSmoothingAngle(const float Angle, const float TargetAngle, const float DeltaTime, const float SmoothingTime)
+	inline FVector::FReal ExponentialSmoothingAngle(const FVector::FReal Angle, const FVector::FReal TargetAngle, const FVector::FReal DeltaTime, const FVector::FReal SmoothingTime)
 	{
 		// Note: based on FMath::ExponentialSmoothingApprox().
 		if (SmoothingTime < KINDA_SMALL_NUMBER)
 		{
 			return TargetAngle;
 		}
-		const float A = DeltaTime / SmoothingTime;
-		const float Exp = FMath::InvExpApprox(A);
+		const FVector::FReal A = DeltaTime / SmoothingTime;
+		const FVector::FReal Exp = FMath::InvExpApprox(A);
 		return TargetAngle + WrapAngle(Angle - TargetAngle) * Exp;
 	}
 
 	// Clamps vectors magnitude to Mag.
-	inline FVector ClampVector(const FVector Vec, const float Mag)
+	inline FVector ClampVector(const FVector Vec, const FVector::FReal Mag)
 	{
-		const float Len = Vec.SizeSquared();
+		const FVector::FReal Len = Vec.SizeSquared();
 		if (Len > FMath::Square(Mag)) {
 			return Vec * Mag / FMath::Sqrt(Len);
 		}
@@ -60,21 +60,36 @@ namespace UE::MassNavigation
 	}
 
 	// Projects a point to segment and returns the time interpolation value.
-	inline float ProjectPtSeg(const FVector2D Point, const FVector2D Start, const FVector2D End)
+	inline FVector::FReal ProjectPtSeg(const FVector2D Point, const FVector2D Start, const FVector2D End)
 	{
 		const FVector2D Seg = End - Start;
 		const FVector2D Dir = Point - Start;
-		const float d = Seg.SizeSquared();
-		const float t = FVector2D::DotProduct(Seg, Dir);
-		if (t < 0.f) return 0;
-		if (t > d) return 1;
-		return d > 0.f ? (t / d) : 0.f;
+		const FVector::FReal SegSizeSquared = Seg.SizeSquared();
+		const FVector::FReal SegDirDot = FVector2D::DotProduct(Seg, Dir);
+
+		if (SegDirDot < 0.)
+		{
+			return 0.;
+		}
+
+		if (SegDirDot > SegSizeSquared)
+		{
+			return 1.;
+		}
+
+		return SegSizeSquared > 0. ? (SegDirDot / SegSizeSquared) : 0.;
 	}
 
 	// Returns the SmoothStep curve for X in range [0..1]. 
 	inline float Smooth(const float X)
 	{
-		return X * X * (3 - 2 * X);
+		return X * X * (3.f - 2.f * X);
+	}
+
+	// Returns the SmoothStep curve for X in range [0..1]. 
+	inline double Smooth(const double X)
+	{
+		return X * X * (3. - 2. * X);
 	}
 
 	// Returns left direction from forward and up directions.
@@ -86,11 +101,11 @@ namespace UE::MassNavigation
 	// Computes miter normal in XY plane from two neighbour edge normals.
 	inline FVector ComputeMiterNormal(const FVector NormalA, const FVector NormalB)
 	{
-		FVector Mid = 0.5f * (NormalA + NormalB);
-		const float MidSquared = FVector::DotProduct(Mid, Mid);
+		FVector Mid = 0.5 * (NormalA + NormalB);
+		const FVector::FReal MidSquared = FVector::DotProduct(Mid, Mid);
 		if (MidSquared > KINDA_SMALL_NUMBER)
 		{
-			const float Scale = FMath::Min(1.f / MidSquared, 20.f);
+			const FVector::FReal Scale = FMath::Min(1. / MidSquared, 20.);
 			Mid *= Scale;
 		}
 		return Mid;
diff --git a/Engine/Plugins/AI/MassAI/Source/MassZoneGraphNavigation/Private/MassZoneGraphNavigationFragments.cpp b/Engine/Plugins/AI/MassAI/Source/MassZoneGraphNavigation/Private/MassZoneGraphNavigationFragments.cpp
--- a/Engine/Plugins/AI/MassAI/Source/MassZoneGraphNavigation/Private/MassZoneGraphNavigationFragments.cpp
+++ b/Engine/Plugins/AI/MassAI/Source/MassZoneGraphNavigation/Private/MassZoneGraphNavigationFragments.cpp
@@ -207,8 +207,8 @@ bool FMassZoneGraphShortPathFragment::RequestPath(const FMassZoneGraphCachedLane
 	// Calculate start point's relation to the start point location on lane.
 	const FVector StartDelta = StartPosition - StartLanePosition;
 	const FVector StartLeftDir = FVector::CrossProduct(StartLaneTangent, FVector::UpVector);
-	float StartLaneOffset = FVector::DotProduct(StartLeftDir, StartDelta);
-	float StartLaneForwardOffset = FVector::DotProduct(StartLaneTangent, StartDelta) * TangentSign;
+	float StartLaneOffset = FloatCastChecked<float>(FVector::DotProduct(StartLeftDir, StartDelta), UE::LWC::DefaultFloatPrecision);
+	float StartLaneForwardOffset = FloatCastChecked<float>(FVector::DotProduct(StartLaneTangent, StartDelta) * TangentSign, UE::LWC::DefaultFloatPrecision);
 	// The point is off-lane if behind the start, or beyond the boundary.
 	const bool bStartOffLane = StartLaneForwardOffset < -OffLaneCapSlop
 								|| StartLaneOffset < -(DeflatedLaneRight + OffLaneEdgeSlop)
@@ -237,8 +237,8 @@ bool FMassZoneGraphShortPathFragment::RequestPath(const FMassZoneGraphCachedLane
 		const FVector EndPosition = Request.EndOfPathPosition;
 		const FVector EndDelta = EndPosition - EndLanePosition;
 		const FVector LeftDir = FVector::CrossProduct(EndLaneTangent, FVector::UpVector);
-		EndLaneOffset = FVector::DotProduct(LeftDir, EndDelta);
-		const float EndLaneForwardOffset = FVector::DotProduct(EndLaneTangent, EndDelta) * TangentSign;
+		EndLaneOffset = FloatCastChecked<float>(FVector::DotProduct(LeftDir, EndDelta), UE::LWC::DefaultFloatPrecision);
+		const float EndLaneForwardOffset = FloatCastChecked<float>(FVector::DotProduct(EndLaneTangent, EndDelta) * TangentSign, UE::LWC::DefaultFloatPrecision);
 		// The point is off-lane if further than the, or beyond the boundary.
 		bEndOffLane = EndLaneForwardOffset > OffLaneCapSlop
 						|| EndLaneOffset < -(DeflatedLaneRight + OffLaneEdgeSlop)
@@ -417,7 +417,7 @@ bool FMassZoneGraphShortPathFragment::RequestPath(const FMassZoneGraphCachedLane
 		FMassZoneGraphPathPoint& Point = Points[PointIndex];
 		const FVector PrevPosition = PrevPoint.Position;
 		const FVector Position = Point.Position;
-		const float DeltaDistance = FVector::Dist(PrevPosition, Position);
+		const float DeltaDistance = FloatCastChecked<float>(FVector::Dist(PrevPosition, Position), UE::LWC::DefaultFloatPrecision);
 		PathDistance += DeltaDistance;
 		Point.Distance.Set(PathDistance);
 	}
diff --git a/Engine/Plugins/AI/MassAI/Source/MassZoneGraphNavigation/Private/MassZoneGraphNavigationProcessors.cpp b/Engine/Plugins/AI/MassAI/Source/MassZoneGraphNavigation/Private/MassZoneGraphNavigationProcessors.cpp
--- a/Engine/Plugins/AI/MassAI/Source/MassZoneGraphNavigation/Private/MassZoneGraphNavigationProcessors.cpp
+++ b/Engine/Plugins/AI/MassAI/Source/MassZoneGraphNavigation/Private/MassZoneGraphNavigationProcessors.cpp
@@ -471,8 +471,8 @@ void UMassZoneGraphLaneCacheBoundaryProcessor::Execute(FMassEntityManager& Entit
 			const FMassEntityHandle Entity = Context.GetEntity(EntityIndex);
 
 			// First check if we moved enough for an update
-			const float DeltaDistSquared = FVector::DistSquared(MovementTarget.Center, LaneCacheBoundary.LastUpdatePosition);
-			const float UpdateDistanceThresholdSquared = FMath::Square(50.f);
+			const FVector::FReal DeltaDistSquared = FVector::DistSquared(MovementTarget.Center, LaneCacheBoundary.LastUpdatePosition);
+			const FVector::FReal UpdateDistanceThresholdSquared = FMath::Square(50.);
 
 #if WITH_MASSGAMEPLAY_DEBUG && 0
 			const FDebugContext ObstacleDebugContext(this, LogAvoidanceObstacles, World, Entity);
diff --git a/Engine/Plugins/AI/MassCrowd/Source/MassCrowd/Private/MassCrowdLaneDataRenderingComponent.cpp b/Engine/Plugins/AI/MassCrowd/Source/MassCrowd/Private/MassCrowdLaneDataRenderingComponent.cpp
--- a/Engine/Plugins/AI/MassCrowd/Source/MassCrowd/Private/MassCrowdLaneDataRenderingComponent.cpp
+++ b/Engine/Plugins/AI/MassCrowd/Source/MassCrowd/Private/MassCrowdLaneDataRenderingComponent.cpp
@@ -259,11 +259,15 @@ void UMassCrowdLaneDataRenderingComponent::DebugDrawOnCanvas(UCanvas* Canvas, AP
 
 			if (WaitAreaData != nullptr)
 			{
-				Canvas->DrawText(Font, FString::Printf(TEXT("        %d/%d"), TrackingLaneData->NumEntitiesOnLane, WaitAreaData->GetNumSlots()), ScreenLoc.X, ScreenLoc.Y);
+				Canvas->DrawText(Font, FString::Printf(TEXT("        %d/%d"), TrackingLaneData->NumEntitiesOnLane, WaitAreaData->GetNumSlots())
+					, FloatCastChecked<float>(ScreenLoc.X, UE::LWC::DefaultFloatPrecision)
+					, FloatCastChecked<float>(ScreenLoc.Y, UE::LWC::DefaultFloatPrecision));
 			}
 			else
 			{
-				Canvas->DrawText(Font, FString::Printf(TEXT(" %d"), TrackingLaneData->NumEntitiesOnLane), ScreenLoc.X, ScreenLoc.Y);
+				Canvas->DrawText(Font, FString::Printf(TEXT(" %d"), TrackingLaneData->NumEntitiesOnLane)
+					, FloatCastChecked<float>(ScreenLoc.X, UE::LWC::DefaultFloatPrecision)
+					, FloatCastChecked<float>(ScreenLoc.Y, UE::LWC::DefaultFloatPrecision));
 			}
 		}
 	}
diff --git a/Engine/Plugins/AI/MassCrowd/Source/MassCrowd/Private/MassCrowdNavigationProcessor.cpp b/Engine/Plugins/AI/MassCrowd/Source/MassCrowd/Private/MassCrowdNavigationProcessor.cpp
--- a/Engine/Plugins/AI/MassCrowd/Source/MassCrowd/Private/MassCrowdNavigationProcessor.cpp
+++ b/Engine/Plugins/AI/MassCrowd/Source/MassCrowd/Private/MassCrowdNavigationProcessor.cpp
@@ -166,7 +166,7 @@ void UMassCrowdDynamicObstacleProcessor::Execute(FMassEntityManager& EntityManag
 			if (Obstacle.bIsMoving)
 			{
 				// Calculate current speed based on velocity or last known position.
-				const float CurrentSpeed = bHasVelocity ? VelocityList[EntityIndex].Value.Length() : (FVector::Dist(Position, Obstacle.LastPosition) / DeltaTime);
+				const FVector::FReal CurrentSpeed = bHasVelocity ? VelocityList[EntityIndex].Value.Length() : (FVector::Dist(Position, Obstacle.LastPosition) / DeltaTime);
 
 				// Update position while moving, the stop logic will use the last position while check if the obstacles moves again.
 				Obstacle.LastPosition = Position;
diff --git a/Engine/Plugins/AI/MassCrowd/Source/MassCrowd/Private/MassCrowdReplicator.cpp b/Engine/Plugins/AI/MassCrowd/Source/MassCrowd/Private/MassCrowdReplicator.cpp
--- a/Engine/Plugins/AI/MassCrowd/Source/MassCrowd/Private/MassCrowdReplicator.cpp
+++ b/Engine/Plugins/AI/MassCrowd/Source/MassCrowd/Private/MassCrowdReplicator.cpp
@@ -35,7 +35,7 @@ void UMassCrowdReplicator::ProcessClientReplication(FMassExecutionContext& Conte
 		check(RepSharedFrag);
 	};
 
-	auto AddEntityCallback = [&RepSharedFrag, &PathHandler, &PositionYawHandler](FMassExecutionContext& Context,const int32 EntityIdx, FReplicatedCrowdAgent& InReplicatedAgent, const FMassClientHandle ClientHandle)->FMassReplicatedAgentHandle
+	auto AddEntityCallback = [&RepSharedFrag, &PathHandler, &PositionYawHandler](FMassExecutionContext& Context, const int32 EntityIdx, FReplicatedCrowdAgent& InReplicatedAgent, const FMassClientHandle ClientHandle)->FMassReplicatedAgentHandle
 	{
 		AMassCrowdClientBubbleInfo& CrowdBubbleInfo = RepSharedFrag->GetTypedClientBubbleInfoChecked<AMassCrowdClientBubbleInfo>(ClientHandle);
 
@@ -45,7 +45,7 @@ void UMassCrowdReplicator::ProcessClientReplication(FMassExecutionContext& Conte
 		return CrowdBubbleInfo.GetCrowdSerializer().Bubble.AddAgent(Context.GetEntity(EntityIdx), InReplicatedAgent);
 	};
 
-	auto ModifyEntityCallback = [&RepSharedFrag, &PathHandler](FMassExecutionContext& Context, const int32 EntityIdx, const EMassLOD::Type LOD, const float Time, const FMassReplicatedAgentHandle Handle, const FMassClientHandle ClientHandle)
+	auto ModifyEntityCallback = [&RepSharedFrag, &PathHandler](FMassExecutionContext& Context, const int32 EntityIdx, const EMassLOD::Type LOD, const double Time, const FMassReplicatedAgentHandle Handle, const FMassClientHandle ClientHandle)
 	{
 		AMassCrowdClientBubbleInfo& CrowdBubbleInfo = RepSharedFrag->GetTypedClientBubbleInfoChecked<AMassCrowdClientBubbleInfo>(ClientHandle);
 
diff --git a/Engine/Plugins/AI/MassCrowd/Source/MassCrowd/Private/MassCrowdSubsystem.cpp b/Engine/Plugins/AI/MassCrowd/Source/MassCrowd/Private/MassCrowdSubsystem.cpp
--- a/Engine/Plugins/AI/MassCrowd/Source/MassCrowd/Private/MassCrowdSubsystem.cpp
+++ b/Engine/Plugins/AI/MassCrowd/Source/MassCrowd/Private/MassCrowdSubsystem.cpp
@@ -585,13 +585,13 @@ int32 UMassCrowdSubsystem::AcquireWaitingSlot(const FMassEntityHandle Entity, co
 			
 			// Find best vacant slot
 			// The most distant slot is used so that later arrivals are less likely to need passing between already standing agents. 
-			float BestDistanceSq = 0;
+			FVector::FReal BestDistanceSq = 0.;
 			for (int32 SlotIndex = 0; SlotIndex < WaitArea.Slots.Num(); SlotIndex++)
 			{
 				const FCrowdWaitSlot& Slot = WaitArea.Slots[SlotIndex];
 				if (!Slot.bOccupied)
 				{
-					const float DistanceToSlotSq = FVector::DistSquared(EntityPosition, Slot.Position);
+					const FVector::FReal DistanceToSlotSq = FVector::DistSquared(EntityPosition, Slot.Position);
 					if (DistanceToSlotSq > BestDistanceSq)
 					{
 						BestDistanceSq = DistanceToSlotSq;
diff --git a/Engine/Plugins/AI/MassCrowd/Source/MassCrowd/Private/MassCrowdVisualizationLODProcessor.cpp b/Engine/Plugins/AI/MassCrowd/Source/MassCrowd/Private/MassCrowdVisualizationLODProcessor.cpp
--- a/Engine/Plugins/AI/MassCrowd/Source/MassCrowd/Private/MassCrowdVisualizationLODProcessor.cpp
+++ b/Engine/Plugins/AI/MassCrowd/Source/MassCrowd/Private/MassCrowdVisualizationLODProcessor.cpp
@@ -79,7 +79,7 @@ void UMassCrowdVisualizationLODProcessor::Execute(FMassEntityManager& EntityMana
 			const TConstArrayView<FMassRepresentationFragment> RepresentationFragmentList = Context.GetFragmentView<FMassRepresentationFragment>();
 			const TConstArrayView<FMassViewerInfoFragment> LODInfoFragmentList = Context.GetFragmentView<FMassViewerInfoFragment>();
 			const int32 NumEntities = Context.GetNumEntities();
-			const float SpecifiedRangeSquaredCentimeters = FMath::Square(UE::MassCrowd::bDebugShowISMUnderSpecifiedRange * 100);
+			const float SpecifiedRangeSquaredCentimeters = static_cast<float>(FMath::Square(UE::MassCrowd::bDebugShowISMUnderSpecifiedRange * 100));
 			for (int EntityIdx = 0; EntityIdx < NumEntities; EntityIdx++)
 			{
 				const FMassRepresentationFragment& RepresentationFragment = RepresentationFragmentList[EntityIdx];
diff --git a/Engine/Plugins/AI/MassCrowd/Source/MassCrowd/Private/ZoneGraphCrowdLaneAnnotations.cpp b/Engine/Plugins/AI/MassCrowd/Source/MassCrowd/Private/ZoneGraphCrowdLaneAnnotations.cpp
--- a/Engine/Plugins/AI/MassCrowd/Source/MassCrowd/Private/ZoneGraphCrowdLaneAnnotations.cpp
+++ b/Engine/Plugins/AI/MassCrowd/Source/MassCrowd/Private/ZoneGraphCrowdLaneAnnotations.cpp
@@ -131,7 +131,7 @@ void UZoneGraphCrowdLaneAnnotations::DebugDraw(FZoneGraphAnnotationSceneProxy* D
 
 		for (const FZoneData& Zone : ZoneStorage->Zones)
 		{
-			const float DistanceSq = FVector::DistSquared(ViewLocation, Zone.Bounds.GetCenter());
+			const FVector::FReal DistanceSq = FVector::DistSquared(ViewLocation, Zone.Bounds.GetCenter());
 			if (DistanceSq > DrawDistanceSq)
 			{
 				continue;
@@ -182,7 +182,7 @@ void UZoneGraphCrowdLaneAnnotations::DebugDraw(FZoneGraphAnnotationSceneProxy* D
 		{
 			if (WaitArea.Slots.Num() > 0)
 			{
-				const float DistanceSq = FVector::DistSquared(ViewLocation, WaitArea.Slots[0].Position);
+				const FVector::FReal DistanceSq = FVector::DistSquared(ViewLocation, WaitArea.Slots[0].Position);
 				if (DistanceSq < DrawDistanceSq)
 				{
 					for (const FCrowdWaitSlot& Slot : WaitArea.Slots)
@@ -225,7 +225,7 @@ void UZoneGraphCrowdLaneAnnotations::DebugDrawCanvas(UCanvas* Canvas, APlayerCon
 	
 	const FVector ViewLocation = Canvas->SceneView->ViewLocation;
 	const float DrawDistance = GetMaxDebugDrawDistance() * 0.25f;
-	const float DrawDistanceSq = FMath::Square(DrawDistance);
+	const FVector::FReal DrawDistanceSq = FMath::Square(DrawDistance);
 
 	auto InFrustum = [Canvas](const FVector& Location)
 	{
@@ -242,7 +242,7 @@ void UZoneGraphCrowdLaneAnnotations::DebugDrawCanvas(UCanvas* Canvas, APlayerCon
 
 		for (const FZoneData& Zone : ZoneStorage->Zones)
 		{
-			const float DistanceSq = FVector::DistSquared(ViewLocation, Zone.Bounds.GetCenter());
+			const FVector::FReal DistanceSq = FVector::DistSquared(ViewLocation, Zone.Bounds.GetCenter());
 			if (DistanceSq > DrawDistanceSq)
 			{
 				continue;
@@ -266,7 +266,10 @@ void UZoneGraphCrowdLaneAnnotations::DebugDrawCanvas(UCanvas* Canvas, APlayerCon
 				if (bDisplayTags)
 				{
 					const FZoneGraphTagMask Mask = ZoneGraphAnnotationSubsystem->GetAnnotationTags(LaneHandle);
-					Canvas->DrawText(RenderFont, FString::Printf(TEXT("%s\n0x%08X"), *UE::ZoneGraph::Helpers::GetTagMaskString(Mask, TEXT(", ")), Mask.GetValue()), ScreenLoc.X, ScreenLoc.Y, /*XScale*/1.0f, /*YScale*/1.0f, FontInfo);
+					Canvas->DrawText(RenderFont, FString::Printf(TEXT("%s\n0x%08X"), *UE::ZoneGraph::Helpers::GetTagMaskString(Mask, TEXT(", ")), Mask.GetValue())
+						, FloatCastChecked<float>(ScreenLoc.X, UE::LWC::DefaultFloatPrecision)
+						, FloatCastChecked<float>(ScreenLoc.Y, UE::LWC::DefaultFloatPrecision)
+						, /*XScale*/1.0f, /*YScale*/1.0f, FontInfo);
 				}
 
 				// Tracking
@@ -274,7 +277,10 @@ void UZoneGraphCrowdLaneAnnotations::DebugDrawCanvas(UCanvas* Canvas, APlayerCon
 				{
 					if (TrackingData->NumEntitiesOnLane > 0)
 					{
-						Canvas->DrawText(RenderFont, FString::Printf(TEXT("Num Entities: %d"), TrackingData->NumEntitiesOnLane), ScreenLoc.X, ScreenLoc.Y + 20, /*XScale*/1.0f, /*YScale*/1.0f, FontInfo);
+						Canvas->DrawText(RenderFont, FString::Printf(TEXT("Num Entities: %d"), TrackingData->NumEntitiesOnLane)
+						, FloatCastChecked<float>(ScreenLoc.X, UE::LWC::DefaultFloatPrecision)
+						, FloatCastChecked<float>(ScreenLoc.Y + 20, UE::LWC::DefaultFloatPrecision)
+						, /*XScale*/1.0f, /*YScale*/1.0f, FontInfo);
 					}
 				}
 			}
@@ -285,7 +291,7 @@ void UZoneGraphCrowdLaneAnnotations::DebugDrawCanvas(UCanvas* Canvas, APlayerCon
 		{
 			if (WaitArea.Slots.Num() > 0)
 			{
-				const float DistanceSq = FVector::DistSquared(ViewLocation, WaitArea.Slots[0].Position);
+				const FVector::FReal DistanceSq = FVector::DistSquared(ViewLocation, WaitArea.Slots[0].Position);
 				if (DistanceSq < DrawDistanceSq)
 				{
 					for (const FCrowdWaitSlot& Slot : WaitArea.Slots)
@@ -298,7 +304,10 @@ void UZoneGraphCrowdLaneAnnotations::DebugDrawCanvas(UCanvas* Canvas, APlayerCon
 							}
 
 							const FVector ScreenLoc = Canvas->Project(Slot.Position + ZOffset);
-							Canvas->DrawText(RenderFont, TEXT("OCCUPIED"), ScreenLoc.X, ScreenLoc.Y, /*XScale*/1.0f, /*YScale*/1.0f, FontInfo);
+							Canvas->DrawText(RenderFont, TEXT("OCCUPIED")
+								, FloatCastChecked<float>(ScreenLoc.X, UE::LWC::DefaultFloatPrecision)
+								, FloatCastChecked<float>(ScreenLoc.Y, UE::LWC::DefaultFloatPrecision)
+								, /*XScale*/1.0f, /*YScale*/1.0f, FontInfo);
 						}
 					}
 				}
diff --git a/Engine/Plugins/Runtime/MassEntity/Source/MassEntityDebugger/Private/SMassBitSet.h b/Engine/Plugins/Runtime/MassEntity/Source/MassEntityDebugger/Private/SMassBitSet.h
--- a/Engine/Plugins/Runtime/MassEntity/Source/MassEntityDebugger/Private/SMassBitSet.h
+++ b/Engine/Plugins/Runtime/MassEntity/Source/MassEntityDebugger/Private/SMassBitSet.h
@@ -60,7 +60,7 @@ public:
 			.AutoWidth()
 			[
 				SNew(SBox)
-				.MinDesiredWidth(150)
+				.MinDesiredWidth(150.0f)
 				.Padding(FMargin(0.0f, InArgs._SlotPadding.Get().Top, 0.0f, InArgs._SlotPadding.Get().Bottom))
 				[
 					SNew(STextBlock)
@@ -134,7 +134,7 @@ public:
 			.AutoWidth()
 			[
 				SNew(SBox)
-				.MinDesiredWidth(150)
+				.MinDesiredWidth(150.0f)
 				.Padding(FMargin(0.0f, InArgs._SlotPadding.Get().Top, 0.0f, InArgs._SlotPadding.Get().Bottom))
 				[
 					SNew(STextBlock)
@@ -261,7 +261,7 @@ namespace UE::Mass::Debugger::UI
 			[
 				SNew(SMassBitSet<TBitSet>, Label, BitSetAccess)
 				.BackgroundBrush(Brush)
-				.SlotPadding(5)
+				.SlotPadding(5.0f)
 			];
 		}
 	}
@@ -277,7 +277,7 @@ namespace UE::Mass::Debugger::UI
 			[
 				SNew(SMassBitSetDiff<TBitSet>, Label, BaseBitSet, BitSet, Prune)
 				.BackgroundBrush(Brush)
-				.SlotPadding(5)
+				.SlotPadding(5.0f)
 			];
 		}
 	}
diff --git a/Engine/Plugins/Runtime/MassEntity/Source/MassEntityDebugger/Private/SMassQuery.cpp b/Engine/Plugins/Runtime/MassEntity/Source/MassEntityDebugger/Private/SMassQuery.cpp
--- a/Engine/Plugins/Runtime/MassEntity/Source/MassEntityDebugger/Private/SMassQuery.cpp
+++ b/Engine/Plugins/Runtime/MassEntity/Source/MassEntityDebugger/Private/SMassQuery.cpp
@@ -29,7 +29,7 @@ namespace UE::Mass::Debugger::UI::Private
 			.AutoHeight()
 			[
 				SNew(SMassBitSet<TBitSet>, Label, BitSetAccess.AsArrayView(), RequiredAccessBrushes)
-				.SlotPadding(5)
+				.SlotPadding(5.0f)
 			];
 		}
 	}
diff --git a/Engine/Plugins/Runtime/MassEntity/Source/MassEntityTestSuite/Private/MassCommandsTest.cpp b/Engine/Plugins/Runtime/MassEntity/Source/MassEntityTestSuite/Private/MassCommandsTest.cpp
--- a/Engine/Plugins/Runtime/MassEntity/Source/MassEntityTestSuite/Private/MassCommandsTest.cpp
+++ b/Engine/Plugins/Runtime/MassEntity/Source/MassEntityTestSuite/Private/MassCommandsTest.cpp
@@ -31,8 +31,8 @@ struct FCommands_FragmentInstanceList : FEntityTestBase
 
 		for (int i = 0; i < Count; ++i)
 		{
-			EntityManager->Defer().PushCommand<FMassCommandAddFragmentInstances>(IntEntities[i], FTestFragment_Int(i), FTestFragment_Float(i));
-			EntityManager->Defer().PushCommand<FMassCommandAddFragmentInstances>(FloatEntities[i], FTestFragment_Int(i), FTestFragment_Float(i));
+			EntityManager->Defer().PushCommand<FMassCommandAddFragmentInstances>(IntEntities[i], FTestFragment_Int(i), FTestFragment_Float((float)i));
+			EntityManager->Defer().PushCommand<FMassCommandAddFragmentInstances>(FloatEntities[i], FTestFragment_Int(i), FTestFragment_Float((float)i));
 		}
 
 		EntityManager->FlushCommands();
@@ -73,7 +73,7 @@ struct FCommands_FragmentMemoryCleanup : FEntityTestBase
 		const FMassArchetypeHandle ArrayArchetype = EntityManager->CreateArchetype(MakeArrayView(ArrayFragmentTypes, 1));
 		const FMassArchetypeHandle ArrayIntArchetype = EntityManager->CreateArchetype(MakeArrayView(ArrayFragmentTypes, 2));
 		const int32 EntitiesPerChunk = EntityManager->DebugGetArchetypeEntitiesCountPerChunk(ArrayArchetype);
-		const int32 Count = EntitiesPerChunk * 2.5f;
+		const int32 Count = (int32)(EntitiesPerChunk * 2.5f);
 		
 		TArray<FMassEntityHandle> Entities;
 		EntityManager->BatchCreateEntities(ArrayArchetype, Count, Entities);
@@ -113,13 +113,13 @@ struct FCommands_BuildEntitiesWithFragments : FEntityTestBase
 	virtual bool InstantTest() override
 	{
 		const int32 EntitiesPerChunk = EntityManager->DebugGetArchetypeEntitiesCountPerChunk(FloatsIntsArchetype);
-		const int32 Count = EntitiesPerChunk * 2.5f;
+		const int32 Count = (int32)(EntitiesPerChunk * 2.5f);
 
 		TArray<FMassEntityHandle> Entities;
 		for (int i = 0; i < Count; ++i)
 		{
 			Entities.Add(EntityManager->ReserveEntity());
-			EntityManager->Defer().PushCommand<FMassCommandAddFragmentInstances>(Entities.Last(), FTestFragment_Int(i), FTestFragment_Float(i));
+			EntityManager->Defer().PushCommand<FMassCommandAddFragmentInstances>(Entities.Last(), FTestFragment_Int(i), FTestFragment_Float((float)i));
 		}
 
 		AITEST_EQUAL(TEXT("All entities created should be in ArrayArchetype"), EntityManager->DebugGetArchetypeEntitiesCount(FloatsIntsArchetype), 0);
@@ -155,7 +155,7 @@ struct FCommands_BuildEntitiesInHoles : FEntityTestBase
 		for (int i = 0; i < EntitiesPerChunk; ++i)
 		{
 			Entities.Add(EntityManager->ReserveEntity());
-			EntityManager->Defer().PushCommand<FMassCommandAddFragmentInstances>(Entities.Last(), FTestFragment_Int(i), FTestFragment_Float(i));
+			EntityManager->Defer().PushCommand<FMassCommandAddFragmentInstances>(Entities.Last(), FTestFragment_Int(i), FTestFragment_Float((float)i));
 		}
 
 		AITEST_EQUAL(TEXT("All entities created should be in ArrayArchetype"), EntityManager->DebugGetArchetypeEntitiesCount(FloatsIntsArchetype), Count / 2);
@@ -179,13 +179,13 @@ struct FCommands_BuildEntitiesWithFragmentInstances : FEntityTestBase
 	virtual bool InstantTest() override
 	{
 		const int32 EntitiesPerChunk = EntityManager->DebugGetArchetypeEntitiesCountPerChunk(FloatsIntsArchetype);
-		const int32 Count = EntitiesPerChunk * 2.5f;
+		const int32 Count = (int32)(EntitiesPerChunk * 2.5f);
 
 		TArray<FMassEntityHandle> Entities;
 		for (int i = 0; i < Count; ++i)
 		{
 			Entities.Add(EntityManager->ReserveEntity());
-			EntityManager->Defer().PushCommand<FMassCommandBuildEntity>(Entities.Last(), FTestFragment_Int(i), FTestFragment_Float(i));
+			EntityManager->Defer().PushCommand<FMassCommandBuildEntity>(Entities.Last(), FTestFragment_Int(i), FTestFragment_Float((float)i));
 		}
 
 		AITEST_EQUAL(TEXT("All entities created should be in ArrayArchetype"), EntityManager->DebugGetArchetypeEntitiesCount(FloatsIntsArchetype), 0);
diff --git a/Engine/Plugins/Runtime/MassGameplay/Source/MassActors/Public/MassActorSpawnerSubsystem.h b/Engine/Plugins/Runtime/MassGameplay/Source/MassActors/Public/MassActorSpawnerSubsystem.h
--- a/Engine/Plugins/Runtime/MassGameplay/Source/MassActors/Public/MassActorSpawnerSubsystem.h
+++ b/Engine/Plugins/Runtime/MassGameplay/Source/MassActors/Public/MassActorSpawnerSubsystem.h
@@ -87,7 +87,7 @@ public:
 	uint32 SerialNumber = 0;
 
 	/** Requested world time seconds */
-	float RequestedTime = 0.0f;
+	double RequestedTime = 0.;
 
 	void Reset()
 	{
diff --git a/Engine/Plugins/Runtime/MassGameplay/Source/MassLOD/Private/MassSimulationLOD.cpp b/Engine/Plugins/Runtime/MassGameplay/Source/MassLOD/Private/MassSimulationLOD.cpp
--- a/Engine/Plugins/Runtime/MassGameplay/Source/MassLOD/Private/MassSimulationLOD.cpp
+++ b/Engine/Plugins/Runtime/MassGameplay/Source/MassLOD/Private/MassSimulationLOD.cpp
@@ -159,7 +159,7 @@ void UMassSimulationLODProcessor::Execute(FMassEntityManager& EntityManager, FMa
 	{
 		TRACE_CPUPROFILER_EVENT_SCOPE(TEXT("VariableTickRates"))
 		check(World);
-		const float Time = World->GetTimeSeconds();
+		const double Time = World->GetTimeSeconds();
 		EntityQueryVariableTick.ForEachEntityChunk(EntityManager, Context, [Time](FMassExecutionContext& Context)
 		{
 			FMassSimulationVariableTickSharedFragment& TickRateSharedFragment = Context.GetMutableSharedFragment<FMassSimulationVariableTickSharedFragment>();
diff --git a/Engine/Plugins/Runtime/MassGameplay/Source/MassLOD/Public/MassLODCollector.h b/Engine/Plugins/Runtime/MassGameplay/Source/MassLOD/Public/MassLODCollector.h
--- a/Engine/Plugins/Runtime/MassGameplay/Source/MassLOD/Public/MassLODCollector.h
+++ b/Engine/Plugins/Runtime/MassGameplay/Source/MassLOD/Public/MassLODCollector.h
@@ -118,7 +118,7 @@ void TMassLODCollector<FLODLogic>::CollectLODInfo(FMassExecutionContext& Context
 			{
 				const FVector& EntityLocation = EntityTransform.GetTransform().GetLocation();
 				const FVector ViewerToEntity = EntityLocation - Viewer.Location;
-				const float DistanceToViewerSq = ViewerToEntity.SizeSquared();
+				const float DistanceToViewerSq = static_cast<float>(ViewerToEntity.SizeSquared()); // float precision is acceptable for LOD
 				if (ClosestViewerDistanceSq > DistanceToViewerSq)
 				{
 					ClosestViewerDistanceSq = DistanceToViewerSq;
diff --git a/Engine/Plugins/Runtime/MassGameplay/Source/MassLOD/Public/MassLODTickRateController.h b/Engine/Plugins/Runtime/MassGameplay/Source/MassLOD/Public/MassLODTickRateController.h
--- a/Engine/Plugins/Runtime/MassGameplay/Source/MassLOD/Public/MassLODTickRateController.h
+++ b/Engine/Plugins/Runtime/MassGameplay/Source/MassLOD/Public/MassLODTickRateController.h
@@ -44,7 +44,7 @@ public:
 	 * @return bool return if the chunk should be tick this frame
 	 */
 	template <typename TLODFragment, typename TVariableTickRateFragment>
-	bool UpdateTickRateFromLOD(FMassExecutionContext& Context, TConstArrayView<TLODFragment> LODList, TArrayView<TVariableTickRateFragment> TickRateList, const float Time);
+	bool UpdateTickRateFromLOD(FMassExecutionContext& Context, TConstArrayView<TLODFragment> LODList, TArrayView<TVariableTickRateFragment> TickRateList, const double Time);
 
 protected:
 
@@ -87,7 +87,7 @@ bool TMassLODTickRateController<TVariableTickChunkFragment, FLODLogic>::ShouldAd
 
 template <typename TVariableTickChunkFragment, typename FLODLogic>
 template <typename TLODFragment, typename TVariableTickRateFragment>
-bool TMassLODTickRateController<TVariableTickChunkFragment, FLODLogic>::UpdateTickRateFromLOD(FMassExecutionContext& Context, TConstArrayView<TLODFragment> LODList, TArrayView<TVariableTickRateFragment> TickRateList, const float Time)
+bool TMassLODTickRateController<TVariableTickChunkFragment, FLODLogic>::UpdateTickRateFromLOD(FMassExecutionContext& Context, TConstArrayView<TLODFragment> LODList, TArrayView<TVariableTickRateFragment> TickRateList, const double Time)
 {
 	bool bShouldTickThisFrame = true;
 	bool bWasChunkTicked = true;
@@ -148,7 +148,7 @@ bool TMassLODTickRateController<TVariableTickChunkFragment, FLODLogic>::UpdateTi
 		{
 			const TLODFragment& EntityLOD = LODList[Index];
 			TVariableTickRateFragment& TickRate = TickRateList[Index];
-			TickRate.DeltaTime = TickRate.LastTickedTime != 0.0f ? Time - TickRate.LastTickedTime : DeltaTime;
+			TickRate.DeltaTime = TickRate.LastTickedTime != 0.0 ? static_cast<float>(Time - TickRate.LastTickedTime) : DeltaTime;
 			TickRate.LastTickedTime = Time;
 			if (EntityLOD.LOD != ChunkLOD)
 			{
diff --git a/Engine/Plugins/Runtime/MassGameplay/Source/MassLOD/Public/MassSimulationLOD.h b/Engine/Plugins/Runtime/MassGameplay/Source/MassLOD/Public/MassSimulationLOD.h
--- a/Engine/Plugins/Runtime/MassGameplay/Source/MassLOD/Public/MassSimulationLOD.h
+++ b/Engine/Plugins/Runtime/MassGameplay/Source/MassLOD/Public/MassSimulationLOD.h
@@ -33,7 +33,7 @@ struct MASSLOD_API FMassSimulationVariableTickFragment : public FMassFragment
 	GENERATED_BODY()
 
 	/** Accumulated delta time to use upon next tick */
-	float LastTickedTime = 0.0f;
+	double LastTickedTime = 0.;
 	float DeltaTime = 0.0f;
 };
 
diff --git a/Engine/Plugins/Runtime/MassGameplay/Source/MassReplication/Public/MassClientBubbleHandler.h b/Engine/Plugins/Runtime/MassGameplay/Source/MassReplication/Public/MassClientBubbleHandler.h
--- a/Engine/Plugins/Runtime/MassGameplay/Source/MassReplication/Public/MassClientBubbleHandler.h
+++ b/Engine/Plugins/Runtime/MassGameplay/Source/MassReplication/Public/MassClientBubbleHandler.h
@@ -80,11 +80,11 @@ struct FMassAgentLookupData
 struct FMassAgentRemoveData
 {
 	FMassAgentRemoveData() = default;
-	FMassAgentRemoveData(float InTimeLastRemoved)
+	FMassAgentRemoveData(double InTimeLastRemoved)
 		: TimeLastRemoved(InTimeLastRemoved)
 	{}
 
-	float TimeLastRemoved = 0.f;
+	double TimeLastRemoved = 0.;
 };
 
 /**
@@ -357,7 +357,7 @@ void TClientBubbleHandlerBase<AgentArrayItem>::UpdateAgentsToRemove()
 	UMassReplicationSubsystem* ReplicationSubsystem = Serializer->GetReplicationSubsystem();
 	check(ReplicationSubsystem);
 
-	const float TimeRemove = Serializer->GetWorld()->GetRealTimeSeconds() - UE::Mass::Replication::AgentRemoveInterval;
+	const double TimeRemove = Serializer->GetWorld()->GetRealTimeSeconds() - UE::Mass::Replication::AgentRemoveInterval;
 
 	// @todo do this in a more efficient way, we may potentially be able to use ACK's and FReplicatedAgentBase::bRemovedFromServeSim
 	// Check to see if we should free any EntityInfoMap entries, this is to avoid gradually growing EntityInfoMap perpetually
diff --git a/Engine/Plugins/Runtime/MassGameplay/Source/MassReplication/Public/MassReplicationProcessor.h b/Engine/Plugins/Runtime/MassGameplay/Source/MassReplication/Public/MassReplicationProcessor.h
--- a/Engine/Plugins/Runtime/MassGameplay/Source/MassReplication/Public/MassReplicationProcessor.h
+++ b/Engine/Plugins/Runtime/MassGameplay/Source/MassReplication/Public/MassReplicationProcessor.h
@@ -104,7 +104,7 @@ void UMassReplicatorBase::CalculateClientReplication(FMassExecutionContext& Cont
 
 	CacheViews(Context);
 
-	const float Time = ReplicationContext.World.GetRealTimeSeconds();
+	const double Time = ReplicationContext.World.GetRealTimeSeconds();
 
 	for (int EntityIdx = 0; EntityIdx < NumEntities; EntityIdx++)
 	{
diff --git a/Engine/Plugins/Runtime/MassGameplay/Source/MassReplication/Public/MassReplicationTypes.h b/Engine/Plugins/Runtime/MassGameplay/Source/MassReplication/Public/MassReplicationTypes.h
--- a/Engine/Plugins/Runtime/MassGameplay/Source/MassReplication/Public/MassReplicationTypes.h
+++ b/Engine/Plugins/Runtime/MassGameplay/Source/MassReplication/Public/MassReplicationTypes.h
@@ -148,12 +148,12 @@ struct FMassReplicatedAgentData
 	void Invalidate()
 	{
 		Handle.Invalidate();
-		LastUpdateTime = 0.f;
+		LastUpdateTime = 0.;
 		LOD = EMassLOD::Off;
 	}
 
 	FMassReplicatedAgentHandle Handle;
-	float LastUpdateTime = 0.f;
+	double LastUpdateTime = 0.;
 	EMassLOD::Type LOD = EMassLOD::Off;
 };
 
diff --git a/Engine/Plugins/Runtime/MassGameplay/Source/MassSmartObjects/Private/MassSmartObjectProcessor.cpp b/Engine/Plugins/Runtime/MassGameplay/Source/MassSmartObjects/Private/MassSmartObjectProcessor.cpp
--- a/Engine/Plugins/Runtime/MassGameplay/Source/MassSmartObjects/Private/MassSmartObjectProcessor.cpp
+++ b/Engine/Plugins/Runtime/MassGameplay/Source/MassSmartObjects/Private/MassSmartObjectProcessor.cpp
@@ -131,7 +131,7 @@ void UMassSmartObjectCandidatesFinderProcessor::Execute(FMassEntityManager& Enti
 			for (const FSmartObjectRequestResult& QueryResult : QueryResults)
 			{
 				const FVector SlotLocation = SmartObjectSubsystem.GetSlotLocation(QueryResult.SlotHandle).GetValue();
-				SortedCandidateSlots.Emplace(QueryResult, FVector::DistSquared(SearchOrigin, SlotLocation));
+				SortedCandidateSlots.Emplace(QueryResult, UE_REAL_TO_FLOAT_CLAMPED_MAX(FVector::DistSquared(SearchOrigin, SlotLocation)));
 
 #if WITH_MASSGAMEPLAY_DEBUG
 				if (bDisplayDebug)
@@ -144,7 +144,7 @@ void UMassSmartObjectCandidatesFinderProcessor::Execute(FMassEntityManager& Enti
 			}
 			SortedCandidateSlots.Sort([](const FSmartObjectCandidateSlot& First, const FSmartObjectCandidateSlot& Second){ return First.Cost < Second.Cost; });
 
-			Result.Candidates.NumSlots = FMath::Min<uint8>(FMassSmartObjectCandidateSlots::MaxNumCandidates, SortedCandidateSlots.Num());
+			Result.Candidates.NumSlots = IntCastChecked<uint8>(FMath::Min(FMassSmartObjectCandidateSlots::MaxNumCandidates, (uint32)SortedCandidateSlots.Num()));
 			for (int ResultIndex = 0; ResultIndex < Result.Candidates.NumSlots; ResultIndex++)
 			{
 				Result.Candidates.Slots[ResultIndex] = SortedCandidateSlots[ResultIndex];
diff --git a/Engine/Plugins/Runtime/MassGameplay/Source/MassSmartObjects/Public/MassSmartObjectFragments.h b/Engine/Plugins/Runtime/MassGameplay/Source/MassSmartObjects/Public/MassSmartObjectFragments.h
--- a/Engine/Plugins/Runtime/MassGameplay/Source/MassSmartObjects/Public/MassSmartObjectFragments.h
+++ b/Engine/Plugins/Runtime/MassGameplay/Source/MassSmartObjects/Public/MassSmartObjectFragments.h
@@ -31,7 +31,7 @@ struct MASSSMARTOBJECTS_API FMassSmartObjectUserFragment : public FMassFragment
 	 * won't look for new interactions (value of 0 indicates no cooldown).
 	 */
 	UPROPERTY(Transient)
-	float InteractionCooldownEndTime = 0.f;
+	double InteractionCooldownEndTime = 0.;
 };
 
 /** Fragment used to process time based smartobject interactions */
diff --git a/Engine/Plugins/Runtime/MassGameplay/Source/MassSpawner/Private/MassEntityZoneGraphSpawnPointsGenerator.cpp b/Engine/Plugins/Runtime/MassGameplay/Source/MassSpawner/Private/MassEntityZoneGraphSpawnPointsGenerator.cpp
--- a/Engine/Plugins/Runtime/MassGameplay/Source/MassSpawner/Private/MassEntityZoneGraphSpawnPointsGenerator.cpp
+++ b/Engine/Plugins/Runtime/MassGameplay/Source/MassSpawner/Private/MassEntityZoneGraphSpawnPointsGenerator.cpp
@@ -135,7 +135,7 @@ void UMassEntityZoneGraphSpawnPointsGenerator::GeneratePointsForZoneGraphData(co
 			float LaneLength = 0.0f;
 			UE::ZoneGraph::Query::GetLaneLength(ZoneGraphStorage, LaneIndex, LaneLength);
 
-			float Distance = RandomStream.FRandRange(MinGap, MaxGap); // ..initially
+			float Distance = static_cast<float>(RandomStream.FRandRange(MinGap, MaxGap)); // ..initially
 			while (Distance <= LaneLength)
 			{
 				// Add location at the center of this space.
@@ -145,7 +145,7 @@ void UMassEntityZoneGraphSpawnPointsGenerator::GeneratePointsForZoneGraphData(co
 				Locations.Add(LaneLocation.Position + Perp * RandomStream.FRandRange(-LaneHalfWidth, LaneHalfWidth));
 
 				// Advance ahead past the space we just consumed, plus a random gap.
-				Distance += RandomStream.FRandRange(MinGap, MaxGap);
+				Distance += static_cast<float>(RandomStream.FRandRange(MinGap, MaxGap));
 			}
 		}
 	}
diff --git a/Engine/Plugins/Runtime/ZoneGraph/Source/ZoneGraph/Private/ZoneGraphQuery.cpp b/Engine/Plugins/Runtime/ZoneGraph/Source/ZoneGraph/Private/ZoneGraphQuery.cpp
--- a/Engine/Plugins/Runtime/ZoneGraph/Source/ZoneGraph/Private/ZoneGraphQuery.cpp
+++ b/Engine/Plugins/Runtime/ZoneGraph/Source/ZoneGraph/Private/ZoneGraphQuery.cpp
@@ -14,21 +14,21 @@ float ClosestTimeOnSegment(const FVector& Point, const FVector& StartPoint, cons
 	const FVector VectToPoint = Point - StartPoint;
 
 	// See if closest point is before StartPoint
-	const float Dot1 = VectToPoint | Segment;
-	if (Dot1 <= 0.0f)
+	const FVector::FReal Dot1 = VectToPoint | Segment;
+	if (Dot1 <= 0.0)
 	{
 		return 0.0f;
 	}
 
 	// See if closest point is beyond EndPoint
-	const float Dot2 = Segment | Segment;
+	const FVector::FReal Dot2 = Segment | Segment;
 	if (Dot2 <= Dot1)
 	{
 		return 1.0f;
 	}
 
 	// Closest Point is within segment
-	return (Dot1 / Dot2);
+	return FloatCastChecked<float>(Dot1 / Dot2, /* Precision */ 1./128.);
 }
 
 static bool EnsureLaneHandle(const FZoneGraphStorage& Storage, const FZoneGraphLaneHandle LaneHandle, const char* Function)
diff --git a/Engine/Plugins/Runtime/ZoneGraphAnnotations/Source/ZoneGraphAnnotations/Private/Annotations/ZoneGraphDisturbanceAnnotation.cpp b/Engine/Plugins/Runtime/ZoneGraphAnnotations/Source/ZoneGraphAnnotations/Private/Annotations/ZoneGraphDisturbanceAnnotation.cpp
--- a/Engine/Plugins/Runtime/ZoneGraphAnnotations/Source/ZoneGraphAnnotations/Private/Annotations/ZoneGraphDisturbanceAnnotation.cpp
+++ b/Engine/Plugins/Runtime/ZoneGraphAnnotations/Source/ZoneGraphAnnotations/Private/Annotations/ZoneGraphDisturbanceAnnotation.cpp
@@ -274,14 +274,14 @@ void UZoneGraphDisturbanceAnnotation::CalculateEscapeGraph(FZoneGraphDataEscapeG
 			for (const FZoneGraphDisturbanceArea& Danger : Dangers)
 			{
 				const FVector ClosestPoint = FMath::ClosestPointOnSegment(Danger.Position, PrevPosition, LaneLocation.Position);
-				const float DangerCost = 1.0f - FMath::Min(FVector::Dist2D(Danger.Position, ClosestPoint) / Danger.Radius, 1.0f);
+				const float DangerCost = FloatCastChecked<float>(1. - FMath::Min(FVector::Dist2D(Danger.Position, ClosestPoint) / Danger.Radius, 1.), UE::LWC::DefaultFloatPrecision);
 				Span.Danger = FMath::Max(Span.Danger, DangerCost);
 			}
 
 			for (const FZoneGraphObstacleDisturbanceArea& Obstacle : Obstacles)
 			{
 				const FVector ClosestPoint = FMath::ClosestPointOnSegment(Obstacle.Position, PrevPosition, LaneLocation.Position);
-				const float DangerCost = 1.0f - FMath::Min(FVector::Dist2D(Obstacle.Position, ClosestPoint) / Obstacle.Radius, 1.0f);
+				const float DangerCost = FloatCastChecked<float>(1. - FMath::Min(FVector::Dist2D(Obstacle.Position, ClosestPoint) / Obstacle.Radius, 1.), UE::LWC::DefaultFloatPrecision);
 				Span.Danger = FMath::Max(Span.Danger, DangerCost);
 			}
 
@@ -369,7 +369,7 @@ void UZoneGraphDisturbanceAnnotation::CalculateEscapeGraph(FZoneGraphDataEscapeG
 				const uint8 NextSpanIndex = Node.ExitLinkType == EZoneLaneLinkType::Incoming ? (Node.SpanIndex + 1) : (Node.SpanIndex - 1);
 				FZoneGraphEscapeLaneSpan& NextSpan = EscapeLane.Spans[NextSpanIndex];
 				const float Danger = (Span.Danger + NextSpan.Danger) * 0.5f;
-				const float Distance = FVector::Distance(Span.Position, NextSpan.Position);
+				const float Distance = FloatCastChecked<float>(FVector::Distance(Span.Position, NextSpan.Position), UE::LWC::DefaultFloatPrecision);
 				const float NextCost = UE::ZoneGraphAnnotations::CalculateCost(Node.EscapeCost, Distance, Danger * DangerPenalty);
 				if (NextCost < NextSpan.EscapeCost)
 				{
@@ -394,7 +394,7 @@ void UZoneGraphDisturbanceAnnotation::CalculateEscapeGraph(FZoneGraphDataEscapeG
 				const uint8 NextSpanIndex = Node.SpanIndex - 1;
 				FZoneGraphEscapeLaneSpan& NextSpan = EscapeLane.Spans[NextSpanIndex];
 				const float Danger = (Span.Danger + NextSpan.Danger) * 0.5f;
-				const float Distance = FVector::Distance(Span.Position, NextSpan.Position);
+				const float Distance = FloatCastChecked<float>(FVector::Distance(Span.Position, NextSpan.Position), UE::LWC::DefaultFloatPrecision);
 				const float NextCost = UE::ZoneGraphAnnotations::CalculateCost(Node.EscapeCost, Distance, Danger * DangerPenalty);
 				if (NextCost < NextSpan.EscapeCost)
 				{
@@ -411,7 +411,7 @@ void UZoneGraphDisturbanceAnnotation::CalculateEscapeGraph(FZoneGraphDataEscapeG
 				const uint8 NextSpanIndex = Node.SpanIndex + 1;
 				FZoneGraphEscapeLaneSpan& NextSpan = EscapeLane.Spans[NextSpanIndex];
 				const float Danger = (Span.Danger + NextSpan.Danger) * 0.5f;
-				const float Distance = FVector::Distance(Span.Position, NextSpan.Position);
+				const float Distance = FloatCastChecked<float>(FVector::Distance(Span.Position, NextSpan.Position), UE::LWC::DefaultFloatPrecision);
 				const float NextCost = UE::ZoneGraphAnnotations::CalculateCost(Node.EscapeCost, Distance, Danger * DangerPenalty);
 				if (NextCost < NextSpan.EscapeCost)
 				{
@@ -461,7 +461,7 @@ void UZoneGraphDisturbanceAnnotation::CalculateEscapeGraph(FZoneGraphDataEscapeG
 						FZoneGraphEscapeLaneSpan& NextSpan = DestEscapeLane.Spans[NextSpanIndex];
 
 						const float Danger = (Span.Danger + NextSpan.Danger) * 0.5f;
-						const float Distance = FVector::Distance(Span.Position, NextSpan.Position);
+						const float Distance = FloatCastChecked<float>(FVector::Distance(Span.Position, NextSpan.Position), UE::LWC::DefaultFloatPrecision);
 						const float NextCost = UE::ZoneGraphAnnotations::CalculateCost(Node.EscapeCost, Distance, Danger * DangerPenalty);
 						if (NextCost < NextSpan.EscapeCost)
 						{
@@ -496,7 +496,7 @@ void UZoneGraphDisturbanceAnnotation::CalculateEscapeGraph(FZoneGraphDataEscapeG
 						FZoneGraphEscapeLaneSpan& NextSpan = DestEscapeLane.Spans[NextSpanIndex];
 
 						const float Danger = (Span.Danger + NextSpan.Danger) * 0.5f;
-						const float Distance = FVector::Distance(Span.Position, NextSpan.Position);
+						const float Distance = FloatCastChecked<float>(FVector::Distance(Span.Position, NextSpan.Position), UE::LWC::DefaultFloatPrecision);
 						const float NextCost = UE::ZoneGraphAnnotations::CalculateCost(Node.EscapeCost, Distance, Danger * DangerPenalty + AdjacentPenalty);
 						if (NextCost < NextSpan.EscapeCost)
 						{
@@ -698,7 +698,7 @@ void UZoneGraphDisturbanceAnnotation::DebugDraw(FZoneGraphAnnotationSceneProxy*
 				{
 					continue;
 				}
-				const float DistanceSq = FVector::DistSquared(ViewLocation, EscapeLane.Spans[EscapeLane.SpanCount/2].Position);
+				const float DistanceSq = FloatCastChecked<float>(FVector::DistSquared(ViewLocation, EscapeLane.Spans[EscapeLane.SpanCount/2].Position), UE::LWC::DefaultFloatPrecision);
 				if (DistanceSq > DrawDistanceSq)
 				{
 					continue;
@@ -787,7 +787,7 @@ void UZoneGraphDisturbanceAnnotation::DebugDraw(FZoneGraphAnnotationSceneProxy*
 
 	for (const FZoneGraphDisturbanceArea& Danger : Dangers)
 	{
-		const float DistanceSq = FVector::DistSquared(ViewLocation, Danger.Position);
+		const float DistanceSq = FloatCastChecked<float>(FVector::DistSquared(ViewLocation, Danger.Position), UE::LWC::DefaultFloatPrecision);
 		if (DistanceSq > DrawDistanceSq)
 		{
 			continue;
@@ -802,7 +802,7 @@ void UZoneGraphDisturbanceAnnotation::DebugDraw(FZoneGraphAnnotationSceneProxy*
 
 	for (const FZoneGraphObstacleDisturbanceArea& Obstacle : Obstacles)
 	{
-		const float DistanceSq = FVector::DistSquared(ViewLocation, Obstacle.Position);
+		const float DistanceSq = FloatCastChecked<float>(FVector::DistSquared(ViewLocation, Obstacle.Position), UE::LWC::DefaultFloatPrecision);
 		if (DistanceSq > DrawDistanceSq)
 		{
 			continue;
