Replaced FloatCastChecked by static_cast for MassDebugger scores since precision was not big enough and small diff is not an issue for this tool.
#rnx

diff --git a/Engine/Plugins/AI/MassAI/Source/MassAIDebug/Private/GameplayDebuggerCategory_Mass.cpp b/Engine/Plugins/AI/MassAI/Source/MassAIDebug/Private/GameplayDebuggerCategory_Mass.cpp
--- a/Engine/Plugins/AI/MassAI/Source/MassAIDebug/Private/GameplayDebuggerCategory_Mass.cpp
+++ b/Engine/Plugins/AI/MassAI/Source/MassAIDebug/Private/GameplayDebuggerCategory_Mass.cpp
@@ -623,7 +623,7 @@ void FGameplayDebuggerCategory_Mass::CollectData(APlayerController* OwnerPC, AAc
 							BasePos += FVector(0,0,50);
 							constexpr FVector::FReal ViewWeight = 0.6f; // Higher the number the more the view angle affects the score.
 							const FVector::FReal ViewScale = 1. - (ViewDot / MinViewDirDot); // Zero at center of screen
-							NearEntityDescriptions.Emplace(FloatCastChecked<float>(DistanceToEntitySq * ((1. - ViewWeight) + ViewScale * ViewWeight), UE::LWC::DefaultFloatPrecision), BasePos, Status);
+							NearEntityDescriptions.Emplace(static_cast<float>(DistanceToEntitySq * ((1. - ViewWeight) + ViewScale * ViewWeight)), BasePos, Status);
 						}
 					}
 				}
