MassNavigation: Work around intermittent VS2022 ICE by replacing TArray::AddDefaulted_GetRef() with TArray::Add_GetRef(FCollider{})

#rnx

diff --git a/Engine/Plugins/AI/MassAI/Source/MassNavigation/MassNavigation.Build.cs b/Engine/Plugins/AI/MassAI/Source/MassNavigation/MassNavigation.Build.cs
--- a/Engine/Plugins/AI/MassAI/Source/MassNavigation/MassNavigation.Build.cs
+++ b/Engine/Plugins/AI/MassAI/Source/MassNavigation/MassNavigation.Build.cs
@@ -8,9 +8,6 @@ namespace UnrealBuildTool.Rules
 		{
 			PCHUsage = ModuleRules.PCHUsageMode.UseExplicitOrSharedPCHs;
 
-			// Work-around for intermittent VS2022 msvc crash
-			bUseUnity = false;
-
 			PublicDependencyModuleNames.AddRange(
 				new string[] {
 					"AIModule",
diff --git a/Engine/Plugins/AI/MassAI/Source/MassNavigation/Private/MassAvoidanceProcessors.cpp b/Engine/Plugins/AI/MassAI/Source/MassNavigation/Private/MassAvoidanceProcessors.cpp
--- a/Engine/Plugins/AI/MassAI/Source/MassNavigation/Private/MassAvoidanceProcessors.cpp
+++ b/Engine/Plugins/AI/MassAI/Source/MassNavigation/Private/MassAvoidanceProcessors.cpp
@@ -799,7 +799,7 @@ void UMassMovingAvoidanceProcessor::Execute(FMassEntityManager& EntityManager, F
 						{
 							const FMassCircleCollider Circle = ColliderFragment->GetCircleCollider();
 							
-							FCollider& Collider = Colliders.AddDefaulted_GetRef();
+							FCollider& Collider = Colliders.Add_GetRef(FCollider{});
 							Collider.Velocity = OtherVelocity;
 							Collider.bCanAvoid = bCanAvoid;
 							Collider.bIsMoving = bOtherIsMoving;
@@ -810,7 +810,7 @@ void UMassMovingAvoidanceProcessor::Execute(FMassEntityManager& EntityManager, F
 						{
 							const FMassPillCollider Pill = ColliderFragment->GetPillCollider(); 
 
-							FCollider& Collider = Colliders.AddDefaulted_GetRef();
+							FCollider& Collider = Colliders.Add_GetRef(FCollider{});
 							Collider.Velocity = OtherVelocity;
 							Collider.bCanAvoid = bCanAvoid;
 							Collider.bIsMoving = bOtherIsMoving;
@@ -819,7 +819,7 @@ void UMassMovingAvoidanceProcessor::Execute(FMassEntityManager& EntityManager, F
 
 							if (Colliders.Num() < MaxColliders)
 							{
-								FCollider& Collider2 = Colliders.AddDefaulted_GetRef();
+								FCollider& Collider2 = Colliders.Add_GetRef(FCollider{});
 								Collider2.Velocity = OtherVelocity;
 								Collider2.bCanAvoid = bCanAvoid;
 								Collider2.bIsMoving = bOtherIsMoving;
@@ -831,7 +831,7 @@ void UMassMovingAvoidanceProcessor::Execute(FMassEntityManager& EntityManager, F
 				}
 				else
 				{
-					FCollider& Collider = Colliders.AddDefaulted_GetRef();
+					FCollider& Collider = Colliders.Add_GetRef(FCollider{});
 					Collider.Location = Obstacle.LocationCached;
 					Collider.Velocity = OtherVelocity;
 					Collider.Radius = OtherEntityView.GetFragmentData<FAgentRadiusFragment>().Radius;
