Manual merge of Mass over from Dev-EngineMerge, along with other bits necessary to make it work (most notably AITestSuite and VisualLogger).

diff --git a/Engine/Plugins/AI/MassAI/Source/MassAIDebug/Private/GameplayDebuggerCategory_Mass.cpp b/Engine/Plugins/AI/MassAI/Source/MassAIDebug/Private/GameplayDebuggerCategory_Mass.cpp
--- a/Engine/Plugins/AI/MassAI/Source/MassAIDebug/Private/GameplayDebuggerCategory_Mass.cpp
+++ b/Engine/Plugins/AI/MassAI/Source/MassAIDebug/Private/GameplayDebuggerCategory_Mass.cpp
@@ -29,6 +29,13 @@
 #include "Engine/World.h"
 #include "MassDebugger.h"
 
+// temporary namespace to bridge the gap between UE5 revisions when integrating changes. 
+// @TODO: remove
+namespace UE::Mass::Temp
+{
+	constexpr FVector::FReal DefaultFloatPrecision = 1. / 16.; 
+}
+
 namespace UE::Mass::Debug
 {
 	FMassEntityHandle GetEntityFromActor(const AActor& Actor, const UMassAgentComponent*& OutMassAgentComponent)
@@ -750,7 +757,7 @@ void FGameplayDebuggerCategory_Mass::DrawData(APlayerController* OwnerPC, FGamep
 				Coverage += (IntersectingArea * InvArea) * Other.Alpha;
 			}
 
-			Rect.Alpha = FloatCastChecked<float>(FMath::Square(1.0 - FMath::Min(Coverage, 1.0)), UE::LWC::DefaultFloatPrecision);
+			Rect.Alpha = FloatCastChecked<float>(FMath::Square(1.0 - FMath::Min(Coverage, 1.0)), UE::Mass::Temp::DefaultFloatPrecision);
 			
 			if (Rect.Alpha > KINDA_SMALL_NUMBER)
 			{
@@ -770,11 +777,11 @@ void FGameplayDebuggerCategory_Mass::DrawData(APlayerController* OwnerPC, FGamep
 		FCanvasTileItem Background(Rect.Min - Padding, Rect.Max - Rect.Min + Padding * 2.0f, FLinearColor(0.0f, 0.0f, 0.0f, 0.35f * Rect.Alpha));
 		Background.BlendMode = SE_BLEND_TranslucentAlphaOnly;
 		CanvasContext.DrawItem(Background
-			, FloatCastChecked<float>(BackgroundPosition.X, UE::LWC::DefaultFloatPrecision)
-			, FloatCastChecked<float>(BackgroundPosition.Y, UE::LWC::DefaultFloatPrecision));
+			, FloatCastChecked<float>(BackgroundPosition.X, UE::Mass::Temp::DefaultFloatPrecision)
+			, FloatCastChecked<float>(BackgroundPosition.Y, UE::Mass::Temp::DefaultFloatPrecision));
 
-		CanvasContext.PrintAt(FloatCastChecked<float>(Rect.Min.X, UE::LWC::DefaultFloatPrecision)
-			, FloatCastChecked<float>(Rect.Min.Y, UE::LWC::DefaultFloatPrecision)
+		CanvasContext.PrintAt(FloatCastChecked<float>(Rect.Min.X, UE::Mass::Temp::DefaultFloatPrecision)
+			, FloatCastChecked<float>(Rect.Min.Y, UE::Mass::Temp::DefaultFloatPrecision)
 			, FColor::White, Rect.Alpha, Desc.Description);
 	}
 
diff --git a/Engine/Plugins/Runtime/MassEntity/Source/MassEntity/Private/MassEntitySubsystem.cpp b/Engine/Plugins/Runtime/MassEntity/Source/MassEntity/Private/MassEntitySubsystem.cpp
--- a/Engine/Plugins/Runtime/MassEntity/Source/MassEntity/Private/MassEntitySubsystem.cpp
+++ b/Engine/Plugins/Runtime/MassEntity/Source/MassEntity/Private/MassEntitySubsystem.cpp
@@ -42,7 +42,7 @@ void UMassEntitySubsystem::Deinitialize()
 
 bool UMassEntitySubsystem::DoesSupportWorldType(const EWorldType::Type WorldType) const
 {
-	return WorldType == EWorldType::EditorStorage || Super::DoesSupportWorldType(WorldType);
+	return Super::DoesSupportWorldType(WorldType);
 }
 
 #if WITH_MASSENTITY_DEBUG
diff --git a/Engine/Plugins/Runtime/MassEntity/Source/MassEntity/Private/MassEntityUtils.cpp b/Engine/Plugins/Runtime/MassEntity/Source/MassEntity/Private/MassEntityUtils.cpp
--- a/Engine/Plugins/Runtime/MassEntity/Source/MassEntity/Private/MassEntityUtils.cpp
+++ b/Engine/Plugins/Runtime/MassEntity/Source/MassEntity/Private/MassEntityUtils.cpp
@@ -12,7 +12,7 @@ namespace UE::Mass::Utils
 {
 EProcessorExecutionFlags GetProcessorExecutionFlagsForWold(const UWorld& World)
 {
-	if (World.IsStorageWorld() && World.IsEditorWorld())
+	if (World.IsEditorWorld())
 	{
 		return EProcessorExecutionFlags::Editor;
 	}
diff --git a/Engine/Plugins/Runtime/MassEntity/Source/MassEntity/Private/MassProcessingPhaseManager.cpp b/Engine/Plugins/Runtime/MassEntity/Source/MassEntity/Private/MassProcessingPhaseManager.cpp
--- a/Engine/Plugins/Runtime/MassEntity/Source/MassEntity/Private/MassProcessingPhaseManager.cpp
+++ b/Engine/Plugins/Runtime/MassEntity/Source/MassEntity/Private/MassProcessingPhaseManager.cpp
@@ -194,7 +194,7 @@ void FMassProcessingPhaseManager::Initialize(UObject& InOwner, TConstArrayView<F
 {
 #if WITH_EDITOR
 	UWorld* World = InOwner.GetWorld();
-	const bool bCreateProcessorGraphPreview = (World != nullptr) && (World->IsEditorWorld() && !World->IsGameWorld() && !World->IsStorageWorld());
+	const bool bCreateProcessorGraphPreview = (World != nullptr) && (World->IsEditorWorld() && !World->IsGameWorld());
 #endif // WITH_EDITOR
 	Owner = &InOwner;
 	ProcessingPhasesConfig = InProcessingPhasesConfig;
diff --git a/Engine/Plugins/Runtime/MassEntity/Source/MassEntityDebugger/Private/SMassDebugger.cpp b/Engine/Plugins/Runtime/MassEntity/Source/MassEntityDebugger/Private/SMassDebugger.cpp
--- a/Engine/Plugins/Runtime/MassEntity/Source/MassEntityDebugger/Private/SMassDebugger.cpp
+++ b/Engine/Plugins/Runtime/MassEntity/Source/MassEntityDebugger/Private/SMassDebugger.cpp
@@ -41,8 +41,7 @@ namespace UE::Mass::Debugger::Private
 		return 
 			WorldType == EWorldType::Game || 
 			WorldType == EWorldType::Editor || 
-			WorldType == EWorldType::PIE || 
-			WorldType == EWorldType::EditorStorage;
+			WorldType == EWorldType::PIE;
 	}
 }
 
diff --git a/Engine/Plugins/Runtime/MassGameplay/Source/MassActors/Private/MassActorSubsystem.cpp b/Engine/Plugins/Runtime/MassGameplay/Source/MassActors/Private/MassActorSubsystem.cpp
--- a/Engine/Plugins/Runtime/MassGameplay/Source/MassActors/Private/MassActorSubsystem.cpp
+++ b/Engine/Plugins/Runtime/MassGameplay/Source/MassActors/Private/MassActorSubsystem.cpp
@@ -97,7 +97,7 @@ void UMassActorSubsystem::Deinitialize()
 
 bool UMassActorSubsystem::DoesSupportWorldType(const EWorldType::Type WorldType) const
 {
-	return WorldType == EWorldType::Type::EditorStorage || Super::DoesSupportWorldType(WorldType);
+	return Super::DoesSupportWorldType(WorldType);
 }
 
 //----------------------------------------------------------------------//
diff --git a/Engine/Plugins/Runtime/MassGameplay/Source/MassGameplayDebug/Private/MassDebuggerSubsystem.cpp b/Engine/Plugins/Runtime/MassGameplay/Source/MassGameplayDebug/Private/MassDebuggerSubsystem.cpp
--- a/Engine/Plugins/Runtime/MassGameplay/Source/MassGameplayDebug/Private/MassDebuggerSubsystem.cpp
+++ b/Engine/Plugins/Runtime/MassGameplay/Source/MassGameplayDebug/Private/MassDebuggerSubsystem.cpp
@@ -34,11 +34,7 @@ void UMassDebuggerSubsystem::Deinitialize()
 
 bool UMassDebuggerSubsystem::DoesSupportWorldType(const EWorldType::Type WorldType) const
 {
-	return
-#if WITH_MASSENTITY_DEBUG
-		WorldType == EWorldType::EditorStorage ||
-#endif
-		Super::DoesSupportWorldType(WorldType);
+	return Super::DoesSupportWorldType(WorldType);
 }
 
 #if WITH_MASSENTITY_DEBUG
diff --git a/Engine/Plugins/Runtime/MassGameplay/Source/MassSimulation/Private/MassSimulationSubsystem.cpp b/Engine/Plugins/Runtime/MassGameplay/Source/MassSimulation/Private/MassSimulationSubsystem.cpp
--- a/Engine/Plugins/Runtime/MassGameplay/Source/MassSimulation/Private/MassSimulationSubsystem.cpp
+++ b/Engine/Plugins/Runtime/MassGameplay/Source/MassSimulation/Private/MassSimulationSubsystem.cpp
@@ -58,7 +58,7 @@ void UMassSimulationSubsystem::BeginDestroy()
 
 bool UMassSimulationSubsystem::DoesSupportWorldType(const EWorldType::Type WorldType) const
 {
-	return WorldType == EWorldType::EditorStorage || Super::DoesSupportWorldType(WorldType);
+	return Super::DoesSupportWorldType(WorldType);
 }
 
 FMassProcessingPhase::FOnPhaseEvent& UMassSimulationSubsystem::GetOnProcessingPhaseStarted(const EMassProcessingPhase Phase)
@@ -201,12 +201,8 @@ void UMassSimulationSubsystem::OnProcessingPhaseStarted(const float DeltaSeconds
 #if WITH_EDITOR
 void UMassSimulationSubsystem::OnPieBegin(const bool bIsSimulation)
 {
-	UWorld* World = GetWorld();
-	if (!World || (World && !World->IsStorageWorld()))
-	{
-		// called so that we're not processing phases for the editor world while PIE/SIE is running
-		StopSimulation();
-	}
+	// called so that we're not processing phases for the editor world while PIE/SIE is running
+	StopSimulation();
 }
 
 void UMassSimulationSubsystem::OnPieEnded(const bool bIsSimulation)
diff --git a/Engine/Source/Developer/LogVisualizer/Private/SVisualLogger.cpp b/Engine/Source/Developer/LogVisualizer/Private/SVisualLogger.cpp
--- a/Engine/Source/Developer/LogVisualizer/Private/SVisualLogger.cpp
+++ b/Engine/Source/Developer/LogVisualizer/Private/SVisualLogger.cpp
@@ -1119,7 +1119,8 @@ void SVisualLogger::OnLogLineSelectionChanged(TSharedPtr<struct FLogEntryItem> S
 	}
 }
 
-void SVisualLogger::OnScrubPositionChanged(double NewScrubPosition, bool bScrubbing)
+// @todo NewScrubPosition needs to be `double`, waiting for 5.2 merge
+void SVisualLogger::OnScrubPositionChanged(float NewScrubPosition, bool bScrubbing)
 {
 	const TArray<FName> &SelectedRows = FVisualLoggerDatabase::Get().GetSelectedRows();
 	const double ScrubTime = FLogVisualizer::Get().GetTimeSliderController()->GetTimeSliderArgs().ScrubPosition.Get();
@@ -1140,7 +1141,7 @@ void SVisualLogger::OnScrubPositionChanged(double NewScrubPosition, bool bScrubb
 		FVisualLogExtensionInterface* Extension = (*Iterator).Value;
 		if (Extension != NULL)
 		{
-			Extension->OnScrubPositionChanged(FVisualLoggerEditorInterface::Get(), NewScrubPosition, bScrubbing);
+			Extension->OnScrubPositionChanged(FVisualLoggerEditorInterface::Get(), double(NewScrubPosition), bScrubbing);
 		}
 	}
 }
diff --git a/Engine/Source/Runtime/Engine/Private/VisualLogger/VisualLoggerTypes.cpp b/Engine/Source/Runtime/Engine/Private/VisualLogger/VisualLoggerTypes.cpp
--- a/Engine/Source/Runtime/Engine/Private/VisualLogger/VisualLoggerTypes.cpp
+++ b/Engine/Source/Runtime/Engine/Private/VisualLogger/VisualLoggerTypes.cpp
@@ -558,7 +558,10 @@ FArchive& operator<<(FArchive& Ar, FVisualLogEntry& LogEntry)
 	const int32 VLogsOldVer = Ar.CustomVer(EVisualLoggerVersion::GUID);
 	const int32 VLogsStreamObjectVer = Ar.CustomVer(FUE5MainStreamObjectVersion::GUID);
 
-	if (VLogsStreamObjectVer >= FUE5MainStreamObjectVersion::VisualLoggerTimeStampAsDouble)
+	// @todo replace with FUE5MainStreamObjectVersion::VisualLoggerTimeStampAsDouble when it becomes available
+	constexpr int32 TEMPVisualLoggerTimeStampAsDouble = 98;
+
+	if (VLogsStreamObjectVer >= TEMPVisualLoggerTimeStampAsDouble)
 	{
 		Ar << LogEntry.TimeStamp;
 	}
