Fixed changes that unintentionally made it to Dev-EngineMerge

diff --git a/Engine/Plugins/AI/MassAI/Source/MassAIDebug/Private/GameplayDebuggerCategory_Mass.cpp b/Engine/Plugins/AI/MassAI/Source/MassAIDebug/Private/GameplayDebuggerCategory_Mass.cpp
--- a/Engine/Plugins/AI/MassAI/Source/MassAIDebug/Private/GameplayDebuggerCategory_Mass.cpp
+++ b/Engine/Plugins/AI/MassAI/Source/MassAIDebug/Private/GameplayDebuggerCategory_Mass.cpp
@@ -29,13 +29,6 @@
 #include "Engine/World.h"
 #include "MassDebugger.h"
 
-// temporary namespace to bridge the gap between UE5 revisions when integrating changes. 
-// @TODO: remove
-namespace UE::Mass::Temp
-{
-	constexpr FVector::FReal DefaultFloatPrecision = 1. / 16.; 
-}
-
 namespace UE::Mass::Debug
 {
 	FMassEntityHandle GetEntityFromActor(const AActor& Actor, const UMassAgentComponent*& OutMassAgentComponent)
@@ -757,7 +750,7 @@ void FGameplayDebuggerCategory_Mass::DrawData(APlayerController* OwnerPC, FGamep
 				Coverage += (IntersectingArea * InvArea) * Other.Alpha;
 			}
 
-			Rect.Alpha = FloatCastChecked<float>(FMath::Square(1.0 - FMath::Min(Coverage, 1.0)), UE::Mass::Temp::DefaultFloatPrecision);
+			Rect.Alpha = FloatCastChecked<float>(FMath::Square(1.0 - FMath::Min(Coverage, 1.0)), UE::LWC::DefaultFloatPrecision);
 			
 			if (Rect.Alpha > KINDA_SMALL_NUMBER)
 			{
@@ -777,11 +770,11 @@ void FGameplayDebuggerCategory_Mass::DrawData(APlayerController* OwnerPC, FGamep
 		FCanvasTileItem Background(Rect.Min - Padding, Rect.Max - Rect.Min + Padding * 2.0f, FLinearColor(0.0f, 0.0f, 0.0f, 0.35f * Rect.Alpha));
 		Background.BlendMode = SE_BLEND_TranslucentAlphaOnly;
 		CanvasContext.DrawItem(Background
-			, FloatCastChecked<float>(BackgroundPosition.X, UE::Mass::Temp::DefaultFloatPrecision)
-			, FloatCastChecked<float>(BackgroundPosition.Y, UE::Mass::Temp::DefaultFloatPrecision));
+			, FloatCastChecked<float>(BackgroundPosition.X, UE::LWC::DefaultFloatPrecision)
+			, FloatCastChecked<float>(BackgroundPosition.Y, UE::LWC::DefaultFloatPrecision));
 
-		CanvasContext.PrintAt(FloatCastChecked<float>(Rect.Min.X, UE::Mass::Temp::DefaultFloatPrecision)
-			, FloatCastChecked<float>(Rect.Min.Y, UE::Mass::Temp::DefaultFloatPrecision)
+		CanvasContext.PrintAt(FloatCastChecked<float>(Rect.Min.X, UE::LWC::DefaultFloatPrecision)
+			, FloatCastChecked<float>(Rect.Min.Y, UE::LWC::DefaultFloatPrecision)
 			, FColor::White, Rect.Alpha, Desc.Description);
 	}
 
diff --git a/Engine/Source/Developer/LogVisualizer/Private/SVisualLogger.cpp b/Engine/Source/Developer/LogVisualizer/Private/SVisualLogger.cpp
--- a/Engine/Source/Developer/LogVisualizer/Private/SVisualLogger.cpp
+++ b/Engine/Source/Developer/LogVisualizer/Private/SVisualLogger.cpp
@@ -1119,8 +1119,7 @@ void SVisualLogger::OnLogLineSelectionChanged(TSharedPtr<struct FLogEntryItem> S
 	}
 }
 
-// @todo NewScrubPosition needs to be `double`, waiting for 5.2 merge
-void SVisualLogger::OnScrubPositionChanged(float NewScrubPosition, bool bScrubbing)
+void SVisualLogger::OnScrubPositionChanged(double NewScrubPosition, bool bScrubbing)
 {
 	const TArray<FName> &SelectedRows = FVisualLoggerDatabase::Get().GetSelectedRows();
 	const double ScrubTime = FLogVisualizer::Get().GetTimeSliderController()->GetTimeSliderArgs().ScrubPosition.Get();
@@ -1141,7 +1140,7 @@ void SVisualLogger::OnScrubPositionChanged(float NewScrubPosition, bool bScrubbi
 		FVisualLogExtensionInterface* Extension = (*Iterator).Value;
 		if (Extension != NULL)
 		{
-			Extension->OnScrubPositionChanged(FVisualLoggerEditorInterface::Get(), double(NewScrubPosition), bScrubbing);
+			Extension->OnScrubPositionChanged(FVisualLoggerEditorInterface::Get(), NewScrubPosition, bScrubbing);
 		}
 	}
 }
