Modified FGameplayDebuggerCategory_Mass for easier use in editor-mode:
* PickEntity works without a PlayerController now
* added console commands that duplicate key-press commands

diff --git a/Engine/Plugins/AI/MassAI/Source/MassAIDebug/Private/GameplayDebuggerCategory_Mass.cpp b/Engine/Plugins/AI/MassAI/Source/MassAIDebug/Private/GameplayDebuggerCategory_Mass.cpp
--- a/Engine/Plugins/AI/MassAI/Source/MassAIDebug/Private/GameplayDebuggerCategory_Mass.cpp
+++ b/Engine/Plugins/AI/MassAI/Source/MassAIDebug/Private/GameplayDebuggerCategory_Mass.cpp
@@ -111,6 +111,15 @@ FGameplayDebuggerCategory_Mass::FGameplayDebuggerCategory_Mass()
 	BindKeyPress(EKeys::O.GetFName(), FGameplayDebuggerInputModifier::Shift, this, &FGameplayDebuggerCategory_Mass::OnToggleNearEntityOverview, EGameplayDebuggerInputMode::Replicated);
 	BindKeyPress(EKeys::V.GetFName(), FGameplayDebuggerInputModifier::Shift, this, &FGameplayDebuggerCategory_Mass::OnToggleNearEntityAvoidance, EGameplayDebuggerInputMode::Replicated);
 	BindKeyPress(EKeys::C.GetFName(), FGameplayDebuggerInputModifier::Shift, this, &FGameplayDebuggerCategory_Mass::OnToggleNearEntityPath, EGameplayDebuggerInputMode::Replicated);
+
+	ConsoleCommands.Emplace(TEXT("gdt.mass.ToggleArchetypes"), TEXT(""), FConsoleCommandDelegate::CreateLambda([this]() { OnToggleArchetypes(); }));
+	ConsoleCommands.Emplace(TEXT("gdt.mass.ToggleShapes"), TEXT(""), FConsoleCommandDelegate::CreateLambda([this]() { OnToggleShapes(); }));
+	ConsoleCommands.Emplace(TEXT("gdt.mass.ToggleAgentFragments"), TEXT(""), FConsoleCommandDelegate::CreateLambda([this]() { OnToggleAgentFragments(); }));
+	ConsoleCommands.Emplace(TEXT("gdt.mass.PickEntity"), TEXT(""), FConsoleCommandDelegate::CreateLambda([this]() { OnPickEntity(); }));
+	ConsoleCommands.Emplace(TEXT("gdt.mass.ToggleEntityDetails"), TEXT(""), FConsoleCommandDelegate::CreateLambda([this]() { OnToggleEntityDetails(); }));
+	ConsoleCommands.Emplace(TEXT("gdt.mass.ToggleNearEntityOverview"), TEXT(""), FConsoleCommandDelegate::CreateLambda([this]() { OnToggleNearEntityOverview(); }));
+	ConsoleCommands.Emplace(TEXT("gdt.mass.ToggleNearEntityAvoidance"), TEXT(""), FConsoleCommandDelegate::CreateLambda([this]() { OnToggleNearEntityAvoidance(); }));
+	ConsoleCommands.Emplace(TEXT("gdt.mass.ToggleNearEntityPath"), TEXT(""), FConsoleCommandDelegate::CreateLambda([this]() { OnToggleNearEntityPath(); }));
 }
 
 void FGameplayDebuggerCategory_Mass::SetCachedEntity(const FMassEntityHandle Entity, const FMassEntityManager& EntityManager)
@@ -122,12 +131,8 @@ void FGameplayDebuggerCategory_Mass::SetCachedEntity(const FMassEntityHandle Ent
 	}
 }
 
-void FGameplayDebuggerCategory_Mass::PickEntity(const APlayerController& OwnerPC, const UWorld& World, FMassEntityManager& EntityManager, const bool bLimitAngle)
+void FGameplayDebuggerCategory_Mass::PickEntity(const FVector& ViewLocation, const FVector& ViewDirection, const UWorld& World, FMassEntityManager& EntityManager, const bool bLimitAngle)
 {
-	FVector ViewLocation = FVector::ZeroVector;
-	FVector ViewDirection = FVector::ForwardVector;
-	ensureMsgf(GetViewPoint(&OwnerPC, ViewLocation, ViewDirection), TEXT("GetViewPoint is expected to always succeed when passing a valid controller."));
-
 	FMassEntityHandle BestEntity;
 	// entities indicated by UE::Mass::Debug take precedence 
     if (UE::Mass::Debug::HasDebugEntities())
@@ -204,20 +209,22 @@ void FGameplayDebuggerCategory_Mass::CollectData(APlayerController* OwnerPC, AAc
 		CachedDebugActor = nullptr;
 	}
 
-	if (OwnerPC)
+	FVector ViewLocation = FVector::ZeroVector;
+	FVector ViewDirection = FVector::ForwardVector;
+	if (GetViewPoint(OwnerPC, ViewLocation, ViewDirection))
 	{
 		// Ideally we would have a way to register in the main picking flow but that would require more changes to
 		// also support client-server picking. For now, we handle explicit mass picking requests on the authority
 		if (bPickEntity)
 		{
-			PickEntity(*OwnerPC, *World, EntityManager);
+			PickEntity(ViewLocation, ViewDirection, *World, EntityManager);
 			bPickEntity = false;
 		}
 		// if we're debugging based on UE::Mass::Debug and the range changed
 		else if (CachedDebugActor == nullptr && UE::Mass::Debug::HasDebugEntities() && UE::Mass::Debug::IsDebuggingEntity(CachedEntity) == false)
 		{
 			// using bLimitAngle = false to not limit the selection to only the things in from of the player
-			PickEntity(*OwnerPC, *World, EntityManager, /*bLimitAngle=*/false);
+			PickEntity(ViewLocation, ViewDirection, *World, EntityManager, /*bLimitAngle=*/false);
 		}
 	}
 
@@ -362,10 +369,6 @@ void FGameplayDebuggerCategory_Mass::CollectData(APlayerController* OwnerPC, AAc
 	NearEntityDescriptions.Reset();
 	if (bShowNearEntityOverview && OwnerPC)
 	{
-		FVector ViewLocation = FVector::ZeroVector;
-		FVector ViewDirection = FVector::ForwardVector;
-		ensureMsgf(GetViewPoint(OwnerPC, ViewLocation, ViewDirection), TEXT("GetViewPoint is expected to always succeed when passing a valid controller."));
-
 		FMassEntityQuery EntityQuery;
 		EntityQuery.AddRequirement<FMassStateTreeInstanceFragment>(EMassFragmentAccess::ReadOnly);
 		EntityQuery.AddConstSharedRequirement<FMassStateTreeSharedFragment>();
@@ -780,5 +783,18 @@ void FGameplayDebuggerCategory_Mass::DrawData(APlayerController* OwnerPC, FGamep
 
 	FGameplayDebuggerCategory::DrawData(OwnerPC, CanvasContext);
 }
+
+//-----------------------------------------------------------------------------
+// DEPRECATED
+//-----------------------------------------------------------------------------
+void FGameplayDebuggerCategory_Mass::PickEntity(const APlayerController& OwnerPC, const UWorld& World, FMassEntityManager& EntityManager, const bool bLimitAngle)
+{
+	FVector ViewLocation = FVector::ZeroVector;
+	FVector ViewDirection = FVector::ForwardVector;
+	ensureMsgf(GetViewPoint(&OwnerPC, ViewLocation, ViewDirection), TEXT("GetViewPoint is expected to always succeed when passing a valid controller."));
+
+	PickEntity(ViewLocation, ViewDirection, World, EntityManager, bLimitAngle);
+}
+
 #endif // WITH_GAMEPLAY_DEBUGGER && WITH_MASSGAMEPLAY_DEBUG
 
diff --git a/Engine/Plugins/AI/MassAI/Source/MassAIDebug/Public/GameplayDebuggerCategory_Mass.h b/Engine/Plugins/AI/MassAI/Source/MassAIDebug/Public/GameplayDebuggerCategory_Mass.h
--- a/Engine/Plugins/AI/MassAI/Source/MassAIDebug/Public/GameplayDebuggerCategory_Mass.h
+++ b/Engine/Plugins/AI/MassAI/Source/MassAIDebug/Public/GameplayDebuggerCategory_Mass.h
@@ -7,6 +7,7 @@
 
 #if WITH_GAMEPLAY_DEBUGGER && WITH_MASSGAMEPLAY_DEBUG
 
+#include "HAL/IConsoleManager.h"
 #include "GameplayDebuggerCategory.h"
 
 struct FMassEntityManager;
@@ -36,6 +37,9 @@ protected:
 	void OnToggleNearEntityAvoidance() { bShowNearEntityAvoidance = !bShowNearEntityAvoidance; }
 	void OnToggleNearEntityPath() { bShowNearEntityPath = !bShowNearEntityPath; }
 	
+	void PickEntity(const FVector& ViewLocation, const FVector& ViewDirection, const UWorld& World, FMassEntityManager& EntityManager, const bool bLimitAngle = true);
+
+	UE_DEPRECATED(5.3, "This flavor of PickEntity has been deprecated. Use the one getting ViewLocation and ViewDirection parameters instead.")
 	void PickEntity(const APlayerController& OwnerPC, const UWorld& World, FMassEntityManager& EntityManager, const bool bLimitAngle = true);
 
 protected:
@@ -62,6 +66,7 @@ protected:
 	};
 	TArray<FEntityDescription> NearEntityDescriptions;
 
+	TArray<FAutoConsoleCommand> ConsoleCommands;
 };
 
 #endif // WITH_GAMEPLAY_DEBUGGER && WITH_MASSGAMEPLAY_DEBUG
