Added a togglable support for debugging local/remote MassEntityManager added to GameplayDebuggerCategory_Mass and made similar possible for any category implementation.

diff --git a/Engine/Plugins/AI/MassAI/Source/MassAIDebug/Private/GameplayDebuggerCategory_Mass.cpp b/Engine/Plugins/AI/MassAI/Source/MassAIDebug/Private/GameplayDebuggerCategory_Mass.cpp
--- a/Engine/Plugins/AI/MassAI/Source/MassAIDebug/Private/GameplayDebuggerCategory_Mass.cpp
+++ b/Engine/Plugins/AI/MassAI/Source/MassAIDebug/Private/GameplayDebuggerCategory_Mass.cpp
@@ -102,6 +102,7 @@ FGameplayDebuggerCategory_Mass::FGameplayDebuggerCategory_Mass()
 	bShowNearEntityAvoidance = false;
 	bShowNearEntityPath = false;
 	bMarkEntityBeingDebugged = true;
+	bDebugLocalEntityManager = false;
 
 	BindKeyPress(EKeys::A.GetFName(), FGameplayDebuggerInputModifier::Shift, this, &FGameplayDebuggerCategory_Mass::OnToggleArchetypes, EGameplayDebuggerInputMode::Replicated);
 	BindKeyPress(EKeys::S.GetFName(), FGameplayDebuggerInputModifier::Shift, this, &FGameplayDebuggerCategory_Mass::OnToggleShapes, EGameplayDebuggerInputMode::Replicated);
@@ -111,6 +112,8 @@ FGameplayDebuggerCategory_Mass::FGameplayDebuggerCategory_Mass()
 	BindKeyPress(EKeys::O.GetFName(), FGameplayDebuggerInputModifier::Shift, this, &FGameplayDebuggerCategory_Mass::OnToggleNearEntityOverview, EGameplayDebuggerInputMode::Replicated);
 	BindKeyPress(EKeys::V.GetFName(), FGameplayDebuggerInputModifier::Shift, this, &FGameplayDebuggerCategory_Mass::OnToggleNearEntityAvoidance, EGameplayDebuggerInputMode::Replicated);
 	BindKeyPress(EKeys::C.GetFName(), FGameplayDebuggerInputModifier::Shift, this, &FGameplayDebuggerCategory_Mass::OnToggleNearEntityPath, EGameplayDebuggerInputMode::Replicated);
+	ToggleDebugLocalEntityManagerInputIndex = GetNumInputHandlers();
+	BindKeyPress(EKeys::L.GetFName(), FGameplayDebuggerInputModifier::Shift, this, &FGameplayDebuggerCategory_Mass::OnToggleDebugLocalEntityManager, EGameplayDebuggerInputMode::Local);
 
 	ConsoleCommands.Emplace(TEXT("gdt.mass.ToggleArchetypes"), TEXT(""), FConsoleCommandDelegate::CreateLambda([this]() { OnToggleArchetypes(); }));
 	ConsoleCommands.Emplace(TEXT("gdt.mass.ToggleShapes"), TEXT(""), FConsoleCommandDelegate::CreateLambda([this]() { OnToggleShapes(); }));
@@ -184,6 +187,13 @@ TSharedRef<FGameplayDebuggerCategory> FGameplayDebuggerCategory_Mass::MakeInstan
 
 void FGameplayDebuggerCategory_Mass::CollectData(APlayerController* OwnerPC, AActor* DebugActor)
 {
+	if (bAllowLocalDataCollection)
+	{
+		ResetReplicatedData();
+	}
+
+	AddTextLine(FString::Printf(TEXT("Source: {yellow}%s{white}"), bDebugLocalEntityManager ? TEXT("LOCAL") : TEXT("REMOTE")));
+
 	UWorld* World = GetDataWorld(OwnerPC, DebugActor);
 	check(World);
 
@@ -704,6 +714,7 @@ void FGameplayDebuggerCategory_Mass::DrawData(APlayerController* OwnerPC, FGamep
 	CanvasContext.Printf(TEXT("[{yellow}%s{white}] %s Entity overview"), *GetInputHandlerDescription(5), bShowNearEntityOverview ? TEXT("Hide") : TEXT("Show"));
 	CanvasContext.Printf(TEXT("[{yellow}%s{white}] %s Entity avoidance"), *GetInputHandlerDescription(6), bShowNearEntityAvoidance ? TEXT("Hide") : TEXT("Show"));
 	CanvasContext.Printf(TEXT("[{yellow}%s{white}] %s Entity path"), *GetInputHandlerDescription(7), bShowNearEntityPath ? TEXT("Hide") : TEXT("Show"));
+	CanvasContext.Printf(TEXT("[{yellow}%s{white}] Toggle Local/Remote debugging"), *GetInputHandlerDescription(8));
 
 	struct FEntityLayoutRect
 	{
@@ -784,6 +795,24 @@ void FGameplayDebuggerCategory_Mass::DrawData(APlayerController* OwnerPC, FGamep
 	FGameplayDebuggerCategory::DrawData(OwnerPC, CanvasContext);
 }
 
+void FGameplayDebuggerCategory_Mass::OnToggleDebugLocalEntityManager()
+{
+	ResetReplicatedData();
+	bDebugLocalEntityManager = !bDebugLocalEntityManager;
+	bAllowLocalDataCollection = bDebugLocalEntityManager;
+
+	const EGameplayDebuggerInputMode NewInputMode = bDebugLocalEntityManager ? EGameplayDebuggerInputMode::Local : EGameplayDebuggerInputMode::Replicated;
+	for (int32 HandlerIndex = 0; HandlerIndex < GetNumInputHandlers(); ++HandlerIndex)
+	{
+		if (HandlerIndex != ToggleDebugLocalEntityManagerInputIndex)
+		{
+			GetInputHandler(HandlerIndex).Mode = NewInputMode;
+		}
+	}
+
+	CachedEntity.Reset();
+}
+
 //-----------------------------------------------------------------------------
 // DEPRECATED
 //-----------------------------------------------------------------------------
diff --git a/Engine/Plugins/AI/MassAI/Source/MassAIDebug/Public/GameplayDebuggerCategory_Mass.h b/Engine/Plugins/AI/MassAI/Source/MassAIDebug/Public/GameplayDebuggerCategory_Mass.h
--- a/Engine/Plugins/AI/MassAI/Source/MassAIDebug/Public/GameplayDebuggerCategory_Mass.h
+++ b/Engine/Plugins/AI/MassAI/Source/MassAIDebug/Public/GameplayDebuggerCategory_Mass.h
@@ -36,6 +36,7 @@ protected:
 	void OnToggleNearEntityOverview() { bShowNearEntityOverview = !bShowNearEntityOverview; }
 	void OnToggleNearEntityAvoidance() { bShowNearEntityAvoidance = !bShowNearEntityAvoidance; }
 	void OnToggleNearEntityPath() { bShowNearEntityPath = !bShowNearEntityPath; }
+	void OnToggleDebugLocalEntityManager();
 	
 	void PickEntity(const FVector& ViewLocation, const FVector& ViewDirection, const UWorld& World, FMassEntityManager& EntityManager, const bool bLimitAngle = true);
 
@@ -54,6 +55,8 @@ protected:
 	bool bShowNearEntityAvoidance;
 	bool bShowNearEntityPath;
 	bool bMarkEntityBeingDebugged;
+	bool bDebugLocalEntityManager;
+	int32 ToggleDebugLocalEntityManagerInputIndex = INDEX_NONE;
 
 	struct FEntityDescription
 	{
diff --git a/Engine/Source/Runtime/GameplayDebugger/Private/GameplayDebuggerCategory.cpp b/Engine/Source/Runtime/GameplayDebugger/Private/GameplayDebuggerCategory.cpp
--- a/Engine/Source/Runtime/GameplayDebugger/Private/GameplayDebuggerCategory.cpp
+++ b/Engine/Source/Runtime/GameplayDebugger/Private/GameplayDebuggerCategory.cpp
@@ -13,6 +13,7 @@ FGameplayDebuggerCategory::FGameplayDebuggerCategory() :
 	bShowUpdateTimer(false),
 	bShowCategoryName(true),
 	bShowOnlyWithDebugActor(true),
+	bAllowLocalDataCollection(false),
 	bIsLocal(false),
 	bHasAuthority(true),
 	bIsEnabled(true),
@@ -49,7 +50,7 @@ void FGameplayDebuggerCategory::OnDataPackReplicated(int32 DataPackId)
 
 void FGameplayDebuggerCategory::AddTextLine(const FString& TextLine)
 {
-	if (bHasAuthority)
+	if (bHasAuthority || bAllowLocalDataCollection)
 	{
 		ReplicatedLines.Add(TextLine);
 	}
@@ -57,7 +58,7 @@ void FGameplayDebuggerCategory::AddTextLine(const FString& TextLine)
 
 void FGameplayDebuggerCategory::AddShape(const FGameplayDebuggerShape& Shape)
 {
-	if (bHasAuthority)
+	if (bHasAuthority || bAllowLocalDataCollection)
 	{
 		ReplicatedShapes.Add(Shape);
 	}
@@ -178,3 +179,10 @@ FString FGameplayDebuggerCategory::GetSceneProxyViewFlag() const
 	const bool bIsSimulate = FGameplayDebuggerAddonBase::IsSimulateInEditor();
 	return bIsSimulate ? TEXT("DebugAI") : TEXT("Game");
 }
+
+void FGameplayDebuggerCategory::ResetReplicatedData()
+{
+	ReplicatedLines.Reset();
+	ReplicatedShapes.Reset();
+	ReplicatedDataPacks.Reset();
+}
diff --git a/Engine/Source/Runtime/GameplayDebugger/Private/GameplayDebuggerLocalController.cpp b/Engine/Source/Runtime/GameplayDebugger/Private/GameplayDebuggerLocalController.cpp
--- a/Engine/Source/Runtime/GameplayDebugger/Private/GameplayDebuggerLocalController.cpp
+++ b/Engine/Source/Runtime/GameplayDebugger/Private/GameplayDebuggerLocalController.cpp
@@ -202,8 +202,9 @@ void UGameplayDebuggerLocalController::OnDebugDraw(class UCanvas* Canvas, class
 
 		CanvasContext.FontRenderInfo.bEnableShadow = bEnableTextShadow;
 
-		CanvasContext.PlayerController = CachedReplicator->GetReplicationOwner();
-		CanvasContext.World = CachedReplicator->GetReplicationOwner() ? CachedReplicator->GetReplicationOwner()->GetWorld() : CachedReplicator->GetWorld();
+		APlayerController* ReplicationOwner = CachedReplicator->GetReplicationOwner();
+		CanvasContext.PlayerController = ReplicationOwner;
+		CanvasContext.World = ReplicationOwner ? ReplicationOwner->GetWorld() : CachedReplicator->GetWorld();
 
 		DrawHeader(CanvasContext);
 
@@ -227,12 +228,20 @@ void UGameplayDebuggerLocalController::OnDebugDraw(class UCanvas* Canvas, class
 			TSharedRef<FGameplayDebuggerCategory> Category = CachedReplicator->GetCategory(Idx);
 			if (Category->ShouldDrawCategory(bHasDebugActor))
 			{
+				// this is a special-case collection mode. If we want to collect data on the client this is the 
+				// place to do it, after the data got potentially replicated over from the server, and just 
+				// before drawing, so that new replicated data won't come in as we draw.
+				if ((Category->IsCategoryAuth() == false) && Category->ShouldCollectDataOnClient())
+				{
+					Category->CollectData(ReplicationOwner, CachedReplicator->GetDebugActor());
+				}
+
 				if (Category->IsCategoryHeaderVisible())
 				{
 					DrawCategoryHeader(Idx, Category, CanvasContext);
 				}
 
-				Category->DrawCategory(CachedReplicator->GetReplicationOwner(), CanvasContext);
+				Category->DrawCategory(ReplicationOwner, CanvasContext);
 			}
 		}
 	}
diff --git a/Engine/Source/Runtime/GameplayDebugger/Public/GameplayDebuggerCategory.h b/Engine/Source/Runtime/GameplayDebugger/Public/GameplayDebuggerCategory.h
--- a/Engine/Source/Runtime/GameplayDebugger/Public/GameplayDebuggerCategory.h
+++ b/Engine/Source/Runtime/GameplayDebugger/Public/GameplayDebuggerCategory.h
@@ -94,6 +94,9 @@ public:
 	/** [ALL] check if category has authority (collects data) */
 	bool IsCategoryAuth() const { return bHasAuthority; }
 
+	/** [LOCAL] */
+	bool ShouldCollectDataOnClient() const { return bAllowLocalDataCollection; }
+
 	int32 GetNumDataPacks() const { return ReplicatedDataPacks.Num(); }
 	float GetDataPackProgress(int32 DataPackId) const { return ReplicatedDataPacks.IsValidIndex(DataPackId) ? ReplicatedDataPacks[DataPackId].GetProgress() : 0.0f; }
 	bool IsDataPackReplicating(int32 DataPackId) const { return ReplicatedDataPacks.IsValidIndex(DataPackId) && ReplicatedDataPacks[DataPackId].IsInProgress(); }
@@ -148,6 +151,9 @@ protected:
 	/** [AUTH] force data collection on next update */
 	void ForceImmediateCollect() { LastCollectDataTime = -MAX_dbl; }
 
+	/** [ALL] Clears out accumulated replicated data. */
+	void ResetReplicatedData();
+
 	/** update interval, 0 = each tick */
 	float CollectDataInterval;
 
@@ -163,6 +169,14 @@ protected:
 	/** draw category only when DebugActor is present */
 	uint32 bShowOnlyWithDebugActor : 1;
 
+	/** 
+	  * False by default. If set to True will make data collection functions work on Local categories as well.
+	  * Note that in client-server setting the data is still being replicated over from the server as well,
+	  * and data collection on the client can also be used to annotate the incoming data. If this is undesired
+	  * you need to call ResetReplicatedData before gathering the client-side data. 
+	  */
+	uint32 bAllowLocalDataCollection : 1;
+
 private:
 
 	friend class FGameplayDebuggerAddonManager;
