Changed Mass gameplay debugger category to not allow Local/Remote source toggling in standalone environment.

diff --git a/Engine/Plugins/AI/MassAI/Source/MassAIDebug/Private/GameplayDebuggerCategory_Mass.cpp b/Engine/Plugins/AI/MassAI/Source/MassAIDebug/Private/GameplayDebuggerCategory_Mass.cpp
--- a/Engine/Plugins/AI/MassAI/Source/MassAIDebug/Private/GameplayDebuggerCategory_Mass.cpp
+++ b/Engine/Plugins/AI/MassAI/Source/MassAIDebug/Private/GameplayDebuggerCategory_Mass.cpp
@@ -192,7 +192,11 @@ void FGameplayDebuggerCategory_Mass::CollectData(APlayerController* OwnerPC, AAc
 		ResetReplicatedData();
 	}
 
-	AddTextLine(FString::Printf(TEXT("Source: {yellow}%s{white}"), bDebugLocalEntityManager ? TEXT("LOCAL") : TEXT("REMOTE")));
+	// we only want to display this if there are local/remote roles in play
+	if (IsCategoryAuth() != IsCategoryLocal())
+	{
+		AddTextLine(FString::Printf(TEXT("Source: {yellow}%s{white}"), bDebugLocalEntityManager ? TEXT("LOCAL") : TEXT("REMOTE")));
+	}
 
 	UWorld* World = GetDataWorld(OwnerPC, DebugActor);
 	check(World);
@@ -716,6 +720,12 @@ void FGameplayDebuggerCategory_Mass::DrawData(APlayerController* OwnerPC, FGamep
 	CanvasContext.Printf(TEXT("[{yellow}%s{white}] %s Entity path"), *GetInputHandlerDescription(7), bShowNearEntityPath ? TEXT("Hide") : TEXT("Show"));
 	CanvasContext.Printf(TEXT("[{yellow}%s{white}] Toggle Local/Remote debugging"), *GetInputHandlerDescription(8));
 
+	if (IsCategoryLocal() && !IsCategoryAuth())
+	{
+		// we want to display this line only on clients in client-server environment.
+		CanvasContext.Printf(TEXT("[{yellow}%s{white}] Toggle Local/Remote debugging"), *GetInputHandlerDescription(ToggleDebugLocalEntityManagerInputIndex));
+	}
+
 	struct FEntityLayoutRect
 	{
 		FVector2D Min = FVector2D::ZeroVector;
@@ -797,6 +807,13 @@ void FGameplayDebuggerCategory_Mass::DrawData(APlayerController* OwnerPC, FGamep
 
 void FGameplayDebuggerCategory_Mass::OnToggleDebugLocalEntityManager()
 {
+	// this code will only execute on locally-controlled categories (as per BindKeyPress's EGameplayDebuggerInputMode::Local
+	// parameter). In such a case we don't want to toggle if we're also Auth (there's no client-server relationship here).
+	if (IsCategoryAuth())
+	{
+		return;
+	}
+
 	ResetReplicatedData();
 	bDebugLocalEntityManager = !bDebugLocalEntityManager;
 	bAllowLocalDataCollection = bDebugLocalEntityManager;
