Fixed Mass's gameplay debugger category losing selected debug actor when switching to local-source mode. ALso fixed a duplicated line in the displayed data

diff --git a/Engine/Plugins/AI/MassAI/Source/MassAIDebug/Private/GameplayDebuggerCategory_Mass.cpp b/Engine/Plugins/AI/MassAI/Source/MassAIDebug/Private/GameplayDebuggerCategory_Mass.cpp
--- a/Engine/Plugins/AI/MassAI/Source/MassAIDebug/Private/GameplayDebuggerCategory_Mass.cpp
+++ b/Engine/Plugins/AI/MassAI/Source/MassAIDebug/Private/GameplayDebuggerCategory_Mass.cpp
@@ -211,6 +211,12 @@ void FGameplayDebuggerCategory_Mass::CollectData(APlayerController* OwnerPC, AAc
 	UMassDebuggerSubsystem* Debugger = World->GetSubsystem<UMassDebuggerSubsystem>();
 
 	const UMassAgentComponent* AgentComp = nullptr;
+	
+	if (bAllowLocalDataCollection)
+	{
+		DebugActor = CachedDebugActor;
+	}
+
 	if (DebugActor)
 	{
 		const FMassEntityHandle EntityHandle = UE::Mass::Debug::GetEntityFromActor(*DebugActor, AgentComp);	
@@ -718,7 +724,6 @@ void FGameplayDebuggerCategory_Mass::DrawData(APlayerController* OwnerPC, FGamep
 	CanvasContext.Printf(TEXT("[{yellow}%s{white}] %s Entity overview"), *GetInputHandlerDescription(5), bShowNearEntityOverview ? TEXT("Hide") : TEXT("Show"));
 	CanvasContext.Printf(TEXT("[{yellow}%s{white}] %s Entity avoidance"), *GetInputHandlerDescription(6), bShowNearEntityAvoidance ? TEXT("Hide") : TEXT("Show"));
 	CanvasContext.Printf(TEXT("[{yellow}%s{white}] %s Entity path"), *GetInputHandlerDescription(7), bShowNearEntityPath ? TEXT("Hide") : TEXT("Show"));
-	CanvasContext.Printf(TEXT("[{yellow}%s{white}] Toggle Local/Remote debugging"), *GetInputHandlerDescription(8));
 
 	if (IsCategoryLocal() && !IsCategoryAuth())
 	{
