Fixed UMassStateTreeFragmentDestructor not calling Super::Initialize

diff --git a/Engine/Plugins/AI/MassAI/Source/MassAIBehavior/Private/MassStateTreeProcessors.cpp b/Engine/Plugins/AI/MassAI/Source/MassAIBehavior/Private/MassStateTreeProcessors.cpp
--- a/Engine/Plugins/AI/MassAI/Source/MassAIBehavior/Private/MassStateTreeProcessors.cpp
+++ b/Engine/Plugins/AI/MassAI/Source/MassAIBehavior/Private/MassStateTreeProcessors.cpp
@@ -211,6 +211,7 @@ UMassStateTreeFragmentDestructor::UMassStateTreeFragmentDestructor()
 void UMassStateTreeFragmentDestructor::Initialize(UObject& Owner)
 {
 	SignalSubsystem = UWorld::GetSubsystem<UMassSignalSubsystem>(Owner.GetWorld());
+	Super::Initialize(Owner);
 }
 
 void UMassStateTreeFragmentDestructor::ConfigureQueries()
