Mass gameplay debugger improvements:
- added an option to control "search range" for entity picking functionality
- massaged the entity scoring based on distance and angle to the camera, to better match intuitive expectations (the score is not a weighted average of angle and distance)
- made the category react to FMassDebugger::OnEntitySelectedDelegate broadcasts and causing it to trigger when an entity gets picked via the category (this improves synchronization with MassDebugger).

diff --git a/Engine/Plugins/AI/MassAI/Source/MassAIDebug/Private/GameplayDebuggerCategory_Mass.cpp b/Engine/Plugins/AI/MassAI/Source/MassAIDebug/Private/GameplayDebuggerCategory_Mass.cpp
--- a/Engine/Plugins/AI/MassAI/Source/MassAIDebug/Private/GameplayDebuggerCategory_Mass.cpp
+++ b/Engine/Plugins/AI/MassAI/Source/MassAIDebug/Private/GameplayDebuggerCategory_Mass.cpp
@@ -46,11 +46,11 @@ namespace UE::Mass::Debug
 		return EntityHandle;
 	};
 
-	FMassEntityHandle GetBestEntity(const FVector ViewLocation, const FVector ViewDirection, const TConstArrayView<FMassEntityHandle> Entities, const TConstArrayView<FVector> Locations, const bool bLimitAngle)
+	FMassEntityHandle GetBestEntity(const FVector ViewLocation, const FVector ViewDirection, const TConstArrayView<FMassEntityHandle> Entities
+		, const TConstArrayView<FVector> Locations, const bool bLimitAngle, const FVector::FReal MaxScanDistance)
 	{
-		// Reusing similar algorithm as UGameplayDebuggerLocalController for now 
-		constexpr FVector::FReal MaxScanDistanceSq = 25000. * 25000.;
-		constexpr FVector::FReal MinViewDirDot = 0.707; // 45 degrees
+		constexpr FVector::FReal MinViewDirDot = 0.707; // 45 degrees		
+		const FVector::FReal MaxScanDistanceSq = MaxScanDistance * MaxScanDistance;
 
 		checkf(Entities.Num() == Locations.Num(), TEXT("Both Entities and Locations lists are expected to be of the same size: %d vs %d"), Entities.Num(), Locations.Num());
 		
@@ -71,8 +71,10 @@ namespace UE::Mass::Debug
 				continue;
 			}
 
-			const FVector DirToEntityNormal = (FMath::IsNearlyZero(DistToEntitySq)) ? ViewDirection : (DirToEntity / FMath::Sqrt(DistToEntitySq));
+			const FVector::FReal Distance = FMath::Sqrt(DistToEntitySq);
+			const FVector DirToEntityNormal = (FMath::IsNearlyZero(DistToEntitySq)) ? ViewDirection : (DirToEntity / Distance);
 			const FVector::FReal ViewDot = FVector::DotProduct(ViewDirection, DirToEntityNormal);
+			const FVector::FReal Score = ViewDot * 0.1 * (1. - Distance / MaxScanDistance);
 			if (ViewDot > BestScore)
 			{
 				BestScore = ViewDot;
@@ -114,6 +116,8 @@ FGameplayDebuggerCategory_Mass::FGameplayDebuggerCategory_Mass()
 	BindKeyPress(EKeys::C.GetFName(), FGameplayDebuggerInputModifier::Shift, this, &FGameplayDebuggerCategory_Mass::OnToggleNearEntityPath, EGameplayDebuggerInputMode::Replicated);
 	ToggleDebugLocalEntityManagerInputIndex = GetNumInputHandlers();
 	BindKeyPress(EKeys::L.GetFName(), FGameplayDebuggerInputModifier::Shift, this, &FGameplayDebuggerCategory_Mass::OnToggleDebugLocalEntityManager, EGameplayDebuggerInputMode::Local);
+	BindKeyPress(EKeys::Add.GetFName(), FGameplayDebuggerInputModifier::Shift, this, &FGameplayDebuggerCategory_Mass::OnIncreaseSearchRange, EGameplayDebuggerInputMode::Replicated);
+	BindKeyPress(EKeys::Subtract.GetFName(), FGameplayDebuggerInputModifier::Shift, this, &FGameplayDebuggerCategory_Mass::OnDecreaseSearchRange, EGameplayDebuggerInputMode::Replicated);
 
 	ConsoleCommands.Emplace(TEXT("gdt.mass.ToggleArchetypes"), TEXT(""), FConsoleCommandDelegate::CreateLambda([this]() { OnToggleArchetypes(); }));
 	ConsoleCommands.Emplace(TEXT("gdt.mass.ToggleShapes"), TEXT(""), FConsoleCommandDelegate::CreateLambda([this]() { OnToggleShapes(); }));
@@ -123,27 +127,57 @@ FGameplayDebuggerCategory_Mass::FGameplayDebuggerCategory_Mass()
 	ConsoleCommands.Emplace(TEXT("gdt.mass.ToggleNearEntityOverview"), TEXT(""), FConsoleCommandDelegate::CreateLambda([this]() { OnToggleNearEntityOverview(); }));
 	ConsoleCommands.Emplace(TEXT("gdt.mass.ToggleNearEntityAvoidance"), TEXT(""), FConsoleCommandDelegate::CreateLambda([this]() { OnToggleNearEntityAvoidance(); }));
 	ConsoleCommands.Emplace(TEXT("gdt.mass.ToggleNearEntityPath"), TEXT(""), FConsoleCommandDelegate::CreateLambda([this]() { OnToggleNearEntityPath(); }));
+
+	OnEntitySelectedHandle = FMassDebugger::OnEntitySelectedDelegate.AddRaw(this, &FGameplayDebuggerCategory_Mass::OnEntitySelected);
+}
+
+FGameplayDebuggerCategory_Mass::~FGameplayDebuggerCategory_Mass()
+{
+	FMassDebugger::OnEntitySelectedDelegate.Remove(OnEntitySelectedHandle);
 }
 
 void FGameplayDebuggerCategory_Mass::SetCachedEntity(const FMassEntityHandle Entity, const FMassEntityManager& EntityManager)
 {
 	if (CachedEntity != Entity)
 	{
-		CachedEntity = Entity;
 		FMassDebugger::SelectEntity(EntityManager, Entity);
 	}
 }
 
+void FGameplayDebuggerCategory_Mass::OnEntitySelected(const FMassEntityManager& EntityManager, const FMassEntityHandle EntityHandle)
+{
+	UWorld* World = EntityManager.GetWorld();
+	AActor* BestActor = nullptr;
+	if (EntityHandle.IsSet() && World)
+	{
+		if (const UMassActorSubsystem* ActorSubsystem = World->GetSubsystem<UMassActorSubsystem>())
+		{
+			BestActor = ActorSubsystem->GetActorFromHandle(EntityHandle);
+		}
+	}
+
+	CachedEntity = EntityHandle;
+	CachedDebugActor = BestActor;
+	check(GetReplicator());
+	GetReplicator()->SetDebugActor(BestActor);
+}
+
+void FGameplayDebuggerCategory_Mass::ClearCachedEntity()
+{
+	CachedEntity = FMassEntityHandle();
+}
+
 void FGameplayDebuggerCategory_Mass::PickEntity(const FVector& ViewLocation, const FVector& ViewDirection, const UWorld& World, FMassEntityManager& EntityManager, const bool bLimitAngle)
 {
 	FMassEntityHandle BestEntity;
-	// entities indicated by UE::Mass::Debug take precedence 
-    if (UE::Mass::Debug::HasDebugEntities())
+	
+	// entities indicated by UE::Mass::Debug take precedence
+    if (UE::Mass::Debug::HasDebugEntities() && !UE::Mass::Debug::IsDebuggingSingleEntity())
     {
 		TArray<FMassEntityHandle> Entities;
 	    TArray<FVector> Locations;
 	    UE::Mass::Debug::GetDebugEntitiesAndLocations(EntityManager, Entities, Locations);
-	    BestEntity = UE::Mass::Debug::GetBestEntity(ViewLocation, ViewDirection, Entities, Locations, bLimitAngle);
+	    BestEntity = UE::Mass::Debug::GetBestEntity(ViewLocation, ViewDirection, Entities, Locations, bLimitAngle, SearchRange);
     }
 	else
 	{
@@ -163,21 +197,10 @@ void FGameplayDebuggerCategory_Mass::PickEntity(const FVector& ViewLocation, con
 			}
 		});
 
-		BestEntity = UE::Mass::Debug::GetBestEntity(ViewLocation, ViewDirection, Entities, Locations, bLimitAngle);
-	}
-
-	AActor* BestActor = nullptr;
-	if (BestEntity.IsSet())
-	{
-		if (const UMassActorSubsystem* ActorSubsystem = World.GetSubsystem<UMassActorSubsystem>())
-		{
-			BestActor = ActorSubsystem->GetActorFromHandle(FMassEntityHandle(BestEntity));
-		}
+		BestEntity = UE::Mass::Debug::GetBestEntity(ViewLocation, ViewDirection, Entities, Locations, bLimitAngle, SearchRange);
 	}
 
 	SetCachedEntity(BestEntity, EntityManager);
-	CachedDebugActor = BestActor;
-	GetReplicator()->SetDebugActor(BestActor);
 }
 
 TSharedRef<FGameplayDebuggerCategory> FGameplayDebuggerCategory_Mass::MakeInstance()
@@ -225,9 +248,13 @@ void FGameplayDebuggerCategory_Mass::CollectData(APlayerController* OwnerPC, AAc
 	}
 	else if (CachedDebugActor)
 	{
-		SetCachedEntity(FMassEntityHandle(), EntityManager);
+		ClearCachedEntity();
 		CachedDebugActor = nullptr;
 	}
+	else if (CachedEntity.IsValid() == true && EntityManager.IsEntityValid(CachedEntity) == false)
+	{
+		ClearCachedEntity();
+	}
 
 	FVector ViewLocation = FVector::ZeroVector;
 	FVector ViewDirection = FVector::ForwardVector;
@@ -241,7 +268,8 @@ void FGameplayDebuggerCategory_Mass::CollectData(APlayerController* OwnerPC, AAc
 			bPickEntity = false;
 		}
 		// if we're debugging based on UE::Mass::Debug and the range changed
-		else if (CachedDebugActor == nullptr && UE::Mass::Debug::HasDebugEntities() && UE::Mass::Debug::IsDebuggingEntity(CachedEntity) == false)
+		else if (CachedDebugActor == nullptr && UE::Mass::Debug::HasDebugEntities() && UE::Mass::Debug::IsDebuggingEntity(CachedEntity) == false
+			&& UE::Mass::Debug::IsDebuggingSingleEntity() == false)
 		{
 			// using bLimitAngle = false to not limit the selection to only the things in from of the player
 			PickEntity(ViewLocation, ViewDirection, *World, EntityManager, /*bLimitAngle=*/false);
@@ -250,16 +278,29 @@ void FGameplayDebuggerCategory_Mass::CollectData(APlayerController* OwnerPC, AAc
 
 	AddTextLine(FString::Printf(TEXT("{Green}Entities count active{grey}/all: {white}%d{grey}/%d"), EntityManager.DebugGetEntityCount(), EntityManager.DebugGetEntityCount()));
 	AddTextLine(FString::Printf(TEXT("{Green}Registered Archetypes count: {white}%d {green}data ver: {white}%d"), EntityManager.DebugGetArchetypesCount(), EntityManager.GetArchetypeDataVersion()));
+
+	AddTextLine(FString::Printf(TEXT("{Green}Search range: {White}%.0f"), SearchRange));
+
+	const FTransformFragment* TransformFragment = nullptr;
 	if (CachedEntity.IsValid())
 	{
 		AddTextLine(FString::Printf(TEXT("{Green}Entity: {White}%s"), *CachedEntity.DebugGetDescription()));
+		TransformFragment = EntityManager.GetFragmentDataPtr<FTransformFragment>(CachedEntity);
+		if (TransformFragment)
+		{
+			AddTextLine(FString::Printf(TEXT("{Green}Distance: {White}%.0f"), FVector::Distance(TransformFragment->GetTransform().GetLocation(), ViewLocation)));
+		}
 	}
 
 	if (UE::Mass::Debug::HasDebugEntities())
 	{
 		int32 RangeBegin, RangeEnd;
 		UE::Mass::Debug::GetDebugEntitiesRange(RangeBegin, RangeEnd);
-		AddTextLine(FString::Printf(TEXT("{Green}Debugged entity range: {orange}%d-%d"), RangeBegin, RangeEnd));
+		// not printing single-entity range, since in that case the CachedEntity is already set to the appropriate entity
+		if (RangeBegin != RangeEnd)
+		{
+			AddTextLine(FString::Printf(TEXT("{Green}Debugged entity range: {orange}%d-%d"), RangeBegin, RangeEnd));
+		}
 	}
 
 	if (bShowArchetypes)
@@ -271,14 +312,11 @@ void FGameplayDebuggerCategory_Mass::CollectData(APlayerController* OwnerPC, AAc
 		AddTextLine(Ar);
 	}
 
-	if (CachedEntity.IsSet() && bMarkEntityBeingDebugged)
+	if (CachedEntity.IsSet() && bMarkEntityBeingDebugged && TransformFragment)
 	{
-		if (const FTransformFragment* TransformFragment = EntityManager.GetFragmentDataPtr<FTransformFragment>(CachedEntity))
-		{
-			const FVector Location = TransformFragment->GetTransform().GetLocation();
-			AddShape(FGameplayDebuggerShape::MakeBox(Location, FVector(8,8,500), FColor::Purple,  FString::Printf(TEXT("[%s]"), *CachedEntity.DebugGetDescription())));
-			AddShape(FGameplayDebuggerShape::MakePoint(Location, 10, FColor::Purple));
-		}
+		const FVector Location = TransformFragment->GetTransform().GetLocation();
+		AddShape(FGameplayDebuggerShape::MakeBox(Location, FVector(8,8,500), FColor::Purple,  FString::Printf(TEXT("[%s]"), *CachedEntity.DebugGetDescription())));
+		AddShape(FGameplayDebuggerShape::MakePoint(Location, 10, FColor::Purple));
 	}
 
 	if (CachedEntity.IsSet() && Debugger)
@@ -371,9 +409,9 @@ void FGameplayDebuggerCategory_Mass::CollectData(APlayerController* OwnerPC, AAc
 					AddTextLine(FString::Printf(TEXT("{Green}Shared Fragments:{White}%s"), *DescriptionBuilder(ItemNames)));
 				}
 
-				const FTransformFragment& TransformFragment = EntityManager.GetFragmentDataChecked<FTransformFragment>(CachedEntity);
+				check(TransformFragment);
 				constexpr float CapsuleRadius = 50.f;
-				AddShape(FGameplayDebuggerShape::MakeCapsule(TransformFragment.GetTransform().GetLocation() + 2.f * CapsuleRadius * FVector::UpVector, CapsuleRadius, CapsuleRadius * 2.f, FColor::Orange));
+				AddShape(FGameplayDebuggerShape::MakeCapsule(TransformFragment->GetTransform().GetLocation() + 2.f * CapsuleRadius * FVector::UpVector, CapsuleRadius, CapsuleRadius * 2.f, FColor::Orange));
 			}
 			else
 			{
@@ -835,6 +873,16 @@ void FGameplayDebuggerCategory_Mass::OnToggleDebugLocalEntityManager()
 	CachedEntity.Reset();
 }
 
+void FGameplayDebuggerCategory_Mass::OnIncreaseSearchRange()
+{
+	SearchRange = FMath::Clamp(SearchRange * SearchRangeChangeScale, MinSearchRange, MaxSearchRange);
+}
+
+void FGameplayDebuggerCategory_Mass::OnDecreaseSearchRange()
+{
+	SearchRange = FMath::Clamp(SearchRange / SearchRangeChangeScale, MinSearchRange, MaxSearchRange);
+}
+
 //-----------------------------------------------------------------------------
 // DEPRECATED
 //-----------------------------------------------------------------------------
diff --git a/Engine/Plugins/AI/MassAI/Source/MassAIDebug/Public/GameplayDebuggerCategory_Mass.h b/Engine/Plugins/AI/MassAI/Source/MassAIDebug/Public/GameplayDebuggerCategory_Mass.h
--- a/Engine/Plugins/AI/MassAI/Source/MassAIDebug/Public/GameplayDebuggerCategory_Mass.h
+++ b/Engine/Plugins/AI/MassAI/Source/MassAIDebug/Public/GameplayDebuggerCategory_Mass.h
@@ -19,6 +19,7 @@ class FGameplayDebuggerCategory_Mass : public FGameplayDebuggerCategory
 {
 public:
 	FGameplayDebuggerCategory_Mass();
+	virtual ~FGameplayDebuggerCategory_Mass();
 
 	static TSharedRef<FGameplayDebuggerCategory> MakeInstance();
 
@@ -37,12 +38,17 @@ protected:
 	void OnToggleNearEntityAvoidance() { bShowNearEntityAvoidance = !bShowNearEntityAvoidance; }
 	void OnToggleNearEntityPath() { bShowNearEntityPath = !bShowNearEntityPath; }
 	void OnToggleDebugLocalEntityManager();
+	void OnIncreaseSearchRange();
+	void OnDecreaseSearchRange();
 	
 	void PickEntity(const FVector& ViewLocation, const FVector& ViewDirection, const UWorld& World, FMassEntityManager& EntityManager, const bool bLimitAngle = true);
 
 	UE_DEPRECATED(5.3, "This flavor of PickEntity has been deprecated. Use the one getting ViewLocation and ViewDirection parameters instead.")
 	void PickEntity(const APlayerController& OwnerPC, const UWorld& World, FMassEntityManager& EntityManager, const bool bLimitAngle = true);
 
+	void OnEntitySelected(const FMassEntityManager& EntityManager, const FMassEntityHandle EntityHandle);
+	void ClearCachedEntity();
+
 protected:
 	AActor* CachedDebugActor;
 	FMassEntityHandle CachedEntity;
@@ -57,6 +63,7 @@ protected:
 	bool bMarkEntityBeingDebugged;
 	bool bDebugLocalEntityManager;
 	int32 ToggleDebugLocalEntityManagerInputIndex = INDEX_NONE;
+	float SearchRange = 25000.f;
 
 	struct FEntityDescription
 	{
@@ -70,6 +77,12 @@ protected:
 	TArray<FEntityDescription> NearEntityDescriptions;
 
 	TArray<FAutoConsoleCommand> ConsoleCommands;
+
+	FDelegateHandle OnEntitySelectedHandle;
+
+	static constexpr float MaxSearchRange = 1000000.f;
+	static constexpr float MinSearchRange = 1.f;
+	static constexpr float SearchRangeChangeScale = 2.f;
 };
 
 #endif // WITH_GAMEPLAY_DEBUGGER && WITH_MASSGAMEPLAY_DEBUG
diff --git a/Engine/Source/Runtime/GameplayDebugger/Public/GameplayDebuggerCategoryReplicator.h b/Engine/Source/Runtime/GameplayDebugger/Public/GameplayDebuggerCategoryReplicator.h
--- a/Engine/Source/Runtime/GameplayDebugger/Public/GameplayDebuggerCategoryReplicator.h
+++ b/Engine/Source/Runtime/GameplayDebugger/Public/GameplayDebuggerCategoryReplicator.h
@@ -172,7 +172,7 @@ public:
 	GAMEPLAYDEBUGGER_API void CollectCategoryData(bool bForce = false);
 
 	/** get current debug actor */
-	AActor* GetDebugActor() const { return DebugActor.Actor.Get(); }
+	AActor* GetDebugActor() const { return DebugActor.Actor.IsStale() ? nullptr : DebugActor.Actor.Get(); }
 	
 	/** get name of debug actor */
 	FName GetDebugActorName() const { return DebugActor.ActorName; }
