Fixed Mass gameplay debugger category's console commands not working in client-server environment (due to trying to register a each command twice)

diff --git a/Engine/Plugins/AI/MassAI/Source/MassAIDebug/Private/GameplayDebuggerCategory_Mass.cpp b/Engine/Plugins/AI/MassAI/Source/MassAIDebug/Private/GameplayDebuggerCategory_Mass.cpp
--- a/Engine/Plugins/AI/MassAI/Source/MassAIDebug/Private/GameplayDebuggerCategory_Mass.cpp
+++ b/Engine/Plugins/AI/MassAI/Source/MassAIDebug/Private/GameplayDebuggerCategory_Mass.cpp
@@ -89,6 +89,16 @@ namespace UE::Mass::Debug
 //----------------------------------------------------------------------//
 //  FGameplayDebuggerCategory_Mass
 //----------------------------------------------------------------------//
+TArray<FAutoConsoleCommandWithWorld> FGameplayDebuggerCategory_Mass::ConsoleCommands;
+FGameplayDebuggerCategory_Mass::FOnConsoleCommandBroadcastDelegate FGameplayDebuggerCategory_Mass::OnToggleArchetypesBroadcast;
+FGameplayDebuggerCategory_Mass::FOnConsoleCommandBroadcastDelegate FGameplayDebuggerCategory_Mass::OnToggleShapesBroadcast;
+FGameplayDebuggerCategory_Mass::FOnConsoleCommandBroadcastDelegate FGameplayDebuggerCategory_Mass::OnToggleAgentFragmentsBroadcast;
+FGameplayDebuggerCategory_Mass::FOnConsoleCommandBroadcastDelegate FGameplayDebuggerCategory_Mass::OnPickEntityBroadcast;
+FGameplayDebuggerCategory_Mass::FOnConsoleCommandBroadcastDelegate FGameplayDebuggerCategory_Mass::OnToggleEntityDetailsBroadcast;
+FGameplayDebuggerCategory_Mass::FOnConsoleCommandBroadcastDelegate FGameplayDebuggerCategory_Mass::OnToggleNearEntityOverviewBroadcast;
+FGameplayDebuggerCategory_Mass::FOnConsoleCommandBroadcastDelegate FGameplayDebuggerCategory_Mass::OnToggleNearEntityAvoidanceBroadcast;
+FGameplayDebuggerCategory_Mass::FOnConsoleCommandBroadcastDelegate FGameplayDebuggerCategory_Mass::OnToggleNearEntityPathBroadcast;
+
 FGameplayDebuggerCategory_Mass::FGameplayDebuggerCategory_Mass()
 {
 	CachedDebugActor = nullptr;
@@ -119,14 +129,26 @@ FGameplayDebuggerCategory_Mass::FGameplayDebuggerCategory_Mass()
 	BindKeyPress(EKeys::Add.GetFName(), FGameplayDebuggerInputModifier::Shift, this, &FGameplayDebuggerCategory_Mass::OnIncreaseSearchRange, EGameplayDebuggerInputMode::Replicated);
 	BindKeyPress(EKeys::Subtract.GetFName(), FGameplayDebuggerInputModifier::Shift, this, &FGameplayDebuggerCategory_Mass::OnDecreaseSearchRange, EGameplayDebuggerInputMode::Replicated);
 
-	ConsoleCommands.Emplace(TEXT("gdt.mass.ToggleArchetypes"), TEXT(""), FConsoleCommandDelegate::CreateLambda([this]() { OnToggleArchetypes(); }));
-	ConsoleCommands.Emplace(TEXT("gdt.mass.ToggleShapes"), TEXT(""), FConsoleCommandDelegate::CreateLambda([this]() { OnToggleShapes(); }));
-	ConsoleCommands.Emplace(TEXT("gdt.mass.ToggleAgentFragments"), TEXT(""), FConsoleCommandDelegate::CreateLambda([this]() { OnToggleAgentFragments(); }));
-	ConsoleCommands.Emplace(TEXT("gdt.mass.PickEntity"), TEXT(""), FConsoleCommandDelegate::CreateLambda([this]() { OnPickEntity(); }));
-	ConsoleCommands.Emplace(TEXT("gdt.mass.ToggleEntityDetails"), TEXT(""), FConsoleCommandDelegate::CreateLambda([this]() { OnToggleEntityDetails(); }));
-	ConsoleCommands.Emplace(TEXT("gdt.mass.ToggleNearEntityOverview"), TEXT(""), FConsoleCommandDelegate::CreateLambda([this]() { OnToggleNearEntityOverview(); }));
-	ConsoleCommands.Emplace(TEXT("gdt.mass.ToggleNearEntityAvoidance"), TEXT(""), FConsoleCommandDelegate::CreateLambda([this]() { OnToggleNearEntityAvoidance(); }));
-	ConsoleCommands.Emplace(TEXT("gdt.mass.ToggleNearEntityPath"), TEXT(""), FConsoleCommandDelegate::CreateLambda([this]() { OnToggleNearEntityPath(); }));
+	if (ConsoleCommands.Num() == 0)
+	{
+		ConsoleCommands.Emplace(TEXT("gdt.mass.ToggleArchetypes"), TEXT(""), FConsoleCommandWithWorldDelegate::CreateLambda([](UWorld* InWorld) { OnToggleArchetypesBroadcast.Broadcast(InWorld); }));
+		ConsoleCommands.Emplace(TEXT("gdt.mass.ToggleShapes"), TEXT(""), FConsoleCommandWithWorldDelegate::CreateLambda([](UWorld* InWorld) { OnToggleShapesBroadcast.Broadcast(InWorld); }));
+		ConsoleCommands.Emplace(TEXT("gdt.mass.ToggleAgentFragments"), TEXT(""), FConsoleCommandWithWorldDelegate::CreateLambda([](UWorld* InWorld) { OnToggleAgentFragmentsBroadcast.Broadcast(InWorld); }));
+		ConsoleCommands.Emplace(TEXT("gdt.mass.PickEntity"), TEXT(""), FConsoleCommandWithWorldDelegate::CreateLambda([](UWorld* InWorld) { OnPickEntityBroadcast.Broadcast(InWorld); }));
+		ConsoleCommands.Emplace(TEXT("gdt.mass.ToggleEntityDetails"), TEXT(""), FConsoleCommandWithWorldDelegate::CreateLambda([](UWorld* InWorld) { OnToggleEntityDetailsBroadcast.Broadcast(InWorld); }));
+		ConsoleCommands.Emplace(TEXT("gdt.mass.ToggleNearEntityOverview"), TEXT(""), FConsoleCommandWithWorldDelegate::CreateLambda([](UWorld* InWorld) { OnToggleNearEntityOverviewBroadcast.Broadcast(InWorld); }));
+		ConsoleCommands.Emplace(TEXT("gdt.mass.ToggleNearEntityAvoidance"), TEXT(""), FConsoleCommandWithWorldDelegate::CreateLambda([](UWorld* InWorld) { OnToggleNearEntityAvoidanceBroadcast.Broadcast(InWorld); }));
+		ConsoleCommands.Emplace(TEXT("gdt.mass.ToggleNearEntityPath"), TEXT(""), FConsoleCommandWithWorldDelegate::CreateLambda([](UWorld* InWorld) { OnToggleNearEntityPathBroadcast.Broadcast(InWorld); }));
+	}
+
+	ConsoleCommandHandles.Add(FDelegateHandlePair(&OnToggleArchetypesBroadcast, OnToggleArchetypesBroadcast.AddLambda([this](UWorld* InWorld) { if (InWorld == GetWorldFromReplicator()) { OnToggleArchetypes(); }})));
+	ConsoleCommandHandles.Add(FDelegateHandlePair(&OnToggleShapesBroadcast, OnToggleShapesBroadcast.AddLambda([this](UWorld* InWorld) { if (InWorld == GetWorldFromReplicator()) { OnToggleShapes(); }})));
+	ConsoleCommandHandles.Add(FDelegateHandlePair(&OnToggleAgentFragmentsBroadcast, OnToggleAgentFragmentsBroadcast.AddLambda([this](UWorld* InWorld) { if (InWorld == GetWorldFromReplicator()) { OnToggleAgentFragments(); }})));
+	ConsoleCommandHandles.Add(FDelegateHandlePair(&OnPickEntityBroadcast, OnPickEntityBroadcast.AddLambda([this](UWorld* InWorld) { if (InWorld == GetWorldFromReplicator()) { OnPickEntity(); }})));
+	ConsoleCommandHandles.Add(FDelegateHandlePair(&OnToggleEntityDetailsBroadcast, OnToggleEntityDetailsBroadcast.AddLambda([this](UWorld* InWorld) { if (InWorld == GetWorldFromReplicator()) { OnToggleEntityDetails(); }})));
+	ConsoleCommandHandles.Add(FDelegateHandlePair(&OnToggleNearEntityOverviewBroadcast, OnToggleNearEntityOverviewBroadcast.AddLambda([this](UWorld* InWorld) { if (InWorld == GetWorldFromReplicator()) { OnToggleNearEntityOverview(); }})));
+	ConsoleCommandHandles.Add(FDelegateHandlePair(&OnToggleNearEntityAvoidanceBroadcast, OnToggleNearEntityAvoidanceBroadcast.AddLambda([this](UWorld* InWorld) { if (InWorld == GetWorldFromReplicator()) { OnToggleNearEntityAvoidance(); }})));
+	ConsoleCommandHandles.Add(FDelegateHandlePair(&OnToggleNearEntityPathBroadcast, OnToggleNearEntityPathBroadcast.AddLambda([this](UWorld* InWorld) { if (InWorld == GetWorldFromReplicator()) { OnToggleNearEntityPath(); }})));
 
 	OnEntitySelectedHandle = FMassDebugger::OnEntitySelectedDelegate.AddRaw(this, &FGameplayDebuggerCategory_Mass::OnEntitySelected);
 }
@@ -134,6 +156,12 @@ FGameplayDebuggerCategory_Mass::FGameplayDebuggerCategory_Mass()
 FGameplayDebuggerCategory_Mass::~FGameplayDebuggerCategory_Mass()
 {
 	FMassDebugger::OnEntitySelectedDelegate.Remove(OnEntitySelectedHandle);
+
+	for (FDelegateHandlePair& Pair : ConsoleCommandHandles)
+	{
+		CA_ASSUME(Pair.Key);
+		Pair.Key->Remove(Pair.Value);
+	}
 }
 
 void FGameplayDebuggerCategory_Mass::SetCachedEntity(const FMassEntityHandle Entity, const FMassEntityManager& EntityManager)
@@ -147,6 +175,12 @@ void FGameplayDebuggerCategory_Mass::SetCachedEntity(const FMassEntityHandle Ent
 void FGameplayDebuggerCategory_Mass::OnEntitySelected(const FMassEntityManager& EntityManager, const FMassEntityHandle EntityHandle)
 {
 	UWorld* World = EntityManager.GetWorld();
+	if (World != GetWorldFromReplicator())
+	{ 
+		// ignore, this call is for a different world
+		return;
+	}
+
 	AActor* BestActor = nullptr;
 	if (EntityHandle.IsSet() && World)
 	{
diff --git a/Engine/Plugins/AI/MassAI/Source/MassAIDebug/Public/GameplayDebuggerCategory_Mass.h b/Engine/Plugins/AI/MassAI/Source/MassAIDebug/Public/GameplayDebuggerCategory_Mass.h
--- a/Engine/Plugins/AI/MassAI/Source/MassAIDebug/Public/GameplayDebuggerCategory_Mass.h
+++ b/Engine/Plugins/AI/MassAI/Source/MassAIDebug/Public/GameplayDebuggerCategory_Mass.h
@@ -17,6 +17,7 @@ class AActor;
 
 class FGameplayDebuggerCategory_Mass : public FGameplayDebuggerCategory
 {
+	using Super = FGameplayDebuggerCategory;
 public:
 	FGameplayDebuggerCategory_Mass();
 	virtual ~FGameplayDebuggerCategory_Mass();
@@ -76,7 +77,19 @@ protected:
 	};
 	TArray<FEntityDescription> NearEntityDescriptions;
 
-	TArray<FAutoConsoleCommand> ConsoleCommands;
+	static TArray<FAutoConsoleCommandWithWorld> ConsoleCommands;
+	DECLARE_MULTICAST_DELEGATE_OneParam(FOnConsoleCommandBroadcastDelegate, UWorld*);
+	static FOnConsoleCommandBroadcastDelegate OnToggleArchetypesBroadcast;
+	static FOnConsoleCommandBroadcastDelegate OnToggleShapesBroadcast;
+	static FOnConsoleCommandBroadcastDelegate OnToggleAgentFragmentsBroadcast;
+	static FOnConsoleCommandBroadcastDelegate OnPickEntityBroadcast;
+	static FOnConsoleCommandBroadcastDelegate OnToggleEntityDetailsBroadcast;
+	static FOnConsoleCommandBroadcastDelegate OnToggleNearEntityOverviewBroadcast;
+	static FOnConsoleCommandBroadcastDelegate OnToggleNearEntityAvoidanceBroadcast;
+	static FOnConsoleCommandBroadcastDelegate OnToggleNearEntityPathBroadcast; 	
+
+	using FDelegateHandlePair = TPair<FOnConsoleCommandBroadcastDelegate*, FDelegateHandle>;
+	TArray<FDelegateHandlePair> ConsoleCommandHandles;
 
 	FDelegateHandle OnEntitySelectedHandle;
 
