MassAI - Added gdt.mass.ToggleDebugLocalEntityManager console command to toggle between local / remote entity manager debugging in the Mass gameplay debugger category.

diff --git a/Engine/Plugins/AI/MassAI/Source/MassAIDebug/Private/GameplayDebuggerCategory_Mass.cpp b/Engine/Plugins/AI/MassAI/Source/MassAIDebug/Private/GameplayDebuggerCategory_Mass.cpp
--- a/Engine/Plugins/AI/MassAI/Source/MassAIDebug/Private/GameplayDebuggerCategory_Mass.cpp
+++ b/Engine/Plugins/AI/MassAI/Source/MassAIDebug/Private/GameplayDebuggerCategory_Mass.cpp
@@ -98,6 +98,7 @@ FGameplayDebuggerCategory_Mass::FOnConsoleCommandBroadcastDelegate FGameplayDebu
 FGameplayDebuggerCategory_Mass::FOnConsoleCommandBroadcastDelegate FGameplayDebuggerCategory_Mass::OnToggleNearEntityOverviewBroadcast;
 FGameplayDebuggerCategory_Mass::FOnConsoleCommandBroadcastDelegate FGameplayDebuggerCategory_Mass::OnToggleNearEntityAvoidanceBroadcast;
 FGameplayDebuggerCategory_Mass::FOnConsoleCommandBroadcastDelegate FGameplayDebuggerCategory_Mass::OnToggleNearEntityPathBroadcast;
+FGameplayDebuggerCategory_Mass::FOnConsoleCommandBroadcastDelegate FGameplayDebuggerCategory_Mass::OnToggleDebugLocalEntityManagerBroadcast;
 
 FGameplayDebuggerCategory_Mass::FGameplayDebuggerCategory_Mass()
 {
@@ -139,6 +140,7 @@ FGameplayDebuggerCategory_Mass::FGameplayDebuggerCategory_Mass()
 		ConsoleCommands.Emplace(TEXT("gdt.mass.ToggleNearEntityOverview"), TEXT(""), FConsoleCommandWithWorldDelegate::CreateLambda([](UWorld* InWorld) { OnToggleNearEntityOverviewBroadcast.Broadcast(InWorld); }));
 		ConsoleCommands.Emplace(TEXT("gdt.mass.ToggleNearEntityAvoidance"), TEXT(""), FConsoleCommandWithWorldDelegate::CreateLambda([](UWorld* InWorld) { OnToggleNearEntityAvoidanceBroadcast.Broadcast(InWorld); }));
 		ConsoleCommands.Emplace(TEXT("gdt.mass.ToggleNearEntityPath"), TEXT(""), FConsoleCommandWithWorldDelegate::CreateLambda([](UWorld* InWorld) { OnToggleNearEntityPathBroadcast.Broadcast(InWorld); }));
+		ConsoleCommands.Emplace(TEXT("gdt.mass.ToggleDebugLocalEntityManager"), TEXT(""), FConsoleCommandWithWorldDelegate::CreateLambda([](UWorld* InWorld) { OnToggleDebugLocalEntityManagerBroadcast.Broadcast(InWorld); }));
 	}
 
 	ConsoleCommandHandles.Add(FDelegateHandlePair(&OnToggleArchetypesBroadcast, OnToggleArchetypesBroadcast.AddLambda([this](UWorld* InWorld) { if (InWorld == GetWorldFromReplicator()) { OnToggleArchetypes(); }})));
@@ -149,6 +151,7 @@ FGameplayDebuggerCategory_Mass::FGameplayDebuggerCategory_Mass()
 	ConsoleCommandHandles.Add(FDelegateHandlePair(&OnToggleNearEntityOverviewBroadcast, OnToggleNearEntityOverviewBroadcast.AddLambda([this](UWorld* InWorld) { if (InWorld == GetWorldFromReplicator()) { OnToggleNearEntityOverview(); }})));
 	ConsoleCommandHandles.Add(FDelegateHandlePair(&OnToggleNearEntityAvoidanceBroadcast, OnToggleNearEntityAvoidanceBroadcast.AddLambda([this](UWorld* InWorld) { if (InWorld == GetWorldFromReplicator()) { OnToggleNearEntityAvoidance(); }})));
 	ConsoleCommandHandles.Add(FDelegateHandlePair(&OnToggleNearEntityPathBroadcast, OnToggleNearEntityPathBroadcast.AddLambda([this](UWorld* InWorld) { if (InWorld == GetWorldFromReplicator()) { OnToggleNearEntityPath(); }})));
+	ConsoleCommandHandles.Add(FDelegateHandlePair(&OnToggleDebugLocalEntityManagerBroadcast, OnToggleDebugLocalEntityManagerBroadcast.AddLambda([this](UWorld* InWorld) { if (InWorld == GetWorldFromReplicator()) { OnToggleDebugLocalEntityManager(); }})));
 
 	OnEntitySelectedHandle = FMassDebugger::OnEntitySelectedDelegate.AddRaw(this, &FGameplayDebuggerCategory_Mass::OnEntitySelected);
 }
diff --git a/Engine/Plugins/AI/MassAI/Source/MassAIDebug/Public/GameplayDebuggerCategory_Mass.h b/Engine/Plugins/AI/MassAI/Source/MassAIDebug/Public/GameplayDebuggerCategory_Mass.h
--- a/Engine/Plugins/AI/MassAI/Source/MassAIDebug/Public/GameplayDebuggerCategory_Mass.h
+++ b/Engine/Plugins/AI/MassAI/Source/MassAIDebug/Public/GameplayDebuggerCategory_Mass.h
@@ -87,6 +87,7 @@ protected:
 	static FOnConsoleCommandBroadcastDelegate OnToggleNearEntityOverviewBroadcast;
 	static FOnConsoleCommandBroadcastDelegate OnToggleNearEntityAvoidanceBroadcast;
 	static FOnConsoleCommandBroadcastDelegate OnToggleNearEntityPathBroadcast; 	
+	static FOnConsoleCommandBroadcastDelegate OnToggleDebugLocalEntityManagerBroadcast; 	
 
 	using FDelegateHandlePair = TPair<FOnConsoleCommandBroadcastDelegate*, FDelegateHandle>;
 	TArray<FDelegateHandlePair> ConsoleCommandHandles;
