Added a way to make actor LOD viewers via Mass gameplay debugger category

diff --git a/Engine/Plugins/AI/MassAI/Source/MassAIDebug/Private/GameplayDebuggerCategory_Mass.cpp b/Engine/Plugins/AI/MassAI/Source/MassAIDebug/Private/GameplayDebuggerCategory_Mass.cpp
--- a/Engine/Plugins/AI/MassAI/Source/MassAIDebug/Private/GameplayDebuggerCategory_Mass.cpp
+++ b/Engine/Plugins/AI/MassAI/Source/MassAIDebug/Private/GameplayDebuggerCategory_Mass.cpp
@@ -28,6 +28,8 @@
 #include "CanvasItem.h"
 #include "Engine/World.h"
 #include "MassDebugger.h"
+#include "MassLODSubsystem.h"
+
 
 namespace UE::Mass::Debug
 {
@@ -99,6 +101,9 @@ FGameplayDebuggerCategory_Mass::FOnConsoleCommandBroadcastDelegate FGameplayDebu
 FGameplayDebuggerCategory_Mass::FOnConsoleCommandBroadcastDelegate FGameplayDebuggerCategory_Mass::OnToggleNearEntityAvoidanceBroadcast;
 FGameplayDebuggerCategory_Mass::FOnConsoleCommandBroadcastDelegate FGameplayDebuggerCategory_Mass::OnToggleNearEntityPathBroadcast;
 FGameplayDebuggerCategory_Mass::FOnConsoleCommandBroadcastDelegate FGameplayDebuggerCategory_Mass::OnToggleDebugLocalEntityManagerBroadcast;
+FGameplayDebuggerCategory_Mass::FOnConsoleCommandBroadcastDelegate FGameplayDebuggerCategory_Mass::OnTogglePickedActorAsViewerBroadcast;
+FGameplayDebuggerCategory_Mass::FOnConsoleCommandBroadcastDelegate FGameplayDebuggerCategory_Mass::OnToggleDrawViewersBroadcast;
+FGameplayDebuggerCategory_Mass::FOnConsoleCommandBroadcastDelegate FGameplayDebuggerCategory_Mass::OnClearActorViewersBroadcast;
 
 FGameplayDebuggerCategory_Mass::FGameplayDebuggerCategory_Mass()
 {
@@ -116,6 +121,7 @@ FGameplayDebuggerCategory_Mass::FGameplayDebuggerCategory_Mass()
 	bShowNearEntityPath = false;
 	bMarkEntityBeingDebugged = true;
 	bDebugLocalEntityManager = false;
+	bShowViewers = false;
 
 	BindKeyPress(EKeys::A.GetFName(), FGameplayDebuggerInputModifier::Shift, this, &FGameplayDebuggerCategory_Mass::OnToggleArchetypes, EGameplayDebuggerInputMode::Replicated);
 	BindKeyPress(EKeys::S.GetFName(), FGameplayDebuggerInputModifier::Shift, this, &FGameplayDebuggerCategory_Mass::OnToggleShapes, EGameplayDebuggerInputMode::Replicated);
@@ -129,6 +135,12 @@ FGameplayDebuggerCategory_Mass::FGameplayDebuggerCategory_Mass()
 	BindKeyPress(EKeys::L.GetFName(), FGameplayDebuggerInputModifier::Shift, this, &FGameplayDebuggerCategory_Mass::OnToggleDebugLocalEntityManager, EGameplayDebuggerInputMode::Local);
 	BindKeyPress(EKeys::Add.GetFName(), FGameplayDebuggerInputModifier::Shift, this, &FGameplayDebuggerCategory_Mass::OnIncreaseSearchRange, EGameplayDebuggerInputMode::Replicated);
 	BindKeyPress(EKeys::Subtract.GetFName(), FGameplayDebuggerInputModifier::Shift, this, &FGameplayDebuggerCategory_Mass::OnDecreaseSearchRange, EGameplayDebuggerInputMode::Replicated);
+	TogglePickedActorAsViewerInputIndex = GetNumInputHandlers();
+	BindKeyPress(EKeys::U.GetFName(), FGameplayDebuggerInputModifier::Ctrl, this, &FGameplayDebuggerCategory_Mass::OnTogglePickedActorAsViewer, EGameplayDebuggerInputMode::Replicated);
+	ToggleDrawViewersInputIndex = GetNumInputHandlers();
+	BindKeyPress(EKeys::U.GetFName(), FGameplayDebuggerInputModifier::Shift, this, &FGameplayDebuggerCategory_Mass::OnToggleDrawViewers, EGameplayDebuggerInputMode::Replicated);
+	ClearViewersInputIndex = GetNumInputHandlers();
+	BindKeyPress(EKeys::U.GetFName(), FGameplayDebuggerInputModifier::Shift + FGameplayDebuggerInputModifier::Ctrl, this, &FGameplayDebuggerCategory_Mass::OnClearActorViewers, EGameplayDebuggerInputMode::Replicated);
 
 	if (ConsoleCommands.Num() == 0)
 	{
@@ -141,6 +153,9 @@ FGameplayDebuggerCategory_Mass::FGameplayDebuggerCategory_Mass()
 		ConsoleCommands.Emplace(TEXT("gdt.mass.ToggleNearEntityAvoidance"), TEXT(""), FConsoleCommandWithWorldDelegate::CreateLambda([](UWorld* InWorld) { OnToggleNearEntityAvoidanceBroadcast.Broadcast(InWorld); }));
 		ConsoleCommands.Emplace(TEXT("gdt.mass.ToggleNearEntityPath"), TEXT(""), FConsoleCommandWithWorldDelegate::CreateLambda([](UWorld* InWorld) { OnToggleNearEntityPathBroadcast.Broadcast(InWorld); }));
 		ConsoleCommands.Emplace(TEXT("gdt.mass.ToggleDebugLocalEntityManager"), TEXT(""), FConsoleCommandWithWorldDelegate::CreateLambda([](UWorld* InWorld) { OnToggleDebugLocalEntityManagerBroadcast.Broadcast(InWorld); }));
+		ConsoleCommands.Emplace(TEXT("gdt.mass.TogglePickedActorAsViewer"), TEXT(""), FConsoleCommandWithWorldDelegate::CreateLambda([](UWorld* InWorld) { OnTogglePickedActorAsViewerBroadcast.Broadcast(InWorld); }));
+		ConsoleCommands.Emplace(TEXT("gdt.mass.ToggleDrawViewers"), TEXT(""), FConsoleCommandWithWorldDelegate::CreateLambda([](UWorld* InWorld) { OnToggleDrawViewersBroadcast.Broadcast(InWorld); }));	
+		ConsoleCommands.Emplace(TEXT("gdt.mass.ClearActorViewers"), TEXT(""), FConsoleCommandWithWorldDelegate::CreateLambda([](UWorld* InWorld) { OnClearActorViewersBroadcast.Broadcast(InWorld); }));
 	}
 
 	ConsoleCommandHandles.Add(FDelegateHandlePair(&OnToggleArchetypesBroadcast, OnToggleArchetypesBroadcast.AddLambda([this](UWorld* InWorld) { if (InWorld == GetWorldFromReplicator()) { OnToggleArchetypes(); }})));
@@ -152,6 +167,9 @@ FGameplayDebuggerCategory_Mass::FGameplayDebuggerCategory_Mass()
 	ConsoleCommandHandles.Add(FDelegateHandlePair(&OnToggleNearEntityAvoidanceBroadcast, OnToggleNearEntityAvoidanceBroadcast.AddLambda([this](UWorld* InWorld) { if (InWorld == GetWorldFromReplicator()) { OnToggleNearEntityAvoidance(); }})));
 	ConsoleCommandHandles.Add(FDelegateHandlePair(&OnToggleNearEntityPathBroadcast, OnToggleNearEntityPathBroadcast.AddLambda([this](UWorld* InWorld) { if (InWorld == GetWorldFromReplicator()) { OnToggleNearEntityPath(); }})));
 	ConsoleCommandHandles.Add(FDelegateHandlePair(&OnToggleDebugLocalEntityManagerBroadcast, OnToggleDebugLocalEntityManagerBroadcast.AddLambda([this](UWorld* InWorld) { if (InWorld == GetWorldFromReplicator()) { OnToggleDebugLocalEntityManager(); }})));
+	ConsoleCommandHandles.Add(FDelegateHandlePair(&OnTogglePickedActorAsViewerBroadcast, OnTogglePickedActorAsViewerBroadcast.AddLambda([this](UWorld* InWorld) { if (InWorld == GetWorldFromReplicator()) { OnTogglePickedActorAsViewer(); }})));
+	ConsoleCommandHandles.Add(FDelegateHandlePair(&OnToggleDrawViewersBroadcast, OnToggleDrawViewersBroadcast.AddLambda([this](UWorld* InWorld) { if (InWorld == GetWorldFromReplicator()) { OnToggleDrawViewers(); }})));
+	ConsoleCommandHandles.Add(FDelegateHandlePair(&OnClearActorViewersBroadcast, OnClearActorViewersBroadcast.AddLambda([this](UWorld* InWorld) { if (InWorld == GetWorldFromReplicator()) { OnClearActorViewers(); }})));
 
 	OnEntitySelectedHandle = FMassDebugger::OnEntitySelectedDelegate.AddRaw(this, &FGameplayDebuggerCategory_Mass::OnEntitySelected);
 }
@@ -247,6 +265,9 @@ TSharedRef<FGameplayDebuggerCategory> FGameplayDebuggerCategory_Mass::MakeInstan
 
 void FGameplayDebuggerCategory_Mass::CollectData(APlayerController* OwnerPC, AActor* DebugActor)
 {
+	constexpr float ViewerCylinderMarkLength = 1500.f;
+	constexpr float ViewerCylinderMarkLRadius = 30.f;
+
 	if (bAllowLocalDataCollection)
 	{
 		ResetReplicatedData();
@@ -274,7 +295,7 @@ void FGameplayDebuggerCategory_Mass::CollectData(APlayerController* OwnerPC, AAc
 	
 	if (bAllowLocalDataCollection)
 	{
-		DebugActor = CachedDebugActor;
+		DebugActor = CachedDebugActor.GetEvenIfUnreachable();
 	}
 
 	if (DebugActor)
@@ -283,7 +304,7 @@ void FGameplayDebuggerCategory_Mass::CollectData(APlayerController* OwnerPC, AAc
 		SetCachedEntity(EntityHandle, EntityManager);
 		CachedDebugActor = DebugActor;
 	}
-	else if (CachedDebugActor)
+	else if (CachedDebugActor.Get())
 	{
 		ClearCachedEntity();
 		CachedDebugActor = nullptr;
@@ -349,6 +370,17 @@ void FGameplayDebuggerCategory_Mass::CollectData(APlayerController* OwnerPC, AAc
 		AddTextLine(Ar);
 	}
 
+	if (bShowViewers)
+	{
+		if (UMassLODSubsystem* LODSubsystem = World->GetSubsystem<UMassLODSubsystem>())
+		{
+			for (const FViewerInfo& Viewer : LODSubsystem->GetViewers())
+			{
+				AddShape(FGameplayDebuggerShape::MakeCylinder(Viewer.Location, ViewerCylinderMarkLRadius, ViewerCylinderMarkLength, FColor::Blue));
+			}
+		}
+	}
+
 	if (CachedEntity.IsSet() && bMarkEntityBeingDebugged && TransformFragment)
 	{
 		const FVector Location = TransformFragment->GetTransform().GetLocation();
@@ -796,6 +828,9 @@ void FGameplayDebuggerCategory_Mass::DrawData(APlayerController* OwnerPC, FGamep
 		CanvasContext.Printf(TEXT("{grey}[%s] Entity details [enable Agent Fragments]{white}"), *GetInputHandlerDescription(4));
 	}
 	CanvasContext.Printf(TEXT("[{yellow}%s{white}] Pick Entity"), *GetInputHandlerDescription(3));
+	CanvasContext.Printf(TEXT("[{yellow}%s{white}] Toggle Picked Actor as viewer"), *GetInputHandlerDescription(TogglePickedActorAsViewerInputIndex));
+	CanvasContext.Printf(TEXT("[{yellow}%s{white}] %s LOD Viewers"), *GetInputHandlerDescription(ToggleDrawViewersInputIndex), bShowViewers ? TEXT("Hide") : TEXT("Show"));
+	CanvasContext.Printf(TEXT("[{yellow}%s{white}] Reset Actor LOD Viewers"), *GetInputHandlerDescription(ClearViewersInputIndex));
 	CanvasContext.Printf(TEXT("[{yellow}%s{white}] %s Entity overview"), *GetInputHandlerDescription(5), bShowNearEntityOverview ? TEXT("Hide") : TEXT("Show"));
 	CanvasContext.Printf(TEXT("[{yellow}%s{white}] %s Entity avoidance"), *GetInputHandlerDescription(6), bShowNearEntityAvoidance ? TEXT("Hide") : TEXT("Show"));
 	CanvasContext.Printf(TEXT("[{yellow}%s{white}] %s Entity path"), *GetInputHandlerDescription(7), bShowNearEntityPath ? TEXT("Hide") : TEXT("Show"));
@@ -920,6 +955,35 @@ void FGameplayDebuggerCategory_Mass::OnDecreaseSearchRange()
 	SearchRange = FMath::Clamp(SearchRange / SearchRangeChangeScale, MinSearchRange, MaxSearchRange);
 }
 
+void FGameplayDebuggerCategory_Mass::OnTogglePickedActorAsViewer()
+{
+	if (AActor* DebugActor = CachedDebugActor.Get())
+	{
+		UWorld* World = GetWorldFromReplicator();
+		if (UMassLODSubsystem* LODSubsystem = World->GetSubsystem<UMassLODSubsystem>())
+		{
+			FMassViewerHandle ViewerHandle = LODSubsystem->GetViewerHandleFromActor(*DebugActor);
+			if (ViewerHandle.IsValid() == false)
+			{
+				LODSubsystem->RegisterActorViewer(*DebugActor);
+			}
+			else
+			{
+				LODSubsystem->UnregisterActorViewer(*DebugActor);
+			}
+		}
+	}
+}
+
+void FGameplayDebuggerCategory_Mass::OnClearActorViewers()
+{
+	UWorld* World = GetWorldFromReplicator();
+	if (UMassLODSubsystem* LODSubsystem = World->GetSubsystem<UMassLODSubsystem>())
+	{
+		LODSubsystem->DebugUnregisterActorViewer();
+	}
+}
+
 //-----------------------------------------------------------------------------
 // DEPRECATED
 //-----------------------------------------------------------------------------
diff --git a/Engine/Plugins/AI/MassAI/Source/MassAIDebug/Public/GameplayDebuggerCategory_Mass.h b/Engine/Plugins/AI/MassAI/Source/MassAIDebug/Public/GameplayDebuggerCategory_Mass.h
--- a/Engine/Plugins/AI/MassAI/Source/MassAIDebug/Public/GameplayDebuggerCategory_Mass.h
+++ b/Engine/Plugins/AI/MassAI/Source/MassAIDebug/Public/GameplayDebuggerCategory_Mass.h
@@ -41,6 +41,9 @@ protected:
 	void OnToggleDebugLocalEntityManager();
 	void OnIncreaseSearchRange();
 	void OnDecreaseSearchRange();
+	void OnTogglePickedActorAsViewer();
+	void OnToggleDrawViewers() { bShowViewers = !bShowViewers; }
+	void OnClearActorViewers();
 	
 	void PickEntity(const FVector& ViewLocation, const FVector& ViewDirection, const UWorld& World, FMassEntityManager& EntityManager, const bool bLimitAngle = true);
 
@@ -51,7 +54,7 @@ protected:
 	void ClearCachedEntity();
 
 protected:
-	AActor* CachedDebugActor;
+	TWeakObjectPtr<AActor> CachedDebugActor;
 	FMassEntityHandle CachedEntity;
 	bool bShowArchetypes;
 	bool bShowShapes;
@@ -63,7 +66,11 @@ protected:
 	bool bShowNearEntityPath;
 	bool bMarkEntityBeingDebugged;
 	bool bDebugLocalEntityManager;
+	bool bShowViewers;
 	int32 ToggleDebugLocalEntityManagerInputIndex = INDEX_NONE;
+	int32 TogglePickedActorAsViewerInputIndex = INDEX_NONE;
+	int32 ToggleDrawViewersInputIndex = INDEX_NONE;
+	int32 ClearViewersInputIndex = INDEX_NONE;
 	float SearchRange = 25000.f;
 
 	struct FEntityDescription
@@ -88,6 +95,9 @@ protected:
 	static FOnConsoleCommandBroadcastDelegate OnToggleNearEntityAvoidanceBroadcast;
 	static FOnConsoleCommandBroadcastDelegate OnToggleNearEntityPathBroadcast; 	
 	static FOnConsoleCommandBroadcastDelegate OnToggleDebugLocalEntityManagerBroadcast; 	
+	static FOnConsoleCommandBroadcastDelegate OnTogglePickedActorAsViewerBroadcast;
+	static FOnConsoleCommandBroadcastDelegate OnToggleDrawViewersBroadcast;
+	static FOnConsoleCommandBroadcastDelegate OnClearActorViewersBroadcast;
 
 	using FDelegateHandlePair = TPair<FOnConsoleCommandBroadcastDelegate*, FDelegateHandle>;
 	TArray<FDelegateHandlePair> ConsoleCommandHandles;
diff --git a/Engine/Plugins/Runtime/MassGameplay/Source/MassLOD/Private/MassLODSubsystem.cpp b/Engine/Plugins/Runtime/MassGameplay/Source/MassLOD/Private/MassLODSubsystem.cpp
--- a/Engine/Plugins/Runtime/MassGameplay/Source/MassLOD/Private/MassLODSubsystem.cpp
+++ b/Engine/Plugins/Runtime/MassGameplay/Source/MassLOD/Private/MassLODSubsystem.cpp
@@ -601,6 +601,24 @@ FMassViewerHandle UMassLODSubsystem::GetViewerHandleFromPlayerController(const A
 	PRAGMA_ENABLE_DEPRECATION_WARNINGS
 }
 
+#if WITH_MASSGAMEPLAY_DEBUG
+void UMassLODSubsystem::DebugUnregisterActorViewer()
+{
+	for (const TObjectPtr<AActor>& RegisteredActorViewer : RegisteredActorViewers)
+	{
+		if (const AActor* ActorViewer = RegisteredActorViewer)
+		{
+			const FMassViewerHandle ViewerHandle = GetViewerHandleFromActor(*ActorViewer);
+			if (ensure(ViewerHandle.IsValid()))
+			{
+				RemoveViewer(ViewerHandle);
+			}
+		}
+	}
+	RegisteredActorViewers.Reset();
+}
+#endif // WITH_MASSGAMEPLAY_DEBUG
+
 //-----------------------------------------------------------------------------
 // FViewerInfo
 //-----------------------------------------------------------------------------
diff --git a/Engine/Plugins/Runtime/MassGameplay/Source/MassLOD/Public/MassLODSubsystem.h b/Engine/Plugins/Runtime/MassGameplay/Source/MassLOD/Public/MassLODSubsystem.h
--- a/Engine/Plugins/Runtime/MassGameplay/Source/MassLOD/Public/MassLODSubsystem.h
+++ b/Engine/Plugins/Runtime/MassGameplay/Source/MassLOD/Public/MassLODSubsystem.h
@@ -113,6 +113,7 @@ public:
 
 #if WITH_MASSGAMEPLAY_DEBUG
 	void DebugSetUsePlayerPawnLocationInsteadOfCamera(const bool bInValue) { bUsePlayerPawnLocationInsteadOfCamera = bInValue; }
+	void DebugUnregisterActorViewer();
 #endif
 
 protected:
