Updated exeuction flags of a bunch of Mass processors to go along with the "make processors work for Editor" change

diff --git a/Engine/Plugins/AI/MassAI/Source/MassAIBehavior/Private/MassZoneGraphAnnotationProcessors.cpp b/Engine/Plugins/AI/MassAI/Source/MassAIBehavior/Private/MassZoneGraphAnnotationProcessors.cpp
--- a/Engine/Plugins/AI/MassAI/Source/MassAIBehavior/Private/MassZoneGraphAnnotationProcessors.cpp
+++ b/Engine/Plugins/AI/MassAI/Source/MassAIBehavior/Private/MassZoneGraphAnnotationProcessors.cpp
@@ -72,7 +72,10 @@ void UMassZoneGraphAnnotationTagUpdateProcessor::Initialize(UObject& Owner)
 	Super::Initialize(Owner);
 
 	UMassSignalSubsystem* SignalSubsystem = UWorld::GetSubsystem<UMassSignalSubsystem>(Owner.GetWorld());
-	SubscribeToSignal(*SignalSubsystem, UE::Mass::Signals::CurrentLaneChanged);
+	if (ensureMsgf(SignalSubsystem, TEXT("Expecting to run this processor with UMassSignalSubsystem available")))
+	{
+		SubscribeToSignal(*SignalSubsystem, UE::Mass::Signals::CurrentLaneChanged);
+	}
 }
 
 void UMassZoneGraphAnnotationTagUpdateProcessor::ConfigureQueries()
diff --git a/Engine/Plugins/AI/MassAI/Source/MassNavigation/Private/MassAvoidanceProcessors.cpp b/Engine/Plugins/AI/MassAI/Source/MassNavigation/Private/MassAvoidanceProcessors.cpp
--- a/Engine/Plugins/AI/MassAI/Source/MassNavigation/Private/MassAvoidanceProcessors.cpp
+++ b/Engine/Plugins/AI/MassAI/Source/MassNavigation/Private/MassAvoidanceProcessors.cpp
@@ -351,7 +351,7 @@ UMassMovingAvoidanceProcessor::UMassMovingAvoidanceProcessor()
 	: EntityQuery(*this) 
 {
 	bAutoRegisterWithProcessingPhases = true;
-	ExecutionFlags = (int32)EProcessorExecutionFlags::All;
+	ExecutionFlags = (int32)EProcessorExecutionFlags::AllNetModes;
 	ExecutionOrder.ExecuteInGroup = UE::Mass::ProcessorGroupNames::Avoidance;
 	ExecutionOrder.ExecuteAfter.Add(UE::Mass::ProcessorGroupNames::LOD);
 }
@@ -1016,7 +1016,7 @@ UMassStandingAvoidanceProcessor::UMassStandingAvoidanceProcessor()
 
 {
 	bAutoRegisterWithProcessingPhases = true;
-	ExecutionFlags = (int32)EProcessorExecutionFlags::All;
+	ExecutionFlags = (int32)EProcessorExecutionFlags::AllNetModes;
 	ExecutionOrder.ExecuteInGroup = UE::Mass::ProcessorGroupNames::Avoidance;
 	ExecutionOrder.ExecuteAfter.Add(TEXT("MassMovingAvoidanceProcessor"));
 }
diff --git a/Engine/Plugins/AI/MassAI/Source/MassNavigation/Private/MassNavigationProcessors.cpp b/Engine/Plugins/AI/MassAI/Source/MassNavigation/Private/MassNavigationProcessors.cpp
--- a/Engine/Plugins/AI/MassAI/Source/MassNavigation/Private/MassNavigationProcessors.cpp
+++ b/Engine/Plugins/AI/MassAI/Source/MassNavigation/Private/MassNavigationProcessors.cpp
@@ -73,7 +73,7 @@ void UMassOffLODNavigationProcessor::Execute(FMassEntityManager& EntityManager,
 UMassNavigationSmoothHeightProcessor::UMassNavigationSmoothHeightProcessor()
 	: EntityQuery(*this)
 {
-	ExecutionFlags = (int32)EProcessorExecutionFlags::All;
+	ExecutionFlags = (int32)EProcessorExecutionFlags::AllNetModes;
 	ExecutionOrder.ExecuteAfter.Add(UE::Mass::ProcessorGroupNames::Movement);
 }
 
@@ -163,7 +163,7 @@ void UMassMoveTargetFragmentInitializer::Execute(FMassEntityManager& EntityManag
 //----------------------------------------------------------------------//
 UMassNavigationObstacleGridProcessor::UMassNavigationObstacleGridProcessor()
 {
-	ExecutionFlags = (int32)EProcessorExecutionFlags::All;
+	ExecutionFlags = (int32)EProcessorExecutionFlags::AllNetModes;
 	ExecutionOrder.ExecuteAfter.Add(UE::Mass::ProcessorGroupNames::Movement);
 }
 
@@ -283,7 +283,7 @@ UMassNavigationObstacleRemoverProcessor::UMassNavigationObstacleRemoverProcessor
 {
 	ObservedType = FMassNavigationObstacleGridCellLocationFragment::StaticStruct();
 	Operation = EMassObservedOperation::Remove;
-	ExecutionFlags = (int32)(EProcessorExecutionFlags::All);
+	ExecutionFlags = (int32)EProcessorExecutionFlags::AllNetModes;
 }
 
 void UMassNavigationObstacleRemoverProcessor::ConfigureQueries()
diff --git a/Engine/Plugins/AI/MassAI/Source/MassNavigation/Private/MassSmoothOrientationProcessors.cpp b/Engine/Plugins/AI/MassAI/Source/MassNavigation/Private/MassSmoothOrientationProcessors.cpp
--- a/Engine/Plugins/AI/MassAI/Source/MassNavigation/Private/MassSmoothOrientationProcessors.cpp
+++ b/Engine/Plugins/AI/MassAI/Source/MassNavigation/Private/MassSmoothOrientationProcessors.cpp
@@ -21,7 +21,7 @@ UMassSmoothOrientationProcessor::UMassSmoothOrientationProcessor()
 	: HighResEntityQuery(*this)
 	, LowResEntityQuery_Conditional(*this)
 {
-	ExecutionFlags = (int32)EProcessorExecutionFlags::All;
+	ExecutionFlags = (int32)EProcessorExecutionFlags::AllNetModes;
 	ExecutionOrder.ExecuteInGroup = UE::Mass::ProcessorGroupNames::Movement;
 }
 
diff --git a/Engine/Plugins/AI/MassAI/Source/MassNavigation/Private/MassSteeringProcessors.cpp b/Engine/Plugins/AI/MassAI/Source/MassNavigation/Private/MassSteeringProcessors.cpp
--- a/Engine/Plugins/AI/MassAI/Source/MassNavigation/Private/MassSteeringProcessors.cpp
+++ b/Engine/Plugins/AI/MassAI/Source/MassNavigation/Private/MassSteeringProcessors.cpp
@@ -54,7 +54,7 @@ namespace UE::MassNavigation
 UMassSteerToMoveTargetProcessor::UMassSteerToMoveTargetProcessor()
 	: EntityQuery(*this)
 {
-	ExecutionFlags = int32(EProcessorExecutionFlags::All);
+	ExecutionFlags = int32(EProcessorExecutionFlags::AllNetModes);
 	ExecutionOrder.ExecuteAfter.Add(UE::Mass::ProcessorGroupNames::Tasks);
 	ExecutionOrder.ExecuteBefore.Add(UE::Mass::ProcessorGroupNames::Avoidance);
 }
diff --git a/Engine/Plugins/AI/MassAI/Source/MassZoneGraphNavigation/Private/MassZoneGraphNavigationProcessors.cpp b/Engine/Plugins/AI/MassAI/Source/MassZoneGraphNavigation/Private/MassZoneGraphNavigationProcessors.cpp
--- a/Engine/Plugins/AI/MassAI/Source/MassZoneGraphNavigation/Private/MassZoneGraphNavigationProcessors.cpp
+++ b/Engine/Plugins/AI/MassAI/Source/MassZoneGraphNavigation/Private/MassZoneGraphNavigationProcessors.cpp
@@ -43,7 +43,7 @@ UMassZoneGraphLocationInitializer::UMassZoneGraphLocationInitializer()
 {
 	ObservedType = FMassZoneGraphLaneLocationFragment::StaticStruct();
 	Operation = EMassObservedOperation::Add;
-	ExecutionFlags = (int32)(EProcessorExecutionFlags::All);
+	ExecutionFlags = (int32)EProcessorExecutionFlags::All;
 }
 
 void UMassZoneGraphLocationInitializer::ConfigureQueries()
@@ -116,7 +116,7 @@ void UMassZoneGraphLocationInitializer::Execute(FMassEntityManager& EntityManage
 UMassZoneGraphPathFollowProcessor::UMassZoneGraphPathFollowProcessor()
 	: EntityQuery_Conditional(*this)
 {
-	ExecutionFlags = (int32)EProcessorExecutionFlags::All;
+	ExecutionFlags = (int32)EProcessorExecutionFlags::AllNetModes;
 	ExecutionOrder.ExecuteInGroup = UE::Mass::ProcessorGroupNames::Tasks;
 	ExecutionOrder.ExecuteBefore.Add(UE::Mass::ProcessorGroupNames::Avoidance);
 }
@@ -416,7 +416,7 @@ void UMassZoneGraphPathFollowProcessor::Execute(FMassEntityManager& EntityManage
 UMassZoneGraphLaneCacheBoundaryProcessor::UMassZoneGraphLaneCacheBoundaryProcessor()
 	: EntityQuery(*this)
 {
-	ExecutionFlags = (int32)EProcessorExecutionFlags::All;
+	ExecutionFlags = (int32)EProcessorExecutionFlags::AllNetModes;
 
 	bAutoRegisterWithProcessingPhases = true;
 	ExecutionOrder.ExecuteAfter.Add(UE::Mass::ProcessorGroupNames::LOD);
diff --git a/Engine/Plugins/Runtime/MassGameplay/Source/MassLOD/Private/MassSimulationLOD.cpp b/Engine/Plugins/Runtime/MassGameplay/Source/MassLOD/Private/MassSimulationLOD.cpp
--- a/Engine/Plugins/Runtime/MassGameplay/Source/MassLOD/Private/MassSimulationLOD.cpp
+++ b/Engine/Plugins/Runtime/MassGameplay/Source/MassLOD/Private/MassSimulationLOD.cpp
@@ -69,7 +69,7 @@ UMassSimulationLODProcessor::UMassSimulationLODProcessor()
 	, EntityQueryVariableTick(*this)
 	, EntityQuerySetLODTag(*this)
 {
-	ExecutionFlags = (int32)EProcessorExecutionFlags::All;
+	ExecutionFlags = (int32)EProcessorExecutionFlags::AllNetModes;
 	ExecutionOrder.ExecuteInGroup = UE::Mass::ProcessorGroupNames::LOD;
 	ExecutionOrder.ExecuteAfter.Add(UE::Mass::ProcessorGroupNames::LODCollector);
 }
diff --git a/Engine/Plugins/Runtime/MassGameplay/Source/MassMovement/Private/MassMovementProcessors.cpp b/Engine/Plugins/Runtime/MassGameplay/Source/MassMovement/Private/MassMovementProcessors.cpp
--- a/Engine/Plugins/Runtime/MassGameplay/Source/MassMovement/Private/MassMovementProcessors.cpp
+++ b/Engine/Plugins/Runtime/MassGameplay/Source/MassMovement/Private/MassMovementProcessors.cpp
@@ -19,7 +19,7 @@
 UMassApplyMovementProcessor::UMassApplyMovementProcessor()
 	: EntityQuery(*this)
 {
-	ExecutionFlags = (int32)EProcessorExecutionFlags::All;
+	ExecutionFlags = (int32)EProcessorExecutionFlags::AllNetModes;
 	ExecutionOrder.ExecuteInGroup = UE::Mass::ProcessorGroupNames::Movement;
 	ExecutionOrder.ExecuteAfter.Add(UE::Mass::ProcessorGroupNames::Avoidance);
 }
diff --git a/Engine/Plugins/Runtime/MassGameplay/Source/MassMovement/Private/MassSimpleMovementTrait.cpp b/Engine/Plugins/Runtime/MassGameplay/Source/MassMovement/Private/MassSimpleMovementTrait.cpp
--- a/Engine/Plugins/Runtime/MassGameplay/Source/MassMovement/Private/MassSimpleMovementTrait.cpp
+++ b/Engine/Plugins/Runtime/MassGameplay/Source/MassMovement/Private/MassSimpleMovementTrait.cpp
@@ -26,7 +26,7 @@ UMassSimpleMovementProcessor::UMassSimpleMovementProcessor()
 	: EntityQuery(*this)
 {
 	bAutoRegisterWithProcessingPhases = true;
-	ExecutionFlags = (int32)EProcessorExecutionFlags::All;
+	ExecutionFlags = (int32)EProcessorExecutionFlags::AllNetModes;
 	ExecutionOrder.ExecuteInGroup = UE::Mass::ProcessorGroupNames::Avoidance;
 }
 
diff --git a/Engine/Plugins/Runtime/MassGameplay/Source/MassReplication/Private/MassReplicationGridProcessor.cpp b/Engine/Plugins/Runtime/MassGameplay/Source/MassReplication/Private/MassReplicationGridProcessor.cpp
--- a/Engine/Plugins/Runtime/MassGameplay/Source/MassReplication/Private/MassReplicationGridProcessor.cpp
+++ b/Engine/Plugins/Runtime/MassGameplay/Source/MassReplication/Private/MassReplicationGridProcessor.cpp
@@ -19,7 +19,7 @@ UMassReplicationGridProcessor::UMassReplicationGridProcessor()
 #if !UE_ALLOW_DEBUG_REPLICATION_BUBBLES_STANDALONE
 	ExecutionFlags = int32(EProcessorExecutionFlags::Server);
 #else
-	ExecutionFlags = int32(EProcessorExecutionFlags::All);
+	ExecutionFlags = int32(EProcessorExecutionFlags::AllNetModes);
 #endif // UE_ALLOW_DEBUG_REPLICATION_BUBBLES_STANDALONE
 
 	ProcessingPhase = EMassProcessingPhase::PostPhysics;
@@ -136,7 +136,7 @@ UMassReplicationGridRemoverProcessor::UMassReplicationGridRemoverProcessor()
 {
 	ObservedType = FMassReplicationGridCellLocationFragment::StaticStruct();
 	Operation = EMassObservedOperation::Remove;
-	ExecutionFlags = (int32)(EProcessorExecutionFlags::All);
+	ExecutionFlags = (int32)EProcessorExecutionFlags::AllNetModes;
 }
 
 void UMassReplicationGridRemoverProcessor::ConfigureQueries()
diff --git a/Engine/Plugins/Runtime/MassGameplay/Source/MassReplication/Private/MassReplicationProcessor.cpp b/Engine/Plugins/Runtime/MassGameplay/Source/MassReplication/Private/MassReplicationProcessor.cpp
--- a/Engine/Plugins/Runtime/MassGameplay/Source/MassReplication/Private/MassReplicationProcessor.cpp
+++ b/Engine/Plugins/Runtime/MassGameplay/Source/MassReplication/Private/MassReplicationProcessor.cpp
@@ -25,7 +25,7 @@ UMassReplicationProcessor::UMassReplicationProcessor()
 #if !UE_ALLOW_DEBUG_REPLICATION_BUBBLES_STANDALONE
 	ExecutionFlags = int32(EProcessorExecutionFlags::Server);
 #else
-	ExecutionFlags = int32(EProcessorExecutionFlags::All);
+	ExecutionFlags = int32(EProcessorExecutionFlags::AllNetModes);
 #endif // UE_ALLOW_DEBUG_REPLICATION_BUBBLES_STANDALONE
 
 	ProcessingPhase = EMassProcessingPhase::PostPhysics;
diff --git a/Engine/Plugins/Runtime/MassGameplay/Source/MassRepresentation/Private/MassRepresentationProcessor.cpp b/Engine/Plugins/Runtime/MassGameplay/Source/MassRepresentation/Private/MassRepresentationProcessor.cpp
--- a/Engine/Plugins/Runtime/MassGameplay/Source/MassRepresentation/Private/MassRepresentationProcessor.cpp
+++ b/Engine/Plugins/Runtime/MassGameplay/Source/MassRepresentation/Private/MassRepresentationProcessor.cpp
@@ -421,7 +421,7 @@ UMassRepresentationFragmentDestructor::UMassRepresentationFragmentDestructor()
 {
 	ObservedType = FMassRepresentationFragment::StaticStruct();
 	Operation = EMassObservedOperation::Remove;
-	ExecutionFlags = (int32)(EProcessorExecutionFlags::All);
+	ExecutionFlags = (int32)EProcessorExecutionFlags::All;
 	bRequiresGameThreadExecution = true; // due to FMassRepresentationSubsystemSharedFragment.RepresentationSubsystem use
 }
 
diff --git a/Engine/Plugins/Runtime/MassGameplay/Source/MassRepresentation/Private/MassStationaryISMRepresentationFragmentDestructor.cpp b/Engine/Plugins/Runtime/MassGameplay/Source/MassRepresentation/Private/MassStationaryISMRepresentationFragmentDestructor.cpp
--- a/Engine/Plugins/Runtime/MassGameplay/Source/MassRepresentation/Private/MassStationaryISMRepresentationFragmentDestructor.cpp
+++ b/Engine/Plugins/Runtime/MassGameplay/Source/MassRepresentation/Private/MassStationaryISMRepresentationFragmentDestructor.cpp
@@ -13,7 +13,7 @@ UMassStationaryISMRepresentationFragmentDestructor::UMassStationaryISMRepresenta
 {
 	ObservedType = FMassRepresentationFragment::StaticStruct();
 	Operation = EMassObservedOperation::Remove;
-	ExecutionFlags = (int32)(EProcessorExecutionFlags::All);
+	ExecutionFlags = (int32)EProcessorExecutionFlags::All;
 	bRequiresGameThreadExecution = true; // not sure about this
 }
 
diff --git a/Engine/Plugins/Runtime/MassGameplay/Source/MassSignals/Private/MassSignalProcessorBase.cpp b/Engine/Plugins/Runtime/MassGameplay/Source/MassSignals/Private/MassSignalProcessorBase.cpp
--- a/Engine/Plugins/Runtime/MassGameplay/Source/MassSignals/Private/MassSignalProcessorBase.cpp
+++ b/Engine/Plugins/Runtime/MassGameplay/Source/MassSignals/Private/MassSignalProcessorBase.cpp
@@ -11,6 +11,7 @@
 UMassSignalProcessorBase::UMassSignalProcessorBase(const FObjectInitializer& ObjectInitializer)
 	: EntityQuery(*this)
 {
+	ExecutionFlags = (int32)EProcessorExecutionFlags::AllNetModes;
 }
 
 void UMassSignalProcessorBase::BeginDestroy()
diff --git a/Engine/Plugins/Runtime/MassGameplay/Source/MassSmartObjects/Private/MassSmartObjectProcessor.cpp b/Engine/Plugins/Runtime/MassGameplay/Source/MassSmartObjects/Private/MassSmartObjectProcessor.cpp
--- a/Engine/Plugins/Runtime/MassGameplay/Source/MassSmartObjects/Private/MassSmartObjectProcessor.cpp
+++ b/Engine/Plugins/Runtime/MassGameplay/Source/MassSmartObjects/Private/MassSmartObjectProcessor.cpp
@@ -389,7 +389,7 @@ UMassSmartObjectUserFragmentDeinitializer::UMassSmartObjectUserFragmentDeinitial
 {
 	ObservedType = FMassSmartObjectUserFragment::StaticStruct();
 	Operation = EMassObservedOperation::Remove;
-	ExecutionFlags = (int32)(EProcessorExecutionFlags::All);
+	ExecutionFlags = (int32)EProcessorExecutionFlags::All;
 }
 
 void UMassSmartObjectUserFragmentDeinitializer::ConfigureQueries()
