Removed unreachable code

#rnx

diff --git a/Engine/Plugins/AI/MassAI/Source/MassAIBehavior/Private/Conditions/ZoneGraphTagConditions.cpp b/Engine/Plugins/AI/MassAI/Source/MassAIBehavior/Private/Conditions/ZoneGraphTagConditions.cpp
--- a/Engine/Plugins/AI/MassAI/Source/MassAIBehavior/Private/Conditions/ZoneGraphTagConditions.cpp
+++ b/Engine/Plugins/AI/MassAI/Source/MassAIBehavior/Private/Conditions/ZoneGraphTagConditions.cpp
@@ -70,9 +70,9 @@ namespace UE::MassBehavior::ZoneGraph
 		case EZoneLaneTagMaskComparison::Not:
 			return LOCTEXT("ContainsNot", "Not");
 		default:
-			return FText::FromString(TEXT("??"));
+			break;
 		}
-		return FText::GetEmpty();
+		return FText::FromString(TEXT("??"));
 	}
 
 }
