Adding MassEntityComponent to SceneGraph as a mean to control definition and lifetime of an associated mass entity
#okforversepublic

diff --git a/Engine/Plugins/AI/MassAI/Source/MassNavigation/Public/MassNavigationFragments.h b/Engine/Plugins/AI/MassAI/Source/MassNavigation/Public/MassNavigationFragments.h
--- a/Engine/Plugins/AI/MassAI/Source/MassNavigation/Public/MassNavigationFragments.h
+++ b/Engine/Plugins/AI/MassAI/Source/MassNavigation/Public/MassNavigationFragments.h
@@ -2,10 +2,15 @@
 
 #pragma once
 
+#if UE_ENABLE_INCLUDE_ORDER_DEPRECATED_IN_5_5
 #include "CoreMinimal.h"
+#endif
+
+#include "AI/Navigation/NavigationElement.h"
+#include "MassCommonFragments.h"
 #include "MassCommonTypes.h"
-#include "MassNavigationTypes.h"
 #include "MassNavigationSubsystem.h"
+#include "MassNavigationTypes.h"
 #include "MassNavigationFragments.generated.h"
 
 class UWorld;
@@ -177,9 +182,42 @@ struct MASSNAVIGATION_API FMassAvoidanceColliderFragment : public FMassFragment
 };
 
 
-/** Component Tag to tell if the entity is in the navigation obstacle grid */
+/** Tag to tell if the entity is in the navigation obstacle grid */
 USTRUCT()
 struct MASSNAVIGATION_API FMassInNavigationObstacleGridTag : public FMassTag
 {
 	GENERATED_BODY()
 };
+
+/**
+ * Shared Fragment holding the properties defining how a given entity should affect navigation data
+ */
+USTRUCT()
+struct FNavigationRelevantParameters : public FMassConstSharedFragment
+{
+	GENERATED_BODY()
+
+	/** If set, navmesh will not be generated under the surface of the geometry */
+	UPROPERTY()
+	bool bFillCollisionUnderneathForNavData = false;
+};
+
+/**
+ * Fragment holding the registration handle to the navigation element created from a Mass entity.
+ * The fragment is added to indicate that a Mass entity is relevant to the AI navigation system.
+ */
+USTRUCT()
+struct FNavigationRelevantFragment : public FMassFragment
+{
+	GENERATED_BODY()
+
+	FNavigationRelevantFragment() = default;
+	explicit FNavigationRelevantFragment(const FNavigationElementHandle Handle)
+		: Handle(Handle)
+	{
+	}
+
+	/** Handle to the Navigation element created and registered for the entity */
+	FNavigationElementHandle Handle;
+};
+
