[MassStateTree] added option for statetree processor to run on multiple threads and use ParallelForEachEntityChunk (instead of ForEachEntityChunk)

diff --git a/Engine/Plugins/AI/MassAI/Source/MassAIBehavior/Private/MassStateTreeProcessors.cpp b/Engine/Plugins/AI/MassAI/Source/MassAIBehavior/Private/MassStateTreeProcessors.cpp
--- a/Engine/Plugins/AI/MassAI/Source/MassAIBehavior/Private/MassStateTreeProcessors.cpp
+++ b/Engine/Plugins/AI/MassAI/Source/MassAIBehavior/Private/MassStateTreeProcessors.cpp
@@ -225,7 +225,7 @@ void UMassStateTreeActivationProcessor::Execute(FMassEntityManager& EntityManage
 UMassStateTreeProcessor::UMassStateTreeProcessor(const FObjectInitializer& ObjectInitializer)
 	: Super(ObjectInitializer)
 {
-	bRequiresGameThreadExecution = true;
+	bRequiresGameThreadExecution = bProcessEntitiesInParallel;
 
 	ExecutionOrder.ExecuteInGroup = UE::Mass::ProcessorGroupNames::Behavior;
 
@@ -286,15 +286,14 @@ void UMassStateTreeProcessor::SignalEntities(FMassEntityManager& EntityManager,
 
 	TArray<FMassEntityHandle> EntitiesToSignal;
 
-	EntityQuery.ForEachEntityChunk(EntityManager, Context,
-		[this, &TimeInSeconds, &EntitiesToSignal, &EntityManager, SignalSubsystem = &SignalSubsystem, &EntitySignals](FMassExecutionContext& Context)
+	auto TickFunc = [&](FMassExecutionContext& Context)
 		{
 			// Keep stats regarding the amount of tree instances ticked per frame
 			CSV_CUSTOM_STAT(StateTreeProcessor, NumTickedStateTree, Context.GetNumEntities(), ECsvCustomStatOp::Accumulate);
 
 			UMassStateTreeSubsystem& MassStateTreeSubsystem = Context.GetMutableSubsystemChecked<UMassStateTreeSubsystem>();
 
-			UE::MassBehavior::ForEachEntityInChunk(Context, EntityManager, *SignalSubsystem, MassStateTreeSubsystem,
+			UE::MassBehavior::ForEachEntityInChunk(Context, EntityManager, SignalSubsystem, MassStateTreeSubsystem,
 				[TimeInSeconds, &EntitiesToSignal, &EntitySignals, &MassStateTreeSubsystem]
 				(FMassStateTreeExecutionContext& StateTreeExecutionContext, FMassStateTreeInstanceFragment& StateTreeFragment)
 				{
@@ -337,7 +336,16 @@ void UMassStateTreeProcessor::SignalEntities(FMassEntityManager& EntityManager,
 						}
 					}
 				});
-		});
+		};
+
+	if (bProcessEntitiesInParallel)
+	{
+		EntityQuery.ParallelForEachEntityChunk(EntityManager, Context, TickFunc);
+	}
+	else
+	{
+		EntityQuery.ForEachEntityChunk(EntityManager, Context, TickFunc);
+	}
 
 	if (EntitiesToSignal.Num())
 	{
diff --git a/Engine/Plugins/AI/MassAI/Source/MassAIBehavior/Public/MassStateTreeProcessors.h b/Engine/Plugins/AI/MassAI/Source/MassAIBehavior/Public/MassStateTreeProcessors.h
--- a/Engine/Plugins/AI/MassAI/Source/MassAIBehavior/Public/MassStateTreeProcessors.h
+++ b/Engine/Plugins/AI/MassAI/Source/MassAIBehavior/Public/MassStateTreeProcessors.h
@@ -70,4 +70,8 @@ protected:
 	virtual void Initialize(UObject& Owner) override;
 	virtual void ConfigureQueries() override;
 	virtual void SignalEntities(FMassEntityManager& EntityManager, FMassExecutionContext& Context, FMassSignalNameLookup& EntitySignals) override;
+
+	/** Configures whether parallel update for FMassArchetypeChunks should be used instead of the default single threaded update (i.e., ParallelForEachEntityChunk instead of ForEachEntityChunk). */
+	UPROPERTY(EditDefaultsOnly, Category = Processor, config)
+	bool bProcessEntitiesInParallel = false;
 };
