[MassSignals] added MT access detector for delayed signals and replaced calls to DelaySignalEntity by DelaySignalEntityDeferred to fix threading issues

diff --git a/Engine/Plugins/AI/MassAI/Source/MassAIBehavior/Private/MassStateTreeExecutionContext.cpp b/Engine/Plugins/AI/MassAI/Source/MassAIBehavior/Private/MassStateTreeExecutionContext.cpp
--- a/Engine/Plugins/AI/MassAI/Source/MassAIBehavior/Private/MassStateTreeExecutionContext.cpp
+++ b/Engine/Plugins/AI/MassAI/Source/MassAIBehavior/Private/MassStateTreeExecutionContext.cpp
@@ -15,12 +15,12 @@ bool CollectExternalData(const FStateTreeExecutionContext& Context, const UState
 {
 	CSV_SCOPED_TIMING_STAT_EXCLUSIVE(StateTreeProcessorCollectExternalData);
 	
-	const FMassStateTreeExecutionContext& MassContext = static_cast<const FMassStateTreeExecutionContext&>(Context); 
-	const FMassEntityManager& EntityManager = MassContext.GetEntityManager();
-	const UWorld* World = MassContext.GetWorld();
+	const FMassStateTreeExecutionContext& MassStateTreeContext = static_cast<const FMassStateTreeExecutionContext&>(Context); 
+	const FMassEntityManager& EntityManager = MassStateTreeContext.GetEntityManager();
+	const UWorld* World = MassStateTreeContext.GetWorld();
 	
 	bool bFoundAll = true;
-	const FMassEntityView EntityView(EntityManager, MassContext.GetEntity());
+	const FMassEntityView EntityView(EntityManager, MassStateTreeContext.GetEntity());
 
 	check(ExternalDataDescs.Num() == OutDataViews.Num());
 
@@ -120,7 +120,7 @@ FMassStateTreeExecutionContext::FMassStateTreeExecutionContext(UObject& InOwner,
 	: FStateTreeExecutionContext(InOwner, InStateTree, InInstanceData, FOnCollectStateTreeExternalData::CreateStatic(UE::MassBehavior::CollectExternalData))
 	, EntityManager(&InEntityManager)
 	, SignalSubsystem(&InSignalSubsystem)
-	, EntitySubsystemExecutionContext(&InContext)
+	, MassEntityExecutionContext(&InContext)
 {
 }
 
@@ -129,6 +129,6 @@ void FMassStateTreeExecutionContext::BeginDelayedTransition(const FStateTreeTran
 	if (SignalSubsystem != nullptr && Entity.IsSet())
 	{
 		// Tick again after the games time has passed to see if the condition still holds true.
-		SignalSubsystem->DelaySignalEntity(UE::Mass::Signals::DelayedTransitionWakeup, Entity, DelayedState.TimeLeft + KINDA_SMALL_NUMBER);
+		SignalSubsystem->DelaySignalEntityDeferred(GetMassEntityExecutionContext(), UE::Mass::Signals::DelayedTransitionWakeup, Entity, DelayedState.TimeLeft + KINDA_SMALL_NUMBER);
 	}
 }
diff --git a/Engine/Plugins/AI/MassAI/Source/MassAIBehavior/Private/Tasks/MassClaimSmartObjectTask.cpp b/Engine/Plugins/AI/MassAI/Source/MassAIBehavior/Private/Tasks/MassClaimSmartObjectTask.cpp
--- a/Engine/Plugins/AI/MassAI/Source/MassAIBehavior/Private/Tasks/MassClaimSmartObjectTask.cpp
+++ b/Engine/Plugins/AI/MassAI/Source/MassAIBehavior/Private/Tasks/MassClaimSmartObjectTask.cpp
@@ -51,10 +51,10 @@ EStateTreeRunStatus FMassClaimSmartObjectTask::EnterState(FStateTreeExecutionCon
 	InstanceData.ClaimedSlot.Invalidate();
 	
 	// Setup MassSmartObject handler and claim
-	const FMassStateTreeExecutionContext& MassContext = static_cast<FMassStateTreeExecutionContext&>(Context);
-	const FMassSmartObjectHandler MassSmartObjectHandler(MassContext.GetEntityManager(), MassContext.GetEntitySubsystemExecutionContext(), SmartObjectSubsystem, SignalSubsystem);
+	const FMassStateTreeExecutionContext& MassStateTreeContext = static_cast<FMassStateTreeExecutionContext&>(Context);
+	const FMassSmartObjectHandler MassSmartObjectHandler(MassStateTreeContext.GetEntityManager(), MassStateTreeContext.GetMassEntityExecutionContext(), SmartObjectSubsystem, SignalSubsystem);
 	
-	InstanceData.ClaimedSlot = MassSmartObjectHandler.ClaimCandidate(MassContext.GetEntity(), SOUser, *CandidateSlots);
+	InstanceData.ClaimedSlot = MassSmartObjectHandler.ClaimCandidate(MassStateTreeContext.GetEntity(), SOUser, *CandidateSlots);
 
 	// Treat claiming a slot as consuming all the candidate slots.
 	// This is done here because of the limited ways we can communicate between FindSmartObject() and ClaimSmartObject().
@@ -80,12 +80,12 @@ void FMassClaimSmartObjectTask::ExitState(FStateTreeExecutionContext& Context, c
 
 	if (InstanceData.ClaimedSlot.IsValid())
 	{
-		const FMassStateTreeExecutionContext& MassContext = static_cast<FMassStateTreeExecutionContext&>(Context);
+		const FMassStateTreeExecutionContext& MassStateTreeContext = static_cast<FMassStateTreeExecutionContext&>(Context);
 		USmartObjectSubsystem& SmartObjectSubsystem = Context.GetExternalData(SmartObjectSubsystemHandle);
 		UMassSignalSubsystem& SignalSubsystem = Context.GetExternalData(MassSignalSubsystemHandle);
-		const FMassSmartObjectHandler MassSmartObjectHandler(MassContext.GetEntityManager(), MassContext.GetEntitySubsystemExecutionContext(), SmartObjectSubsystem, SignalSubsystem);
+		const FMassSmartObjectHandler MassSmartObjectHandler(MassStateTreeContext.GetEntityManager(), MassStateTreeContext.GetMassEntityExecutionContext(), SmartObjectSubsystem, SignalSubsystem);
 
-		MassSmartObjectHandler.ReleaseSmartObject(MassContext.GetEntity(), SOUser, InstanceData.ClaimedSlot);
+		MassSmartObjectHandler.ReleaseSmartObject(MassStateTreeContext.GetEntity(), SOUser, InstanceData.ClaimedSlot);
 	}
 	else
 	{
diff --git a/Engine/Plugins/AI/MassAI/Source/MassAIBehavior/Private/Tasks/MassFindSmartObjectTask.cpp b/Engine/Plugins/AI/MassAI/Source/MassAIBehavior/Private/Tasks/MassFindSmartObjectTask.cpp
--- a/Engine/Plugins/AI/MassAI/Source/MassAIBehavior/Private/Tasks/MassFindSmartObjectTask.cpp
+++ b/Engine/Plugins/AI/MassAI/Source/MassAIBehavior/Private/Tasks/MassFindSmartObjectTask.cpp
@@ -37,12 +37,12 @@ void FMassFindSmartObjectTask::ExitState(FStateTreeExecutionContext& Context, co
 	// Stop any request that are still in flight.
 	if (InstanceData.SearchRequestID.IsSet())
 	{
-		const FMassStateTreeExecutionContext& MassContext = static_cast<FMassStateTreeExecutionContext&>(Context);
+		const FMassStateTreeExecutionContext& MassStateTreeContext = static_cast<FMassStateTreeExecutionContext&>(Context);
 		USmartObjectSubsystem& SmartObjectSubsystem = Context.GetExternalData(SmartObjectSubsystemHandle);
 		UMassSignalSubsystem& SignalSubsystem = Context.GetExternalData(MassSignalSubsystemHandle);
 		const FMassSmartObjectHandler MassSmartObjectHandler(
-			MassContext.GetEntityManager(),
-			MassContext.GetEntitySubsystemExecutionContext(),
+			MassStateTreeContext.GetEntityManager(),
+			MassStateTreeContext.GetMassEntityExecutionContext(),
 			SmartObjectSubsystem,
 			SignalSubsystem);
 		MassSmartObjectHandler.RemoveRequest(InstanceData.SearchRequestID);
@@ -80,10 +80,10 @@ EStateTreeRunStatus FMassFindSmartObjectTask::Tick(FStateTreeExecutionContext& C
 	const UWorld* World = Context.GetWorld();
 	USmartObjectSubsystem& SmartObjectSubsystem = Context.GetExternalData(SmartObjectSubsystemHandle);
 	UMassSignalSubsystem& SignalSubsystem = Context.GetExternalData(MassSignalSubsystemHandle);
-	const FMassStateTreeExecutionContext& MassContext = static_cast<FMassStateTreeExecutionContext&>(Context);
+	const FMassStateTreeExecutionContext& MassStateTreeContext = static_cast<FMassStateTreeExecutionContext&>(Context);
 	const FMassSmartObjectHandler MassSmartObjectHandler(
-		MassContext.GetEntityManager(),
-		MassContext.GetEntitySubsystemExecutionContext(),
+		MassStateTreeContext.GetEntityManager(),
+		MassStateTreeContext.GetMassEntityExecutionContext(),
 		SmartObjectSubsystem,
 		SignalSubsystem);
 
@@ -114,7 +114,7 @@ EStateTreeRunStatus FMassFindSmartObjectTask::Tick(FStateTreeExecutionContext& C
 		if (bTimeForNextUpdate || bLaneHasChanged)
 		{
 			// Use lanes if possible for faster queries using zone graph annotations
-			const FMassEntityHandle RequestingEntity = MassContext.GetEntity();
+			const FMassEntityHandle RequestingEntity = MassStateTreeContext.GetEntity();
 			if (LaneLocation != nullptr)
 			{
 				MASSBEHAVIOR_CLOG(!LaneLocation->LaneHandle.IsValid(), Error, TEXT("Always expecting a valid lane from the ZoneGraph movement"));
@@ -151,14 +151,14 @@ EStateTreeRunStatus FMassFindSmartObjectTask::Tick(FStateTreeExecutionContext& C
 			InstanceData.SearchRequestID.Reset();
 
 			// Schedule next update.
-			const FMassEntityHandle Entity = MassContext.GetEntity(); 
+			const FMassEntityHandle Entity = MassStateTreeContext.GetEntity();
 
 			constexpr float SearchIntervalDeviation = 0.1f;
 			const float DelayInSeconds = SearchInterval * FMath::FRandRange(1.0f - SearchIntervalDeviation, 1.0f + SearchIntervalDeviation);
 				
 			InstanceData.NextUpdate = World->GetTimeSeconds() + DelayInSeconds;
 			UMassSignalSubsystem& MassSignalSubsystem = Context.GetExternalData(MassSignalSubsystemHandle);
-			MassSignalSubsystem.DelaySignalEntity(UE::Mass::Signals::SmartObjectRequestCandidates, Entity, DelayInSeconds);
+			MassSignalSubsystem.DelaySignalEntityDeferred(MassStateTreeContext.GetMassEntityExecutionContext(), UE::Mass::Signals::SmartObjectRequestCandidates, Entity, DelayInSeconds);
 		}
 	}
 	
diff --git a/Engine/Plugins/AI/MassAI/Source/MassAIBehavior/Private/Tasks/MassLookAtTask.cpp b/Engine/Plugins/AI/MassAI/Source/MassAIBehavior/Private/Tasks/MassLookAtTask.cpp
--- a/Engine/Plugins/AI/MassAI/Source/MassAIBehavior/Private/Tasks/MassLookAtTask.cpp
+++ b/Engine/Plugins/AI/MassAI/Source/MassAIBehavior/Private/Tasks/MassLookAtTask.cpp
@@ -21,8 +21,8 @@ EStateTreeRunStatus FMassLookAtTask::EnterState(FStateTreeExecutionContext& Cont
 	
 	InstanceData.Time = 0.f;
 	
-	const FMassStateTreeExecutionContext& MassContext = static_cast<FMassStateTreeExecutionContext&>(Context);
-	FMassLookAtFragment& LookAtFragment = MassContext.GetExternalData(LookAtHandle);
+	const FMassStateTreeExecutionContext& MassStateTreeContext = static_cast<FMassStateTreeExecutionContext&>(Context);
+	FMassLookAtFragment& LookAtFragment = MassStateTreeContext.GetExternalData(LookAtHandle);
 
 	LookAtFragment.Reset();
 	LookAtFragment.LookAtMode = LookAtMode;
@@ -50,8 +50,8 @@ EStateTreeRunStatus FMassLookAtTask::EnterState(FStateTreeExecutionContext& Cont
 	// Otherwise we schedule a signal to end the task.
 	if (InstanceData.Duration > 0.0f)
 	{
-		UMassSignalSubsystem& MassSignalSubsystem = MassContext.GetExternalData(MassSignalSubsystemHandle);
-		MassSignalSubsystem.DelaySignalEntity(UE::Mass::Signals::LookAtFinished, MassContext.GetEntity(), InstanceData.Duration);
+		UMassSignalSubsystem& MassSignalSubsystem = MassStateTreeContext.GetExternalData(MassSignalSubsystemHandle);
+		MassSignalSubsystem.DelaySignalEntityDeferred(MassStateTreeContext.GetMassEntityExecutionContext(), UE::Mass::Signals::LookAtFinished, MassStateTreeContext.GetEntity(), InstanceData.Duration);
 	}
 
 	return EStateTreeRunStatus::Running;
@@ -59,8 +59,8 @@ EStateTreeRunStatus FMassLookAtTask::EnterState(FStateTreeExecutionContext& Cont
 
 void FMassLookAtTask::ExitState(FStateTreeExecutionContext& Context, const FStateTreeTransitionResult& Transition) const
 {
-	const FMassStateTreeExecutionContext& MassContext = static_cast<FMassStateTreeExecutionContext&>(Context);
-	FMassLookAtFragment& LookAtFragment = MassContext.GetExternalData(LookAtHandle);
+	const FMassStateTreeExecutionContext& MassStateTreeContext = static_cast<FMassStateTreeExecutionContext&>(Context);
+	FMassLookAtFragment& LookAtFragment = MassStateTreeContext.GetExternalData(LookAtHandle);
 	
 	LookAtFragment.Reset();
 }
diff --git a/Engine/Plugins/AI/MassAI/Source/MassAIBehavior/Private/Tasks/MassUseSmartObjectTask.cpp b/Engine/Plugins/AI/MassAI/Source/MassAIBehavior/Private/Tasks/MassUseSmartObjectTask.cpp
--- a/Engine/Plugins/AI/MassAI/Source/MassAIBehavior/Private/Tasks/MassUseSmartObjectTask.cpp
+++ b/Engine/Plugins/AI/MassAI/Source/MassAIBehavior/Private/Tasks/MassUseSmartObjectTask.cpp
@@ -51,10 +51,10 @@ EStateTreeRunStatus FMassUseSmartObjectTask::EnterState(FStateTreeExecutionConte
 	const FInstanceDataType& InstanceData = Context.GetInstanceData(*this);
 
 	// Setup MassSmartObject handler and start interaction
-	const FMassStateTreeExecutionContext& MassContext = static_cast<FMassStateTreeExecutionContext&>(Context);
-	const FMassSmartObjectHandler MassSmartObjectHandler(MassContext.GetEntityManager(), MassContext.GetEntitySubsystemExecutionContext(), SmartObjectSubsystem, SignalSubsystem);
+	const FMassStateTreeExecutionContext& MassStateTreeContext = static_cast<FMassStateTreeExecutionContext&>(Context);
+	const FMassSmartObjectHandler MassSmartObjectHandler(MassStateTreeContext.GetEntityManager(), MassStateTreeContext.GetMassEntityExecutionContext(), SmartObjectSubsystem, SignalSubsystem);
 
-	if (!MassSmartObjectHandler.StartUsingSmartObject(MassContext.GetEntity(), SOUser, InstanceData.ClaimedSlot))
+	if (!MassSmartObjectHandler.StartUsingSmartObject(MassStateTreeContext.GetEntity(), SOUser, InstanceData.ClaimedSlot))
 	{
 		return EStateTreeRunStatus::Failed;
 	}
@@ -64,7 +64,7 @@ EStateTreeRunStatus FMassUseSmartObjectTask::EnterState(FStateTreeExecutionConte
 	checkf(World != nullptr, TEXT("A valid world is expected from the execution context"));
 
 	MoveTarget.CreateNewAction(EMassMovementAction::Animate, *World);
-	const bool bSuccess = UE::MassNavigation::ActivateActionAnimate(*World, Context.GetOwner(), MassContext.GetEntity(), MoveTarget);
+	const bool bSuccess = UE::MassNavigation::ActivateActionAnimate(*World, Context.GetOwner(), MassStateTreeContext.GetEntity(), MoveTarget);
 
 	return bSuccess ? EStateTreeRunStatus::Running : EStateTreeRunStatus::Failed;
 }
@@ -77,11 +77,11 @@ void FMassUseSmartObjectTask::ExitState(FStateTreeExecutionContext& Context, con
 	{
 		MASSBEHAVIOR_LOG(VeryVerbose, TEXT("Exiting state with a valid InteractionHandle: stop using the smart object."));
 
-		const FMassStateTreeExecutionContext& MassContext = static_cast<FMassStateTreeExecutionContext&>(Context);
+		const FMassStateTreeExecutionContext& MassStateTreeContext = static_cast<FMassStateTreeExecutionContext&>(Context);
 		USmartObjectSubsystem& SmartObjectSubsystem = Context.GetExternalData(SmartObjectSubsystemHandle);
 		UMassSignalSubsystem& SignalSubsystem = Context.GetExternalData(MassSignalSubsystemHandle);
-		const FMassSmartObjectHandler MassSmartObjectHandler(MassContext.GetEntityManager(), MassContext.GetEntitySubsystemExecutionContext(), SmartObjectSubsystem, SignalSubsystem);
-		MassSmartObjectHandler.StopUsingSmartObject(MassContext.GetEntity(), SOUser, EMassSmartObjectInteractionStatus::Aborted);
+		const FMassSmartObjectHandler MassSmartObjectHandler(MassStateTreeContext.GetEntityManager(), MassStateTreeContext.GetMassEntityExecutionContext(), SmartObjectSubsystem, SignalSubsystem);
+		MassSmartObjectHandler.StopUsingSmartObject(MassStateTreeContext.GetEntity(), SOUser, EMassSmartObjectInteractionStatus::Aborted);
 	}
 	else
 	{
@@ -97,14 +97,14 @@ void FMassUseSmartObjectTask::StateCompleted(FStateTreeExecutionContext& Context
 	{
 		MASSBEHAVIOR_LOG(VeryVerbose, TEXT("Completing state with a valid InteractionHandle: stop using the smart object."));
 
-		const FMassStateTreeExecutionContext& MassContext = static_cast<FMassStateTreeExecutionContext&>(Context);
+		const FMassStateTreeExecutionContext& MassStateTreeContext = static_cast<FMassStateTreeExecutionContext&>(Context);
 		USmartObjectSubsystem& SmartObjectSubsystem = Context.GetExternalData(SmartObjectSubsystemHandle);
 		UMassSignalSubsystem& SignalSubsystem = Context.GetExternalData(MassSignalSubsystemHandle);
-		const FMassSmartObjectHandler MassSmartObjectHandler(MassContext.GetEntityManager(), MassContext.GetEntitySubsystemExecutionContext(), SmartObjectSubsystem, SignalSubsystem);
+		const FMassSmartObjectHandler MassSmartObjectHandler(MassStateTreeContext.GetEntityManager(), MassStateTreeContext.GetMassEntityExecutionContext(), SmartObjectSubsystem, SignalSubsystem);
 		const EMassSmartObjectInteractionStatus NewStatus = (CompletionStatus == EStateTreeRunStatus::Succeeded)
 																? EMassSmartObjectInteractionStatus::TaskCompleted
 																: EMassSmartObjectInteractionStatus::Aborted;
-		MassSmartObjectHandler.StopUsingSmartObject(MassContext.GetEntity(), SOUser, NewStatus);
+		MassSmartObjectHandler.StopUsingSmartObject(MassStateTreeContext.GetEntity(), SOUser, NewStatus);
 	}
 }
 
diff --git a/Engine/Plugins/AI/MassAI/Source/MassAIBehavior/Private/Tasks/MassZoneGraphStandTask.cpp b/Engine/Plugins/AI/MassAI/Source/MassAIBehavior/Private/Tasks/MassZoneGraphStandTask.cpp
--- a/Engine/Plugins/AI/MassAI/Source/MassAIBehavior/Private/Tasks/MassZoneGraphStandTask.cpp
+++ b/Engine/Plugins/AI/MassAI/Source/MassAIBehavior/Private/Tasks/MassZoneGraphStandTask.cpp
@@ -28,7 +28,7 @@ bool FMassZoneGraphStandTask::Link(FStateTreeLinker& Linker)
 
 EStateTreeRunStatus FMassZoneGraphStandTask::EnterState(FStateTreeExecutionContext& Context, const FStateTreeTransitionResult& Transition) const
 {
-	const FMassStateTreeExecutionContext& MassContext = static_cast<FMassStateTreeExecutionContext&>(Context);
+	const FMassStateTreeExecutionContext& MassStateTreeContext = static_cast<FMassStateTreeExecutionContext&>(Context);
 
 	const FMassZoneGraphLaneLocationFragment& LaneLocation = Context.GetExternalData(LocationHandle);
 	const UZoneGraphSubsystem& ZoneGraphSubsystem = Context.GetExternalData(ZoneGraphSubsystemHandle);
@@ -52,7 +52,7 @@ EStateTreeRunStatus FMassZoneGraphStandTask::EnterState(FStateTreeExecutionConte
 	checkf(World != nullptr, TEXT("A valid world is expected from the execution context"));
 
 	MoveTarget.CreateNewAction(EMassMovementAction::Stand, *World);
-	const bool bSuccess = UE::MassNavigation::ActivateActionStand(*World, Context.GetOwner(), MassContext.GetEntity(), ZoneGraphSubsystem, LaneLocation, MovementParams.DefaultDesiredSpeed, MoveTarget, ShortPath, CachedLane);
+	const bool bSuccess = UE::MassNavigation::ActivateActionStand(*World, Context.GetOwner(), MassStateTreeContext.GetEntity(), ZoneGraphSubsystem, LaneLocation, MovementParams.DefaultDesiredSpeed, MoveTarget, ShortPath, CachedLane);
 	if (!bSuccess)
 	{
 		return EStateTreeRunStatus::Failed;
@@ -65,7 +65,7 @@ EStateTreeRunStatus FMassZoneGraphStandTask::EnterState(FStateTreeExecutionConte
 	if (InstanceData.Duration > 0.0f)
 	{
 		UMassSignalSubsystem& MassSignalSubsystem = Context.GetExternalData(MassSignalSubsystemHandle);
-		MassSignalSubsystem.DelaySignalEntity(UE::Mass::Signals::StandTaskFinished, MassContext.GetEntity(), InstanceData.Duration);
+		MassSignalSubsystem.DelaySignalEntityDeferred(MassStateTreeContext.GetMassEntityExecutionContext(), UE::Mass::Signals::StandTaskFinished, MassStateTreeContext.GetEntity(), InstanceData.Duration);
 	}
 
 	return EStateTreeRunStatus::Running;
diff --git a/Engine/Plugins/AI/MassAI/Source/MassAIBehavior/Public/MassStateTreeExecutionContext.h b/Engine/Plugins/AI/MassAI/Source/MassAIBehavior/Public/MassStateTreeExecutionContext.h
--- a/Engine/Plugins/AI/MassAI/Source/MassAIBehavior/Public/MassStateTreeExecutionContext.h
+++ b/Engine/Plugins/AI/MassAI/Source/MassAIBehavior/Public/MassStateTreeExecutionContext.h
@@ -20,7 +20,7 @@ public:
 	FMassStateTreeExecutionContext(UObject& InOwner, const UStateTree& InStateTree, FStateTreeInstanceData& InInstanceData, FMassEntityManager& InEntityManager, UMassSignalSubsystem& InSignalSubsystem, FMassExecutionContext& InContext);
 
 	FMassEntityManager& GetEntityManager() const { check(EntityManager); return *EntityManager; }
-	FMassExecutionContext& GetEntitySubsystemExecutionContext() const { return *EntitySubsystemExecutionContext; }
+	FMassExecutionContext& GetMassEntityExecutionContext() const { return *MassEntityExecutionContext; }
 
 	FMassEntityHandle GetEntity() const { return Entity; }
 	void SetEntity(const FMassEntityHandle InEntity) { Entity = InEntity; }
@@ -34,6 +34,6 @@ protected:
 
 	FMassEntityManager* EntityManager;
 	UMassSignalSubsystem* SignalSubsystem = nullptr;
-	FMassExecutionContext* EntitySubsystemExecutionContext = nullptr;
+	FMassExecutionContext* MassEntityExecutionContext = nullptr;
 	FMassEntityHandle Entity;
 };
diff --git a/Engine/Plugins/Runtime/MassGameplay/Source/MassSignals/Private/MassSignalSubsystem.cpp b/Engine/Plugins/Runtime/MassGameplay/Source/MassSignals/Private/MassSignalSubsystem.cpp
--- a/Engine/Plugins/Runtime/MassGameplay/Source/MassSignals/Private/MassSignalSubsystem.cpp
+++ b/Engine/Plugins/Runtime/MassGameplay/Source/MassSignals/Private/MassSignalSubsystem.cpp
@@ -28,6 +28,8 @@ void UMassSignalSubsystem::Tick(float DeltaTime)
 	CA_ASSUME(CachedWorld);
 	const double CurrentTime = CachedWorld->GetTimeSeconds();
 
+	UE_MT_SCOPED_WRITE_ACCESS(DelayedSignalsAccessDetector);
+
 	for (int i = 0; i < DelayedSignals.Num();)
 	{
 		FDelayedSignal& DelayedSignal = DelayedSignals[i];
@@ -76,6 +78,13 @@ void UMassSignalSubsystem::DelaySignalEntity(FName SignalName, const FMassEntity
 
 void UMassSignalSubsystem::DelaySignalEntities(FName SignalName, TConstArrayView<FMassEntityHandle> Entities, const float DelayInSeconds)
 {
+	// If you hit this ensure
+	// - With another thread trying to delay signal then you can use DelaySignalEntityDeferred/DelaySignalEntitiesDeferred
+	//   if you have access to a FMassExecutionContext.
+	// - With the game thread executing UMassSignalSubsystem::Tick then you need to reorganize your tasks to prevent senders from executing
+	//   at the same time as the subsystem tick.
+	UE_MT_SCOPED_WRITE_ACCESS(DelayedSignalsAccessDetector);
+
 	FDelayedSignal& DelayedSignal = DelayedSignals.Emplace_GetRef();
 	DelayedSignal.SignalName = SignalName;
 	DelayedSignal.Entities = Entities;
diff --git a/Engine/Plugins/Runtime/MassGameplay/Source/MassSignals/Public/MassSignalSubsystem.h b/Engine/Plugins/Runtime/MassGameplay/Source/MassSignals/Public/MassSignalSubsystem.h
--- a/Engine/Plugins/Runtime/MassGameplay/Source/MassSignals/Public/MassSignalSubsystem.h
+++ b/Engine/Plugins/Runtime/MassGameplay/Source/MassSignals/Public/MassSignalSubsystem.h
@@ -4,6 +4,7 @@
 
 #include "MassEntityManager.h"
 #include "MassSubsystemBase.h"
+#include "Misc/MTAccessDetector.h"
 #include "MassSignalSubsystem.generated.h"
 
 namespace UE::MassSignal 
@@ -103,6 +104,9 @@ protected:
 	virtual void Tick(float DeltaTime) override;
 	virtual TStatId GetStatId() const override;
 
+	/** Multithreading access detector to validate accesses to the list of delayed signals */
+	UE_MT_DECLARE_RW_ACCESS_DETECTOR(DelayedSignalsAccessDetector);
+
 	TMap<FName, UE::MassSignal::FSignalDelegate> NamedSignals;
 
 	struct FDelayedSignal
