StateTree: Fix runtime issue with context constructed with one owner on start by with another to gather info.

diff --git a/Engine/Plugins/AI/MassAI/Source/MassAIDebug/Private/GameplayDebuggerCategory_Mass.cpp b/Engine/Plugins/AI/MassAI/Source/MassAIDebug/Private/GameplayDebuggerCategory_Mass.cpp
--- a/Engine/Plugins/AI/MassAI/Source/MassAIDebug/Private/GameplayDebuggerCategory_Mass.cpp
+++ b/Engine/Plugins/AI/MassAI/Source/MassAIDebug/Private/GameplayDebuggerCategory_Mass.cpp
@@ -722,7 +722,7 @@ void FGameplayDebuggerCategory_Mass::CollectData(APlayerController* OwnerPC, AAc
 						{
 							if (FStateTreeInstanceData* InstanceData = MassStateTreeSubsystem->GetInstanceData(StateTreeInstance.InstanceHandle))
 							{
-								FMassStateTreeExecutionContext StateTreeContext(*OwnerPC, *StateTree, *InstanceData, EntityManager, *SignalSubsystem, Context);
+								FMassStateTreeExecutionContext StateTreeContext(*MassStateTreeSubsystem, *StateTree, *InstanceData, EntityManager, *SignalSubsystem, Context);
 								StateTreeContext.SetEntity(Entity);
 
 								Status += StateTreeContext.GetActiveStateName();
