Fix UMassZoneGraphLaneCacheBoundaryProcessor not running by adding the missing FMassLaneCacheBoundaryFragment removed at CL 18102996

diff --git a/Engine/Plugins/AI/MassAI/Source/MassZoneGraphNavigation/Private/MassZoneGraphNavigationTrait.cpp b/Engine/Plugins/AI/MassAI/Source/MassZoneGraphNavigation/Private/MassZoneGraphNavigationTrait.cpp
--- a/Engine/Plugins/AI/MassAI/Source/MassZoneGraphNavigation/Private/MassZoneGraphNavigationTrait.cpp
+++ b/Engine/Plugins/AI/MassAI/Source/MassZoneGraphNavigation/Private/MassZoneGraphNavigationTrait.cpp
@@ -22,6 +22,7 @@ void UMassZoneGraphNavigationTrait::BuildTemplate(FMassEntityTemplateBuildContex
 	BuildContext.AddFragment<FMassZoneGraphPathRequestFragment>();
 	BuildContext.AddFragment<FMassZoneGraphShortPathFragment>();
 	BuildContext.AddFragment<FMassZoneGraphCachedLaneFragment>();
+	BuildContext.AddFragment<FMassLaneCacheBoundaryFragment>();
 
 	const FConstSharedStruct ZGMovementParamsFragment = EntityManager.GetOrCreateConstSharedFragment(NavigationParameters);
 	BuildContext.AddConstSharedFragment(ZGMovementParamsFragment);
