[MassStatetreeExecutionContext] refactored context to use MassEntityExecutionContext instead of MassEntityManager and to fetch required subsystems

diff --git a/Engine/Plugins/AI/MassAI/Source/MassAIBehavior/Private/MassStateTreeExecutionContext.cpp b/Engine/Plugins/AI/MassAI/Source/MassAIBehavior/Private/MassStateTreeExecutionContext.cpp
--- a/Engine/Plugins/AI/MassAI/Source/MassAIBehavior/Private/MassStateTreeExecutionContext.cpp
+++ b/Engine/Plugins/AI/MassAI/Source/MassAIBehavior/Private/MassStateTreeExecutionContext.cpp
@@ -115,17 +115,15 @@ bool CollectExternalData(const FStateTreeExecutionContext& Context, const UState
 
 }; // UE::MassBehavior
 
-FMassStateTreeExecutionContext::FMassStateTreeExecutionContext(UObject& InOwner, const UStateTree& InStateTree, FStateTreeInstanceData& InInstanceData,
-															   FMassEntityManager& InEntityManager, UMassSignalSubsystem& InSignalSubsystem, FMassExecutionContext& InContext)
+FMassStateTreeExecutionContext::FMassStateTreeExecutionContext(UObject& InOwner, const UStateTree& InStateTree, FStateTreeInstanceData& InInstanceData, FMassExecutionContext& InContext)
 	: FStateTreeExecutionContext(InOwner, InStateTree, InInstanceData, FOnCollectStateTreeExternalData::CreateStatic(UE::MassBehavior::CollectExternalData))
-	, EntityManager(&InEntityManager)
-	, SignalSubsystem(&InSignalSubsystem)
 	, MassEntityExecutionContext(&InContext)
 {
 }
 
 void FMassStateTreeExecutionContext::BeginDelayedTransition(const FStateTreeTransitionDelayedState& DelayedState)
 {
+	UMassSignalSubsystem* SignalSubsystem = MassEntityExecutionContext->GetMutableSubsystem<UMassSignalSubsystem>();
 	if (SignalSubsystem != nullptr && Entity.IsSet())
 	{
 		// Tick again after the games time has passed to see if the condition still holds true.
diff --git a/Engine/Plugins/AI/MassAI/Source/MassAIBehavior/Private/MassStateTreeProcessors.cpp b/Engine/Plugins/AI/MassAI/Source/MassAIBehavior/Private/MassStateTreeProcessors.cpp
--- a/Engine/Plugins/AI/MassAI/Source/MassAIBehavior/Private/MassStateTreeProcessors.cpp
+++ b/Engine/Plugins/AI/MassAI/Source/MassAIBehavior/Private/MassStateTreeProcessors.cpp
@@ -3,7 +3,6 @@
 #include "MassStateTreeProcessors.h"
 #include "StateTree.h"
 #include "MassStateTreeExecutionContext.h"
-#include "MassEntityView.h"
 #include "MassExecutionContext.h"
 #include "MassNavigationTypes.h"
 #include "MassSimulationLOD.h"
@@ -27,8 +26,6 @@ namespace UE::MassBehavior
 template<typename TFunc>
 void ForEachEntityInChunk(
 	FMassExecutionContext& Context,
-	FMassEntityManager& EntityManager,
-	UMassSignalSubsystem& SignalSubsystem,
 	UMassStateTreeSubsystem& MassStateTreeSubsystem,
 	TFunc&& Callback)
 {
@@ -48,7 +45,7 @@ void ForEachEntityInChunk(
 		FStateTreeInstanceData* InstanceData = MassStateTreeSubsystem.GetInstanceData(StateTreeFragment.InstanceHandle);
 		if (InstanceData != nullptr)
 		{
-			FMassStateTreeExecutionContext StateTreeContext(MassStateTreeSubsystem, *StateTree, *InstanceData, EntityManager, SignalSubsystem, Context);
+			FMassStateTreeExecutionContext StateTreeContext(MassStateTreeSubsystem, *StateTree, *InstanceData, Context);
 			StateTreeContext.SetEntity(Entity);
 
 			// Make sure all required external data are set.
@@ -102,16 +99,15 @@ void UMassStateTreeFragmentDestructor::Execute(FMassEntityManager& EntityManager
 	}
 
 	EntityQuery.ForEachEntityChunk(EntityManager, Context,
-		[&EntityManager, SignalSubsystem = SignalSubsystem](FMassExecutionContext& Context)
+		[](FMassExecutionContext& Context)
 		{
 			UMassStateTreeSubsystem& MassStateTreeSubsystem = Context.GetMutableSubsystemChecked<UMassStateTreeSubsystem>();
 			const TArrayView<FMassStateTreeInstanceFragment> StateTreeInstanceList = Context.GetMutableFragmentView<FMassStateTreeInstanceFragment>();
-			const FMassStateTreeSharedFragment& SharedStateTree = Context.GetConstSharedFragment<FMassStateTreeSharedFragment>();
 
 			const int32 NumEntities = Context.GetNumEntities();
 	
-			UE::MassBehavior::ForEachEntityInChunk(Context, EntityManager, *SignalSubsystem, MassStateTreeSubsystem,
-				[](FMassStateTreeExecutionContext& StateTreeExecutionContext, FMassStateTreeInstanceFragment& StateTreeFragment)
+			UE::MassBehavior::ForEachEntityInChunk(Context, MassStateTreeSubsystem,
+				[](FMassStateTreeExecutionContext& StateTreeExecutionContext, FMassStateTreeInstanceFragment&)
 				{
 					// Stop the tree instance
 					StateTreeExecutionContext.Stop();
@@ -159,7 +155,7 @@ void UMassStateTreeActivationProcessor::Execute(FMassEntityManager& EntityManage
 	const UMassBehaviorSettings* BehaviorSettings = GetDefault<UMassBehaviorSettings>();
 	check(BehaviorSettings);
  
-	// StateTree processor relies on signals to be ticked but we need an 'initial tick' to set the tree in the proper state.
+	// StateTree processor relies on signals to be ticked, but we need an 'initial tick' to set the tree in the proper state.
 	// The initializer provides that by sending a signal to all new entities that use StateTree.
 	const double TimeInSeconds = EntityManager.GetWorld()->GetTimeSeconds();
 
@@ -167,7 +163,7 @@ void UMassStateTreeActivationProcessor::Execute(FMassEntityManager& EntityManage
 	int32 ActivationCounts[EMassLOD::Max] {0,0,0,0};
 	
 	EntityQuery.ForEachEntityChunk(EntityManager, Context,
-		[&EntitiesToSignal, &ActivationCounts, SignalSubsystem = &SignalSubsystem, &EntityManager, MaxActivationsPerLOD = BehaviorSettings->MaxActivationsPerLOD, TimeInSeconds](FMassExecutionContext& Context)
+		[&EntitiesToSignal, &ActivationCounts, MaxActivationsPerLOD = BehaviorSettings->MaxActivationsPerLOD, TimeInSeconds](FMassExecutionContext& Context)
 		{
 			UMassStateTreeSubsystem& MassStateTreeSubsystem = Context.GetMutableSubsystemChecked<UMassStateTreeSubsystem>();
 			const int32 NumEntities = Context.GetNumEntities();
@@ -190,7 +186,7 @@ void UMassStateTreeActivationProcessor::Execute(FMassEntityManager& EntityManage
 			}
 			
 			// Start StateTree. This may do substantial amount of work, as we select and enter the first state.
-			UE::MassBehavior::ForEachEntityInChunk(Context, EntityManager, *SignalSubsystem, MassStateTreeSubsystem,
+			UE::MassBehavior::ForEachEntityInChunk(Context, MassStateTreeSubsystem,
 				[TimeInSeconds](FMassStateTreeExecutionContext& StateTreeExecutionContext, FMassStateTreeInstanceFragment& StateTreeFragment)
 				{
 					// Start the tree instance
@@ -271,6 +267,7 @@ void UMassStateTreeProcessor::ConfigureQueries(const TSharedRef<FMassEntityManag
 	EntityQuery.AddRequirement<FMassStateTreeInstanceFragment>(EMassFragmentAccess::ReadWrite);
 	EntityQuery.AddConstSharedRequirement<FMassStateTreeSharedFragment>();
 	EntityQuery.AddSubsystemRequirement<UMassStateTreeSubsystem>(EMassFragmentAccess::ReadWrite);
+	EntityQuery.AddSubsystemRequirement<UMassSignalSubsystem>(EMassFragmentAccess::ReadWrite);
 
 	ProcessorRequirements.AddSubsystemRequirement<UMassSignalSubsystem>(EMassFragmentAccess::ReadWrite);
 }
@@ -286,14 +283,14 @@ void UMassStateTreeProcessor::SignalEntities(FMassEntityManager& EntityManager,
 
 	TArray<FMassEntityHandle> EntitiesToSignal;
 
-	auto TickFunc = [&](FMassExecutionContext& Context)
+	auto TickFunc = [TimeInSeconds, &EntitiesToSignal, &EntitySignals](FMassExecutionContext& ExecutionContext)
 		{
 			// Keep stats regarding the amount of tree instances ticked per frame
-			CSV_CUSTOM_STAT(StateTreeProcessor, NumTickedStateTree, Context.GetNumEntities(), ECsvCustomStatOp::Accumulate);
+			CSV_CUSTOM_STAT(StateTreeProcessor, NumTickedStateTree, ExecutionContext.GetNumEntities(), ECsvCustomStatOp::Accumulate);
 
-			UMassStateTreeSubsystem& MassStateTreeSubsystem = Context.GetMutableSubsystemChecked<UMassStateTreeSubsystem>();
+			UMassStateTreeSubsystem& MassStateTreeSubsystem = ExecutionContext.GetMutableSubsystemChecked<UMassStateTreeSubsystem>();
 
-			UE::MassBehavior::ForEachEntityInChunk(Context, EntityManager, SignalSubsystem, MassStateTreeSubsystem,
+			UE::MassBehavior::ForEachEntityInChunk(ExecutionContext, MassStateTreeSubsystem,
 				[TimeInSeconds, &EntitiesToSignal, &EntitySignals, &MassStateTreeSubsystem]
 				(FMassStateTreeExecutionContext& StateTreeExecutionContext, FMassStateTreeInstanceFragment& StateTreeFragment)
 				{
@@ -323,7 +320,7 @@ void UMassStateTreeProcessor::SignalEntities(FMassEntityManager& EntityManager,
 					// Tick the tree instance
 					StateTreeExecutionContext.Tick(AdjustedDeltaTime);
 
-					// When last tick status is different than "Running", the state tree need to be tick again
+					// When last tick status is different from "Running", the state tree need to be tick again
 					// For performance reason, tick again to see if we could find a new state right away instead of waiting the next frame.
 					if (StateTreeExecutionContext.GetLastTickStatus() != EStateTreeRunStatus::Running)
 					{
diff --git a/Engine/Plugins/AI/MassAI/Source/MassAIBehavior/Private/Tasks/MassClaimSmartObjectTask.cpp b/Engine/Plugins/AI/MassAI/Source/MassAIBehavior/Private/Tasks/MassClaimSmartObjectTask.cpp
--- a/Engine/Plugins/AI/MassAI/Source/MassAIBehavior/Private/Tasks/MassClaimSmartObjectTask.cpp
+++ b/Engine/Plugins/AI/MassAI/Source/MassAIBehavior/Private/Tasks/MassClaimSmartObjectTask.cpp
@@ -51,7 +51,7 @@ EStateTreeRunStatus FMassClaimSmartObjectTask::EnterState(FStateTreeExecutionCon
 	
 	// Setup MassSmartObject handler and claim
 	const FMassStateTreeExecutionContext& MassStateTreeContext = static_cast<FMassStateTreeExecutionContext&>(Context);
-	const FMassSmartObjectHandler MassSmartObjectHandler(MassStateTreeContext.GetEntityManager(), MassStateTreeContext.GetMassEntityExecutionContext(), SmartObjectSubsystem, SignalSubsystem);
+	const FMassSmartObjectHandler MassSmartObjectHandler(MassStateTreeContext.GetMassEntityExecutionContext(), SmartObjectSubsystem, SignalSubsystem);
 	
 	InstanceData.ClaimedSlot = MassSmartObjectHandler.ClaimCandidate(MassStateTreeContext.GetEntity(), SOUser, *CandidateSlots);
 
@@ -82,7 +82,7 @@ void FMassClaimSmartObjectTask::ExitState(FStateTreeExecutionContext& Context, c
 		const FMassStateTreeExecutionContext& MassStateTreeContext = static_cast<FMassStateTreeExecutionContext&>(Context);
 		USmartObjectSubsystem& SmartObjectSubsystem = Context.GetExternalData(SmartObjectSubsystemHandle);
 		UMassSignalSubsystem& SignalSubsystem = Context.GetExternalData(MassSignalSubsystemHandle);
-		const FMassSmartObjectHandler MassSmartObjectHandler(MassStateTreeContext.GetEntityManager(), MassStateTreeContext.GetMassEntityExecutionContext(), SmartObjectSubsystem, SignalSubsystem);
+		const FMassSmartObjectHandler MassSmartObjectHandler(MassStateTreeContext.GetMassEntityExecutionContext(), SmartObjectSubsystem, SignalSubsystem);
 
 		MassSmartObjectHandler.ReleaseSmartObject(MassStateTreeContext.GetEntity(), SOUser, InstanceData.ClaimedSlot);
 	}
diff --git a/Engine/Plugins/AI/MassAI/Source/MassAIBehavior/Private/Tasks/MassFindSmartObjectTask.cpp b/Engine/Plugins/AI/MassAI/Source/MassAIBehavior/Private/Tasks/MassFindSmartObjectTask.cpp
--- a/Engine/Plugins/AI/MassAI/Source/MassAIBehavior/Private/Tasks/MassFindSmartObjectTask.cpp
+++ b/Engine/Plugins/AI/MassAI/Source/MassAIBehavior/Private/Tasks/MassFindSmartObjectTask.cpp
@@ -8,7 +8,6 @@
 #include "MassSmartObjectHandler.h"
 #include "MassStateTreeExecutionContext.h"
 #include "MassZoneGraphNavigationFragments.h"
-#include "SmartObjectDefinitionReference.h"
 #include "SmartObjectSubsystem.h"
 #include "SmartObjectZoneAnnotations.h"
 #include "StateTreeLinker.h"
@@ -44,7 +43,6 @@ void FMassFindSmartObjectTask::ExitState(FStateTreeExecutionContext& Context, co
 		USmartObjectSubsystem& SmartObjectSubsystem = Context.GetExternalData(SmartObjectSubsystemHandle);
 		UMassSignalSubsystem& SignalSubsystem = Context.GetExternalData(MassSignalSubsystemHandle);
 		const FMassSmartObjectHandler MassSmartObjectHandler(
-			MassStateTreeContext.GetEntityManager(),
 			MassStateTreeContext.GetMassEntityExecutionContext(),
 			SmartObjectSubsystem,
 			SignalSubsystem);
@@ -85,7 +83,6 @@ EStateTreeRunStatus FMassFindSmartObjectTask::Tick(FStateTreeExecutionContext& C
 	UMassSignalSubsystem& SignalSubsystem = Context.GetExternalData(MassSignalSubsystemHandle);
 	const FMassStateTreeExecutionContext& MassStateTreeContext = static_cast<FMassStateTreeExecutionContext&>(Context);
 	const FMassSmartObjectHandler MassSmartObjectHandler(
-		MassStateTreeContext.GetEntityManager(),
 		MassStateTreeContext.GetMassEntityExecutionContext(),
 		SmartObjectSubsystem,
 		SignalSubsystem);
diff --git a/Engine/Plugins/AI/MassAI/Source/MassAIBehavior/Private/Tasks/MassUseSmartObjectTask.cpp b/Engine/Plugins/AI/MassAI/Source/MassAIBehavior/Private/Tasks/MassUseSmartObjectTask.cpp
--- a/Engine/Plugins/AI/MassAI/Source/MassAIBehavior/Private/Tasks/MassUseSmartObjectTask.cpp
+++ b/Engine/Plugins/AI/MassAI/Source/MassAIBehavior/Private/Tasks/MassUseSmartObjectTask.cpp
@@ -53,7 +53,7 @@ EStateTreeRunStatus FMassUseSmartObjectTask::EnterState(FStateTreeExecutionConte
 
 	// Setup MassSmartObject handler and start interaction
 	const FMassStateTreeExecutionContext& MassStateTreeContext = static_cast<FMassStateTreeExecutionContext&>(Context);
-	const FMassSmartObjectHandler MassSmartObjectHandler(MassStateTreeContext.GetEntityManager(), MassStateTreeContext.GetMassEntityExecutionContext(), SmartObjectSubsystem, SignalSubsystem);
+	const FMassSmartObjectHandler MassSmartObjectHandler(MassStateTreeContext.GetMassEntityExecutionContext(), SmartObjectSubsystem, SignalSubsystem);
 
 	if (!MassSmartObjectHandler.StartUsingSmartObject(MassStateTreeContext.GetEntity(), SOUser, InstanceData.ClaimedSlot))
 	{
@@ -81,7 +81,7 @@ void FMassUseSmartObjectTask::ExitState(FStateTreeExecutionContext& Context, con
 		const FMassStateTreeExecutionContext& MassStateTreeContext = static_cast<FMassStateTreeExecutionContext&>(Context);
 		USmartObjectSubsystem& SmartObjectSubsystem = Context.GetExternalData(SmartObjectSubsystemHandle);
 		UMassSignalSubsystem& SignalSubsystem = Context.GetExternalData(MassSignalSubsystemHandle);
-		const FMassSmartObjectHandler MassSmartObjectHandler(MassStateTreeContext.GetEntityManager(), MassStateTreeContext.GetMassEntityExecutionContext(), SmartObjectSubsystem, SignalSubsystem);
+		const FMassSmartObjectHandler MassSmartObjectHandler(MassStateTreeContext.GetMassEntityExecutionContext(), SmartObjectSubsystem, SignalSubsystem);
 		MassSmartObjectHandler.StopUsingSmartObject(MassStateTreeContext.GetEntity(), SOUser, EMassSmartObjectInteractionStatus::Aborted);
 	}
 	else
@@ -101,7 +101,7 @@ void FMassUseSmartObjectTask::StateCompleted(FStateTreeExecutionContext& Context
 		const FMassStateTreeExecutionContext& MassStateTreeContext = static_cast<FMassStateTreeExecutionContext&>(Context);
 		USmartObjectSubsystem& SmartObjectSubsystem = Context.GetExternalData(SmartObjectSubsystemHandle);
 		UMassSignalSubsystem& SignalSubsystem = Context.GetExternalData(MassSignalSubsystemHandle);
-		const FMassSmartObjectHandler MassSmartObjectHandler(MassStateTreeContext.GetEntityManager(), MassStateTreeContext.GetMassEntityExecutionContext(), SmartObjectSubsystem, SignalSubsystem);
+		const FMassSmartObjectHandler MassSmartObjectHandler(MassStateTreeContext.GetMassEntityExecutionContext(), SmartObjectSubsystem, SignalSubsystem);
 		const EMassSmartObjectInteractionStatus NewStatus = (CompletionStatus == EStateTreeRunStatus::Succeeded)
 																? EMassSmartObjectInteractionStatus::TaskCompleted
 																: EMassSmartObjectInteractionStatus::Aborted;
diff --git a/Engine/Plugins/AI/MassAI/Source/MassAIBehavior/Public/MassStateTreeExecutionContext.h b/Engine/Plugins/AI/MassAI/Source/MassAIBehavior/Public/MassStateTreeExecutionContext.h
--- a/Engine/Plugins/AI/MassAI/Source/MassAIBehavior/Public/MassStateTreeExecutionContext.h
+++ b/Engine/Plugins/AI/MassAI/Source/MassAIBehavior/Public/MassStateTreeExecutionContext.h
@@ -6,6 +6,7 @@
 #include "MassEntityTypes.h"
 #endif // UE_ENABLE_INCLUDE_ORDER_DEPRECATED_IN_5_6
 #include "MassEntityHandle.h"
+#include "MassExecutionContext.h"
 #include "StateTreeExecutionContext.h"
 
 struct FMassExecutionContext;
@@ -17,15 +18,41 @@ class UMassSignalSubsystem;
  */
 struct MASSAIBEHAVIOR_API FMassStateTreeExecutionContext : public FStateTreeExecutionContext 
 {
-public:
-	// @todo: refactor subsystems out of the exec context.
-	FMassStateTreeExecutionContext(UObject& InOwner, const UStateTree& InStateTree, FStateTreeInstanceData& InInstanceData, FMassEntityManager& InEntityManager, UMassSignalSubsystem& InSignalSubsystem, FMassExecutionContext& InContext);
+	FMassStateTreeExecutionContext(UObject& InOwner
+		, const UStateTree& InStateTree
+		, FStateTreeInstanceData& InInstanceData
+		, FMassExecutionContext& InContext);
 
-	FMassEntityManager& GetEntityManager() const { check(EntityManager); return *EntityManager; }
-	FMassExecutionContext& GetMassEntityExecutionContext() const { return *MassEntityExecutionContext; }
+	UE_DEPRECATED(5.6, "Use the other constructor that doesn't require MassEntityManager and MassSignalSubSystem")
+	FMassStateTreeExecutionContext(UObject& InOwner
+		, const UStateTree& InStateTree
+		, FStateTreeInstanceData& InInstanceData
+		, FMassEntityManager& InEntityManager
+		, UMassSignalSubsystem& InSignalSubsystem
+		, FMassExecutionContext& InContext)
+			: FMassStateTreeExecutionContext(InOwner, InStateTree, InInstanceData, InContext)
+	{
+	}
 
-	FMassEntityHandle GetEntity() const { return Entity; }
-	void SetEntity(const FMassEntityHandle InEntity) { Entity = InEntity; }
+	FMassEntityManager& GetEntityManager() const
+	{
+		return GetMassEntityExecutionContext().GetEntityManagerChecked();
+	}
+
+	FMassExecutionContext& GetMassEntityExecutionContext() const
+	{
+		return *MassEntityExecutionContext;
+	}
+
+	FMassEntityHandle GetEntity() const
+	{
+		return Entity;
+	}
+
+	void SetEntity(const FMassEntityHandle InEntity)
+	{
+		Entity = InEntity;
+	}
 
 protected:
 
@@ -34,8 +61,6 @@ protected:
 
 	virtual void BeginDelayedTransition(const FStateTreeTransitionDelayedState& DelayedState) override;
 
-	FMassEntityManager* EntityManager;
-	UMassSignalSubsystem* SignalSubsystem = nullptr;
 	FMassExecutionContext* MassEntityExecutionContext = nullptr;
 	FMassEntityHandle Entity;
 };
diff --git a/Engine/Plugins/AI/MassAI/Source/MassAIDebug/Private/GameplayDebuggerCategory_Mass.cpp b/Engine/Plugins/AI/MassAI/Source/MassAIDebug/Private/GameplayDebuggerCategory_Mass.cpp
--- a/Engine/Plugins/AI/MassAI/Source/MassAIDebug/Private/GameplayDebuggerCategory_Mass.cpp
+++ b/Engine/Plugins/AI/MassAI/Source/MassAIDebug/Private/GameplayDebuggerCategory_Mass.cpp
@@ -4,12 +4,10 @@
 
 #if WITH_GAMEPLAY_DEBUGGER && WITH_MASSGAMEPLAY_DEBUG
 #include "MassGameplayDebugTypes.h"
-#include "MassEntityView.h"
 #include "MassExecutionContext.h"
 #include "GameplayDebuggerConfig.h"
 #include "GameplayDebuggerCategoryReplicator.h"
 #include "MassDebuggerSubsystem.h"
-#include "MassSignalSubsystem.h"
 #include "MassActorSubsystem.h"
 #include "GameFramework/PlayerController.h"
 #include "MassAgentComponent.h"
@@ -514,14 +512,13 @@ void FGameplayDebuggerCategory_Mass::CollectData(APlayerController* OwnerPC, AAc
 		const double CurrentTime = World->GetTimeSeconds();
 		
 		UMassStateTreeSubsystem* MassStateTreeSubsystem = World->GetSubsystem<UMassStateTreeSubsystem>();
-		UMassSignalSubsystem* SignalSubsystem = World->GetSubsystem<UMassSignalSubsystem>();
 		USmartObjectSubsystem* SmartObjectSubsystem = World->GetSubsystem<USmartObjectSubsystem>();
 		
-		if (MassStateTreeSubsystem && SignalSubsystem && SmartObjectSubsystem)
+		if (MassStateTreeSubsystem && SmartObjectSubsystem)
 		{
 			FMassExecutionContext Context(EntityManager, 0.0f);
 		
-			EntityQuery.ForEachEntityChunk(EntityManager, Context, [this, MassStateTreeSubsystem, SignalSubsystem, SmartObjectSubsystem, OwnerPC, ViewLocation, ViewDirection, CurrentTime](FMassExecutionContext& Context)
+			EntityQuery.ForEachEntityChunk(EntityManager, Context, [this, MassStateTreeSubsystem, SmartObjectSubsystem, ViewLocation, ViewDirection, CurrentTime](FMassExecutionContext& Context)
 			{
 				FMassEntityManager& EntityManager = Context.GetEntityManagerChecked();
 
@@ -721,7 +718,7 @@ void FGameplayDebuggerCategory_Mass::CollectData(APlayerController* OwnerPC, AAc
 						{
 							if (FStateTreeInstanceData* InstanceData = MassStateTreeSubsystem->GetInstanceData(StateTreeInstance.InstanceHandle))
 							{
-								FMassStateTreeExecutionContext StateTreeContext(*MassStateTreeSubsystem, *StateTree, *InstanceData, EntityManager, *SignalSubsystem, Context);
+								FMassStateTreeExecutionContext StateTreeContext(*MassStateTreeSubsystem, *StateTree, *InstanceData, Context);
 								StateTreeContext.SetEntity(Entity);
 
 								Status += StateTreeContext.GetActiveStateName();
diff --git a/Engine/Plugins/AI/MassAI/Source/MassAIDebug/Private/MassDebugStateTreeProcessor.cpp b/Engine/Plugins/AI/MassAI/Source/MassAIDebug/Private/MassDebugStateTreeProcessor.cpp
--- a/Engine/Plugins/AI/MassAI/Source/MassAIDebug/Private/MassDebugStateTreeProcessor.cpp
+++ b/Engine/Plugins/AI/MassAI/Source/MassAIDebug/Private/MassDebugStateTreeProcessor.cpp
@@ -2,7 +2,6 @@
 
 #include "MassDebugStateTreeProcessor.h"
 #include "MassDebuggerSubsystem.h"
-#include "MassSignalSubsystem.h"
 #include "MassStateTreeExecutionContext.h"
 #include "MassStateTreeFragments.h"
 #include "MassCommonFragments.h"
@@ -50,19 +49,13 @@ void UMassDebugStateTreeProcessor::Execute(FMassEntityManager& EntityManager, FM
 		return;
 	}
 
-	UMassSignalSubsystem* MassSignalSubsystem = World->GetSubsystem<UMassSignalSubsystem>();
-	if (MassSignalSubsystem == nullptr)
-	{
-		return;
-	}
-
 	if (!Debugger->GetSelectedEntity().IsSet() && !UE::Mass::Debug::HasDebugEntities())
 	{
 		return;
 	}
 	
 	QUICK_SCOPE_CYCLE_COUNTER(UMassDebugStateTreeProcessor_Run);	
-	EntityQuery.ForEachEntityChunk(EntityManager, Context, [this, Debugger, MassStateTreeSubsystem, &EntityManager, MassSignalSubsystem](FMassExecutionContext& Context)
+	EntityQuery.ForEachEntityChunk(EntityManager, Context, [this, Debugger, MassStateTreeSubsystem](FMassExecutionContext& Context)
 	{
 		const FMassEntityHandle SelectedEntity = Debugger->GetSelectedEntity();
 		const int32 NumEntities = Context.GetNumEntities();
@@ -97,7 +90,7 @@ void UMassDebugStateTreeProcessor::Execute(FMassEntityManager& EntityManager, FM
 			
 			if (Entity == SelectedEntity)
 			{
-				FMassStateTreeExecutionContext StateTreeContext(*MassStateTreeSubsystem, *StateTree, *InstanceData, EntityManager, *MassSignalSubsystem, Context);
+				FMassStateTreeExecutionContext StateTreeContext(*MassStateTreeSubsystem, *StateTree, *InstanceData, Context);
 				StateTreeContext.SetEntity(Entity);
 				
 #if WITH_GAMEPLAY_DEBUGGER
@@ -114,7 +107,7 @@ void UMassDebugStateTreeProcessor::Execute(FMassEntityManager& EntityManager, FM
 				const FVector ZOffset(0,0,50);
 				const FVector Position = Transform.GetTransform().GetLocation() + ZOffset;
 
-				FMassStateTreeExecutionContext StateTreeContext(*MassStateTreeSubsystem, *StateTree, *InstanceData, EntityManager, *MassSignalSubsystem, Context);
+				FMassStateTreeExecutionContext StateTreeContext(*MassStateTreeSubsystem, *StateTree, *InstanceData, Context);
 				StateTreeContext.SetEntity(Entity);
 
 				// State
diff --git a/Engine/Plugins/Runtime/MassGameplay/Source/MassSmartObjects/Private/MassSmartObjectHandler.cpp b/Engine/Plugins/Runtime/MassGameplay/Source/MassSmartObjects/Private/MassSmartObjectHandler.cpp
--- a/Engine/Plugins/Runtime/MassGameplay/Source/MassSmartObjects/Private/MassSmartObjectHandler.cpp
+++ b/Engine/Plugins/Runtime/MassGameplay/Source/MassSmartObjects/Private/MassSmartObjectHandler.cpp
@@ -7,7 +7,6 @@
 #include "SmartObjectSubsystem.h"
 #include "VisualLogger/VisualLogger.h"
 #include "MassCommandBuffer.h"
-#include "MassCommonFragments.h"
 #include "MassSignalSubsystem.h"
 #include "MassSmartObjectFragments.h"
 #include "Engine/World.h"
@@ -69,7 +68,7 @@ namespace UE::Mass::SmartObject
 //----------------------------------------------------------------------//
 FMassSmartObjectRequestID FMassSmartObjectHandler::FindCandidatesAsync(const FMassEntityHandle RequestingEntity, const FGameplayTagContainer& UserTags, const FGameplayTagQuery& ActivityRequirements, const FVector& Location) const
 {
-	const FMassEntityHandle ReservedEntity = EntityManager.ReserveEntity();
+	const FMassEntityHandle ReservedEntity = ExecutionContext.GetEntityManagerChecked().ReserveEntity();
 
 	FMassSmartObjectWorldLocationRequestFragment RequestFragment;
 	RequestFragment.SearchOrigin = Location;
@@ -86,7 +85,7 @@ FMassSmartObjectRequestID FMassSmartObjectHandler::FindCandidatesAsync(const FMa
 
 FMassSmartObjectRequestID FMassSmartObjectHandler::FindCandidatesAsync(const FMassEntityHandle RequestingEntity, const FGameplayTagContainer& UserTags, const FGameplayTagQuery& ActivityRequirements, const FZoneGraphCompactLaneLocation& LaneLocation) const
 {
-	const FMassEntityHandle ReservedEntity = EntityManager.ReserveEntity();
+	const FMassEntityHandle ReservedEntity = ExecutionContext.GetEntityManagerChecked().ReserveEntity();
 
 	FMassSmartObjectLaneLocationRequestFragment RequestFragment;
 	RequestFragment.CompactLaneLocation = LaneLocation;
@@ -104,6 +103,8 @@ FMassSmartObjectRequestID FMassSmartObjectHandler::FindCandidatesAsync(const FMa
 const FMassSmartObjectCandidateSlots* FMassSmartObjectHandler::GetRequestCandidates(const FMassSmartObjectRequestID& RequestID) const
 {
 	const FMassEntityHandle RequestEntity = static_cast<FMassEntityHandle>(RequestID);
+	const FMassEntityManager& EntityManager = ExecutionContext.GetEntityManagerChecked();
+
 	if (!ensureMsgf(EntityManager.IsEntityValid(RequestEntity), TEXT("Invalid request.")))
 	{
 		return nullptr;
@@ -174,7 +175,7 @@ FSmartObjectClaimHandle FMassSmartObjectHandler::ClaimSmartObject(const FMassEnt
 	// Callback will be unregistered by UMassSmartObjectUserFragmentDeinitializer
 	UE::Mass::SmartObject::FPayload Payload;
 	Payload.Entity = Entity;
-	Payload.EntityManager = EntityManager.AsShared();
+	Payload.EntityManager = ExecutionContext.GetEntityManagerChecked().AsShared();
 	Payload.SmartObjectSubsystem = &SmartObjectSubsystem;
 	Payload.SignalSubsystem = &SignalSubsystem;
 	SmartObjectSubsystem.RegisterSlotInvalidationCallback(ClaimHandle, FOnSlotInvalidated::CreateStatic(&UE::Mass::SmartObject::OnSlotInvalidated, Payload));
@@ -207,7 +208,7 @@ bool FMassSmartObjectHandler::StartUsingSmartObject(
 	User.InteractionHandle = ClaimHandle;
 
 	// Activate behavior
-	const FMassEntityView EntityView(EntityManager, Entity);
+	const FMassEntityView EntityView(ExecutionContext.GetEntityManagerChecked(), Entity);
 	const FMassBehaviorEntityContext Context(EntityView, SmartObjectSubsystem);
 	BehaviorDefinition->Activate(ExecutionContext.Defer(), Context);
 
@@ -264,7 +265,7 @@ void FMassSmartObjectHandler::StopUsingSmartObject(const FMassEntityHandle Entit
 	}
 #endif
 
-	const FMassEntityView EntityView(EntityManager, Entity);
+	const FMassEntityView EntityView(ExecutionContext.GetEntityManagerChecked(), Entity);
 	const FMassBehaviorEntityContext Context(EntityView, SmartObjectSubsystem);
 	UE::Mass::SmartObject::StopInteraction(ExecutionContext.Defer(), Context, User.InteractionHandle, NewStatus);
 }
@@ -289,7 +290,7 @@ void FMassSmartObjectHandler::ReleaseSmartObject(const FMassEntityHandle Entity,
 	}
 #endif
 
-	const FMassEntityView EntityView(EntityManager, Entity);
+	const FMassEntityView EntityView(ExecutionContext.GetEntityManagerChecked(), Entity);
 	const FMassBehaviorEntityContext Context(EntityView, SmartObjectSubsystem);
 
 	SmartObjectSubsystem.UnregisterSlotInvalidationCallback(ClaimHandle);
diff --git a/Engine/Plugins/Runtime/MassGameplay/Source/MassSmartObjects/Public/MassSmartObjectHandler.h b/Engine/Plugins/Runtime/MassGameplay/Source/MassSmartObjects/Public/MassSmartObjectHandler.h
--- a/Engine/Plugins/Runtime/MassGameplay/Source/MassSmartObjects/Public/MassSmartObjectHandler.h
+++ b/Engine/Plugins/Runtime/MassGameplay/Source/MassSmartObjects/Public/MassSmartObjectHandler.h
@@ -25,19 +25,23 @@ struct MASSSMARTOBJECTS_API FMassSmartObjectHandler
 {
 	/**
 	 * FMassSmartObjectHandler constructor
-	 * @param InEntitySubsystem is the entity subsystem that the smart object should belong to
 	 * @param InExecutionContext is the current execution context of the entity subsystem
 	 * @param InSmartObjectSubsystem is the smart object subsystem
 	 * @param InSignalSubsystem is the mass signal subsystem to use to send signal to affected entities
 	 */
-	FMassSmartObjectHandler(FMassEntityManager& InEntityManager, FMassExecutionContext& InExecutionContext, USmartObjectSubsystem& InSmartObjectSubsystem, UMassSignalSubsystem& InSignalSubsystem)
-		: EntityManager(InEntityManager)
-		, ExecutionContext(InExecutionContext)
+	FMassSmartObjectHandler(FMassExecutionContext& InExecutionContext, USmartObjectSubsystem& InSmartObjectSubsystem, UMassSignalSubsystem& InSignalSubsystem)
+		: ExecutionContext(InExecutionContext)
 		, SmartObjectSubsystem(InSmartObjectSubsystem)
 		, SignalSubsystem(InSignalSubsystem)
 	{
 	}
 
+	UE_DEPRECATED(5.6, "Use the other constructor that doesn't require MassEntityManager")
+	FMassSmartObjectHandler(FMassEntityManager& InEntityManager, FMassExecutionContext& InExecutionContext, USmartObjectSubsystem& InSmartObjectSubsystem, UMassSignalSubsystem& InSignalSubsystem)
+		: FMassSmartObjectHandler(InExecutionContext, InSmartObjectSubsystem, InSignalSubsystem)
+	{
+	}
+
 	/**
 	 * Creates an async request to build a list of compatible smart objects
 	 * around the provided location. The caller must poll using the request id
@@ -120,7 +124,6 @@ struct MASSSMARTOBJECTS_API FMassSmartObjectHandler
 	void ReleaseSmartObject(const FMassEntityHandle Entity, FMassSmartObjectUserFragment& User, const FSmartObjectClaimHandle ClaimHandle) const;
 
 private:
-	FMassEntityManager& EntityManager;
 	FMassExecutionContext& ExecutionContext;
 	USmartObjectSubsystem& SmartObjectSubsystem;
 	UMassSignalSubsystem& SignalSubsystem;
