[Mass Gameplay/AI] updated LookAt processor, GameplayDebugger and AreAnyEntitiesVisibleInChunk to handle optional fragments

diff --git a/Engine/Plugins/AI/MassAI/Source/MassAIBehavior/Private/MassLookAtProcessors.cpp b/Engine/Plugins/AI/MassAI/Source/MassAIBehavior/Private/MassLookAtProcessors.cpp
--- a/Engine/Plugins/AI/MassAI/Source/MassAIBehavior/Private/MassLookAtProcessors.cpp
+++ b/Engine/Plugins/AI/MassAI/Source/MassAIBehavior/Private/MassLookAtProcessors.cpp
@@ -104,14 +104,14 @@ void UMassLookAtProcessor::ConfigureQueries(const TSharedRef<FMassEntityManager>
 {
 	EntityQuery_Conditional.AddRequirement<FMassLookAtFragment>(EMassFragmentAccess::ReadWrite);
 	EntityQuery_Conditional.AddRequirement<FTransformFragment>(EMassFragmentAccess::ReadOnly);
-	EntityQuery_Conditional.AddRequirement<FMassMoveTargetFragment>(EMassFragmentAccess::ReadOnly);
+	EntityQuery_Conditional.AddRequirement<FMassMoveTargetFragment>(EMassFragmentAccess::ReadOnly, EMassFragmentPresence::Optional);
 	EntityQuery_Conditional.AddRequirement<FMassZoneGraphLaneLocationFragment>(EMassFragmentAccess::ReadOnly, EMassFragmentPresence::Optional);
 	EntityQuery_Conditional.AddRequirement<FMassLookAtTrajectoryFragment>(EMassFragmentAccess::ReadWrite, EMassFragmentPresence::Optional);
 	EntityQuery_Conditional.AddRequirement<FMassZoneGraphShortPathFragment>(EMassFragmentAccess::ReadOnly, EMassFragmentPresence::Optional);
 	EntityQuery_Conditional.AddTagRequirement<FMassMediumLODTag>(EMassFragmentPresence::None);
 	EntityQuery_Conditional.AddTagRequirement<FMassLowLODTag>(EMassFragmentPresence::None);
 	EntityQuery_Conditional.AddTagRequirement<FMassOffLODTag>(EMassFragmentPresence::None);
-	EntityQuery_Conditional.AddChunkRequirement<FMassVisualizationChunkFragment>(EMassFragmentAccess::ReadOnly);
+	EntityQuery_Conditional.AddChunkRequirement<FMassVisualizationChunkFragment>(EMassFragmentAccess::ReadOnly, EMassFragmentPresence::Optional);
 	EntityQuery_Conditional.SetChunkFilter(&FMassVisualizationChunkFragment::AreAnyEntitiesVisibleInChunk);
 	EntityQuery_Conditional.AddSubsystemRequirement<UMassNavigationSubsystem>(EMassFragmentAccess::ReadOnly);
 	EntityQuery_Conditional.AddSubsystemRequirement<UZoneGraphSubsystem>(EMassFragmentAccess::ReadOnly);
@@ -139,10 +139,13 @@ void UMassLookAtProcessor::Execute(FMassEntityManager& EntityManager, FMassExecu
 			for (int32 i = 0; i < NumEntities; ++i)
 			{
 				FMassLookAtFragment& LookAt = LookAtList[i];
-				const FMassMoveTargetFragment& MoveTarget = MoveTargetList[i];
 				const FTransformFragment& TransformFragment = TransformList[i];
 
-				const bool bHasLookAtTrajectory = ZoneGraphLocationList.Num() > 0 && LookAtTrajectoryList.Num() > 0 && ShortPathList.Num() > 0;
+				const bool bHasLookAtTrajectory =
+					MoveTargetList.Num() > 0
+					&& ZoneGraphLocationList.Num() > 0
+					&& LookAtTrajectoryList.Num() > 0
+					&& ShortPathList.Num() > 0;
 
 				bool bDisplayDebug = false;
 				const FMassEntityHandle Entity = Context.GetEntity(i);
@@ -173,6 +176,7 @@ void UMassLookAtProcessor::Execute(FMassEntityManager& EntityManager, FMassExecu
 				case EMassLookAtMode::LookAlongPath:
 					if (bHasLookAtTrajectory)
 					{
+						const FMassMoveTargetFragment& MoveTarget = MoveTargetList[i];
 						const FMassZoneGraphLaneLocationFragment& ZoneGraphLocation = ZoneGraphLocationList[i];
 						FMassLookAtTrajectoryFragment& LookAtTrajectory = LookAtTrajectoryList[i];
 
diff --git a/Engine/Plugins/AI/MassAI/Source/MassAIDebug/Private/GameplayDebuggerCategory_Mass.cpp b/Engine/Plugins/AI/MassAI/Source/MassAIDebug/Private/GameplayDebuggerCategory_Mass.cpp
--- a/Engine/Plugins/AI/MassAI/Source/MassAIDebug/Private/GameplayDebuggerCategory_Mass.cpp
+++ b/Engine/Plugins/AI/MassAI/Source/MassAIDebug/Private/GameplayDebuggerCategory_Mass.cpp
@@ -500,13 +500,13 @@ void FGameplayDebuggerCategory_Mass::CollectData(APlayerController* OwnerPC, AAc
 		EntityQuery.AddRequirement<FAgentRadiusFragment>(EMassFragmentAccess::ReadOnly);
 		EntityQuery.AddRequirement<FMassSteeringFragment>(EMassFragmentAccess::ReadOnly);
 		EntityQuery.AddRequirement<FMassStandingSteeringFragment>(EMassFragmentAccess::ReadOnly);
-		EntityQuery.AddRequirement<FMassGhostLocationFragment>(EMassFragmentAccess::ReadOnly);
+		EntityQuery.AddRequirement<FMassGhostLocationFragment>(EMassFragmentAccess::ReadOnly, EMassFragmentPresence::Optional);
 		EntityQuery.AddRequirement<FMassVelocityFragment>(EMassFragmentAccess::ReadOnly);
 		EntityQuery.AddRequirement<FMassForceFragment>(EMassFragmentAccess::ReadOnly);
 		EntityQuery.AddRequirement<FMassMoveTargetFragment>(EMassFragmentAccess::ReadOnly);
 		EntityQuery.AddRequirement<FMassLookAtFragment>(EMassFragmentAccess::ReadOnly, EMassFragmentPresence::Optional);
 		EntityQuery.AddRequirement<FMassSimulationLODFragment>(EMassFragmentAccess::ReadOnly, EMassFragmentPresence::Optional);
-		EntityQuery.AddRequirement<FMassZoneGraphShortPathFragment>(EMassFragmentAccess::ReadOnly);
+		EntityQuery.AddRequirement<FMassZoneGraphShortPathFragment>(EMassFragmentAccess::ReadOnly, EMassFragmentPresence::Optional);
 		EntityQuery.AddRequirement<FMassSmartObjectUserFragment>(EMassFragmentAccess::ReadOnly, EMassFragmentPresence::Optional);
 
 		const double CurrentTime = World->GetTimeSeconds();
@@ -569,12 +569,10 @@ void FGameplayDebuggerCategory_Mass::CollectData(APlayerController* OwnerPC, AAc
 					const FAgentRadiusFragment& Radius = RadiusList[EntityIndex];
 					const FMassSteeringFragment& Steering = SteeringList[EntityIndex];
 					const FMassStandingSteeringFragment& StandingSteering = StandingSteeringList[EntityIndex];
-					const FMassGhostLocationFragment& Ghost = GhostList[EntityIndex];
 					const FMassVelocityFragment& Velocity = VelocityList[EntityIndex];
 					const FMassForceFragment& Force = ForceList[EntityIndex];
 					const FMassMoveTargetFragment& MoveTarget = MoveTargetList[EntityIndex];
 					const FMassSimulationLODFragment& SimLOD = bHasLOD ? SimLODList[EntityIndex] : FMassSimulationLODFragment();
-					const FMassZoneGraphShortPathFragment& ShortPath = ShortPathList[EntityIndex];
 					const FMassStateTreeInstanceFragment& StateTreeInstance = StateTreeInstanceList[EntityIndex];
 
 					const FVector EntityForward = Transform.GetTransform().GetRotation().GetForwardVector();
@@ -650,8 +648,10 @@ void FGameplayDebuggerCategory_Mass::CollectData(APlayerController* OwnerPC, AAc
 					}
 
 					// Path
-					if (bShowNearEntityPath)
+					if (bShowNearEntityPath && ShortPathList.Num() > 0)
 					{
+						const FMassZoneGraphShortPathFragment& ShortPath = ShortPathList[EntityIndex];
+
 						const FVector ZOffset = FVector(0.0f , 0.0f , 25.0f );
 						for (uint8 PointIndex = 0; PointIndex < ShortPath.NumPoints - 1; PointIndex++)
 						{
@@ -668,9 +668,11 @@ void FGameplayDebuggerCategory_Mass::CollectData(APlayerController* OwnerPC, AAc
 							AddShape(FGameplayDebuggerShape::MakeSegment(CurrBase, CurrBase + CurrPoint.Tangent.GetVector() * 50.0f, 1.0f, FColorList::LightGrey));
 						}
 					}
-					
-					if (bShowNearEntityAvoidance)
+
+					if (bShowNearEntityAvoidance && GhostList.Num() > 0)
 					{
+						const FMassGhostLocationFragment& Ghost = GhostList[EntityIndex];
+
 						// Standing avoidance.
 						if (Ghost.IsValid(MoveTarget.GetCurrentActionID()))
 						{
diff --git a/Engine/Plugins/Runtime/MassGameplay/Source/MassLOD/Public/MassLODFragments.h b/Engine/Plugins/Runtime/MassGameplay/Source/MassLOD/Public/MassLODFragments.h
--- a/Engine/Plugins/Runtime/MassGameplay/Source/MassLOD/Public/MassLODFragments.h
+++ b/Engine/Plugins/Runtime/MassGameplay/Source/MassLOD/Public/MassLODFragments.h
@@ -138,7 +138,7 @@ struct MASSLOD_API FMassVisualizationChunkFragment : public FMassChunkFragment
 	GENERATED_BODY();
 
 	/**
-	 * Fetched existing chunk fragment to know if there is a possibility of a an entity that is visible
+	 * Fetched existing chunk fragment to know if there is a possibility of an entity that is visible.
 	 * In the case that there is not chunk information, we cannot assume that all entities are not visible.
 	 *
 	 * @param Context of the execution from the entity sub system
@@ -146,7 +146,8 @@ struct MASSLOD_API FMassVisualizationChunkFragment : public FMassChunkFragment
 	 */
 	static bool AreAnyEntitiesVisibleInChunk(const FMassExecutionContext& Context)
 	{
-		return Context.GetChunkFragment<FMassVisualizationChunkFragment>().AreAnyEntitiesVisible();
+		const FMassVisualizationChunkFragment* ChunkFragment = Context.GetChunkFragmentPtr<FMassVisualizationChunkFragment>();
+		return ChunkFragment == nullptr || ChunkFragment->AreAnyEntitiesVisible();
 	}
 
 	/**
