[MassAI] fixed typo

diff --git a/Engine/Plugins/AI/MassAI/Source/MassAIBehavior/Public/MassLookAtFragments.h b/Engine/Plugins/AI/MassAI/Source/MassAIBehavior/Public/MassLookAtFragments.h
--- a/Engine/Plugins/AI/MassAI/Source/MassAIBehavior/Public/MassLookAtFragments.h
+++ b/Engine/Plugins/AI/MassAI/Source/MassAIBehavior/Public/MassLookAtFragments.h
@@ -21,7 +21,7 @@ enum class EMassLookAtMode : uint8
 	LookAtEntity,
 };
 
-/** Gifferent gaze modes applied on top of the look at mode. */
+/** Different gaze modes applied on top of the look at mode. */
 UENUM()
 enum class EMassLookAtGazeMode : uint8
 {
