Mass navmesh navigation, first pass
Allows mass entities to find and follow a navmesh path using FMassNavMeshPathFollowTask.
Addition of a MassNavMeshNavigation module.
Addition of UMassNavMeshNavigationTrait.
Addition of a FMassNavMeshPathFollowTask to request a navmesh path and update the short path when needed.
Addition of FMassNavMeshShortPathFragment, a subpath used by UMassNavMeshPathFollowProcessor.
Addition of FMassFindSmartObjectTargetTask to set target location and associated intent.

diff --git a/Engine/Plugins/AI/MassAI/MassAI.uplugin b/Engine/Plugins/AI/MassAI/MassAI.uplugin
--- a/Engine/Plugins/AI/MassAI/MassAI.uplugin
+++ b/Engine/Plugins/AI/MassAI/MassAI.uplugin
@@ -3,7 +3,7 @@
 	"Version": 1,
 	"VersionName": "0.4",
 	"FriendlyName": "MassAI",
-	"Description": "AI-specific functionality extenting MassGameplay",
+	"Description": "AI-specific functionality extending MassGameplay",
 	"Category": "AI",
 	"CreatedBy": "Epic Games, Inc.",
 	"CreatedByURL": "https://epicgames.com",
@@ -20,6 +20,10 @@
 			"Name": "MassNavigation",
 			"Type": "Runtime"
 		},
+		{
+			"Name": "MassNavMeshNavigation",
+			"Type": "Runtime"
+		},
 		{
 			"Name": "MassZoneGraphNavigation",
 			"Type": "Runtime"
diff --git a/Engine/Plugins/AI/MassAI/Source/MassAIBehavior/MassAIBehavior.Build.cs b/Engine/Plugins/AI/MassAI/Source/MassAIBehavior/MassAIBehavior.Build.cs
--- a/Engine/Plugins/AI/MassAI/Source/MassAIBehavior/MassAIBehavior.Build.cs
+++ b/Engine/Plugins/AI/MassAI/Source/MassAIBehavior/MassAIBehavior.Build.cs
@@ -26,6 +26,7 @@ namespace UnrealBuildTool.Rules
 					"MassLOD",
 					"MassMovement",
 					"MassNavigation",
+					"MassNavMeshNavigation",
 					"MassZoneGraphNavigation",
 					"MassRepresentation",
 					"MassSignals",
diff --git a/Engine/Plugins/AI/MassAI/Source/MassAIBehavior/Private/Tasks/MassFindSmartObjectTargetTask.cpp b/Engine/Plugins/AI/MassAI/Source/MassAIBehavior/Private/Tasks/MassFindSmartObjectTargetTask.cpp
new file mode 100644
--- /dev/null
+++ b/Engine/Plugins/AI/MassAI/Source/MassAIBehavior/Private/Tasks/MassFindSmartObjectTargetTask.cpp
@@ -0,0 +1,38 @@
+// Copyright Epic Games, Inc. All Rights Reserved.
+
+#include "Tasks/MassFindSmartObjectTargetTask.h"
+
+#include "MassAIBehaviorTypes.h"
+#include "SmartObjectSubsystem.h"
+#include "MassSmartObjectFragments.h"
+#include "MassStateTreeExecutionContext.h"
+#include "StateTreeLinker.h"
+
+bool FMassFindSmartObjectTargetTask::Link(FStateTreeLinker& Linker)
+{
+	Linker.LinkExternalData(SmartObjectSubsystemHandle);
+	
+	return true;
+}
+
+EStateTreeRunStatus FMassFindSmartObjectTargetTask::EnterState(FStateTreeExecutionContext& Context, const FStateTreeTransitionResult& Transition) const
+{
+	const USmartObjectSubsystem& SmartObjectSubsystem = Context.GetExternalData(SmartObjectSubsystemHandle);
+
+	FInstanceDataType& InstanceData = Context.GetInstanceData(*this);
+
+	InstanceData.SmartObjectLocation.Reset();
+
+	if (!InstanceData.ClaimedSlot.SmartObjectHandle.IsValid())
+	{
+		MASSBEHAVIOR_LOG(Error, TEXT("Invalid claimed smart object ID."));
+		return EStateTreeRunStatus::Failed;
+	}
+
+	const FTransform Transform = SmartObjectSubsystem.GetSlotTransform(InstanceData.ClaimedSlot).Get(FTransform::Identity);
+
+	InstanceData.SmartObjectLocation.EndOfPathIntent = EMassMovementAction::Stand;
+	InstanceData.SmartObjectLocation.EndOfPathPosition = Transform.GetLocation();
+	
+	return EStateTreeRunStatus::Running;
+}
diff --git a/Engine/Plugins/AI/MassAI/Source/MassAIBehavior/Private/Tasks/MassNavMeshPathFollowTask.cpp b/Engine/Plugins/AI/MassAI/Source/MassAIBehavior/Private/Tasks/MassNavMeshPathFollowTask.cpp
new file mode 100644
--- /dev/null
+++ b/Engine/Plugins/AI/MassAI/Source/MassAIBehavior/Private/Tasks/MassNavMeshPathFollowTask.cpp
@@ -0,0 +1,180 @@
+// Copyright Epic Games, Inc. All Rights Reserved.
+
+#include "Tasks/MassNavMeshPathFollowTask.h"
+
+#include "MassAIBehaviorTypes.h"
+#include "MassCommonFragments.h"
+#include "MassDebugger.h"
+#include "MassMovementFragments.h"
+#include "MassNavigationFragments.h"
+#include "MassStateTreeExecutionContext.h"
+#include "NavigationSystem.h"
+#include "StateTreeLinker.h"
+
+bool FMassNavMeshPathFollowTask::Link(FStateTreeLinker& Linker)
+{
+	Linker.LinkExternalData(TransformHandle);
+	Linker.LinkExternalData(MoveTargetHandle);
+	Linker.LinkExternalData(AgentRadiusHandle);
+	Linker.LinkExternalData(MovementParamsHandle);
+	Linker.LinkExternalData(CachedPathHandle);
+	Linker.LinkExternalData(ShortPathHandle);
+
+	return true;
+}
+
+bool FMassNavMeshPathFollowTask::RequestPath(FStateTreeExecutionContext& Context, const FMassTargetLocation& TargetLocation) const
+{
+	FMassStateTreeExecutionContext& MassContext = static_cast<FMassStateTreeExecutionContext&>(Context);
+	FMassNavMeshPathFollowTaskInstanceData& InstanceData = Context.GetInstanceData<FMassNavMeshPathFollowTaskInstanceData>(*this);
+
+	bool bDisplayDebug = false;
+#if WITH_MASSGAMEPLAY_DEBUG
+	bDisplayDebug = UE::Mass::Debug::IsDebuggingEntity(MassContext.GetEntity());
+#endif // WITH_MASSGAMEPLAY_DEBUG
+	
+	const FAgentRadiusFragment& AgentRadius = Context.GetExternalData(AgentRadiusHandle);
+	const FVector AgentNavLocation = Context.GetExternalData(TransformHandle).GetTransform().GetLocation();
+	const FNavAgentProperties& NavAgentProperties = FNavAgentProperties(AgentRadius.Radius);
+
+	UNavigationSystemV1* NavMeshSubsystem = Cast<UNavigationSystemV1>(Context.GetWorld()->GetNavigationSystem());
+	const ANavigationData* NavData = NavMeshSubsystem->GetNavDataForProps(NavAgentProperties, AgentNavLocation);
+	
+	if (!NavData || !TargetLocation.EndOfPathPosition.IsSet())
+	{
+		return false;
+	}
+		
+	FPathFindingQuery Query(NavMeshSubsystem, *NavData, AgentNavLocation, InstanceData.TargetLocation.EndOfPathPosition.GetValue());
+
+	// Why fix it after if there is none??
+	if (!Query.NavData.IsValid())
+	{
+		Query.NavData = NavMeshSubsystem->GetNavDataForProps(NavAgentProperties, Query.StartLocation);
+	}
+
+	FPathFindingResult Result(ENavigationQueryResult::Error);
+	if (Query.NavData.IsValid())
+	{
+		if (bDisplayDebug)
+		{
+			MASSBEHAVIOR_LOG(Verbose, TEXT("requesting synchronous path"));
+		}
+		
+		Result = Query.NavData->FindPath(NavAgentProperties, Query);
+	}
+
+	if (Result.IsSuccessful() && Result.Path.Get()->GetPathPoints().Num() > 1)
+	{
+		// A path was found
+		if (bDisplayDebug)
+		{
+			MASSBEHAVIOR_LOG(Verbose, TEXT("path found"));
+		}
+
+		FMassNavMeshCachedPathFragment& CachedPathFragment = Context.GetExternalData(CachedPathHandle);
+		CachedPathFragment.NavPath = Result.Path;
+
+		// Update short path
+		const TArray<FNavPathPoint>& PathPoints = Result.Path.Get()->GetPathPoints();
+
+		FMassNavMeshShortPathFragment& ShortPathFragment = Context.GetExternalData(ShortPathHandle);
+		ShortPathFragment.RequestShortPath(PathPoints, 0);
+
+		CachedPathFragment.NavPathNextStartIndex += ShortPathFragment.NumPoints-1;
+		
+		// Update MoveTarget
+		FMassMoveTargetFragment& MoveTarget = Context.GetExternalData(MoveTargetHandle);
+		const FMassMovementParameters& MovementParams = Context.GetExternalData(MovementParamsHandle);
+		const float DesiredSpeed = FMath::Min(
+			MovementParams.GenerateDesiredSpeed(InstanceData.MovementStyle, MassContext.GetEntity().Index) * InstanceData.SpeedScale,
+			MovementParams.MaxSpeed);
+
+		MoveTarget.DesiredSpeed.Set(DesiredSpeed);
+		
+		MoveTarget.CreateNewAction(EMassMovementAction::Move, *Context.GetWorld());
+
+		return true;
+	}
+	
+	return false;
+}
+
+bool FMassNavMeshPathFollowTask::UpdateShortPath(FStateTreeExecutionContext& Context) const
+{
+	FMassNavMeshCachedPathFragment& CachedPathFragment = Context.GetExternalData(CachedPathHandle);
+
+	MASSBEHAVIOR_LOG(Verbose, TEXT("updating short path, starting at index %i"), CachedPathFragment.NavPathNextStartIndex);
+	
+	FMassNavMeshShortPathFragment& ShortPathFragment = Context.GetExternalData(ShortPathHandle);
+	ShortPathFragment.RequestShortPath(CachedPathFragment.NavPath.Get()->GetPathPoints(), CachedPathFragment.NavPathNextStartIndex);
+
+	CachedPathFragment.NavPathNextStartIndex += ShortPathFragment.NumPoints-1;
+
+	return true;
+}
+
+EStateTreeRunStatus FMassNavMeshPathFollowTask::EnterState(FStateTreeExecutionContext& Context, const FStateTreeTransitionResult& Transition) const
+{
+	FMassStateTreeExecutionContext& MassContext = static_cast<FMassStateTreeExecutionContext&>(Context);
+	FMassNavMeshPathFollowTaskInstanceData& InstanceData = Context.GetInstanceData<FMassNavMeshPathFollowTaskInstanceData>(*this);
+
+	bool bDisplayDebug = false;
+#if WITH_MASSGAMEPLAY_DEBUG
+	bDisplayDebug = UE::Mass::Debug::IsDebuggingEntity(MassContext.GetEntity());
+#endif // WITH_MASSGAMEPLAY_DEBUG
+	if (bDisplayDebug)
+	{
+		MASSBEHAVIOR_LOG(Verbose, TEXT("enterstate."));
+	}
+
+	if (!InstanceData.TargetLocation.EndOfPathPosition.IsSet())
+	{
+		MASSBEHAVIOR_LOG(Error, TEXT("Target is not defined."));
+		return EStateTreeRunStatus::Failed;
+	}
+		
+	if (!RequestPath(Context, InstanceData.TargetLocation))
+	{
+		MASSBEHAVIOR_LOG(Error, TEXT("Failed to request path."));
+		return EStateTreeRunStatus::Failed;
+	}
+
+	return EStateTreeRunStatus::Running;
+}
+
+EStateTreeRunStatus FMassNavMeshPathFollowTask::Tick(FStateTreeExecutionContext& Context, const float DeltaTime) const
+{
+	FMassStateTreeExecutionContext& MassContext = static_cast<FMassStateTreeExecutionContext&>(Context);
+
+	FMassNavMeshPathFollowTaskInstanceData& InstanceData = Context.GetInstanceData<FMassNavMeshPathFollowTaskInstanceData>(*this);
+
+	FMassNavMeshShortPathFragment& ShortPathFragment = Context.GetExternalData(ShortPathHandle);
+	
+	bool bDisplayDebug = false;
+#if WITH_MASSGAMEPLAY_DEBUG
+	bDisplayDebug = UE::Mass::Debug::IsDebuggingEntity(MassContext.GetEntity());
+#endif // WITH_MASSGAMEPLAY_DEBUG
+	if (bDisplayDebug)
+	{
+		MASSBEHAVIOR_LOG(Verbose, TEXT("tick"));
+	}
+	
+	// Current path follow is done, but it was partial (i.e. many points on a curve), try again until we get there.
+	if (ShortPathFragment.IsDone() && ShortPathFragment.bPartialResult)
+	{
+		if (!InstanceData.TargetLocation.EndOfPathPosition.IsSet())
+		{
+			MASSBEHAVIOR_LOG(Error, TEXT("Target is not defined."));
+			return EStateTreeRunStatus::Failed;
+		}
+
+		if (!UpdateShortPath(Context))
+		{
+			MASSBEHAVIOR_LOG(Error, TEXT("Failed to update short path."));
+			return EStateTreeRunStatus::Failed;
+		}
+	}
+	
+	return ShortPathFragment.IsDone() ? EStateTreeRunStatus::Succeeded : EStateTreeRunStatus::Running;
+}
diff --git a/Engine/Plugins/AI/MassAI/Source/MassAIBehavior/Public/Tasks/MassFindSmartObjectTargetTask.h b/Engine/Plugins/AI/MassAI/Source/MassAIBehavior/Public/Tasks/MassFindSmartObjectTargetTask.h
new file mode 100644
--- /dev/null
+++ b/Engine/Plugins/AI/MassAI/Source/MassAIBehavior/Public/Tasks/MassFindSmartObjectTargetTask.h
@@ -0,0 +1,38 @@
+// Copyright Epic Games, Inc. All Rights Reserved.
+
+#pragma once
+
+#include "MassNavMeshPathFollowTask.h"
+#include "MassStateTreeTypes.h"
+#include "SmartObjectRuntime.h"
+#include "MassFindSmartObjectTargetTask.generated.h"
+
+class USmartObjectSubsystem;
+
+USTRUCT()
+struct MASSAIBEHAVIOR_API FMassFindSmartObjectTargetInstanceData
+{
+	GENERATED_BODY()
+
+	UPROPERTY(VisibleAnywhere, Category = Input)
+	FSmartObjectClaimHandle ClaimedSlot;
+
+	UPROPERTY(EditAnywhere, Category = Output)
+	FMassTargetLocation SmartObjectLocation;
+};
+
+/** Computes move target to a smart object based on current location. */
+USTRUCT(meta = (DisplayName = "Find Smart Object Target"))
+struct MASSAIBEHAVIOR_API FMassFindSmartObjectTargetTask : public FMassStateTreeTaskBase
+{
+	GENERATED_BODY()
+
+	using FInstanceDataType = FMassFindSmartObjectTargetInstanceData;
+
+protected:
+	virtual bool Link(FStateTreeLinker& Linker) override;
+	virtual const UStruct* GetInstanceDataType() const override { return FInstanceDataType::StaticStruct(); }
+	virtual EStateTreeRunStatus EnterState(FStateTreeExecutionContext& Context, const FStateTreeTransitionResult& Transition) const override;
+
+	TStateTreeExternalDataHandle<USmartObjectSubsystem> SmartObjectSubsystemHandle;
+};
diff --git a/Engine/Plugins/AI/MassAI/Source/MassAIBehavior/Public/Tasks/MassNavMeshPathfollowTask.h b/Engine/Plugins/AI/MassAI/Source/MassAIBehavior/Public/Tasks/MassNavMeshPathfollowTask.h
new file mode 100644
--- /dev/null
+++ b/Engine/Plugins/AI/MassAI/Source/MassAIBehavior/Public/Tasks/MassNavMeshPathfollowTask.h
@@ -0,0 +1,77 @@
+// Copyright Epic Games, Inc. All Rights Reserved.
+
+#pragma once
+
+#include "MassMovementTypes.h"
+#include "MassNavigationTypes.h"
+#include "MassNavMeshNavigationFragments.h"
+#include "MassStateTreeExecutionContext.h"
+#include "MassStateTreeTypes.h"
+#include "MassNavMeshPathfollowTask.generated.h"
+
+struct FTransformFragment;
+struct FMassMoveTargetFragment;
+struct FAgentRadiusFragment;
+struct FMassMovementParameters;
+
+/** Target location used by FMassNavMeshPathFollowTask */
+USTRUCT()
+struct MASSAIBEHAVIOR_API FMassTargetLocation
+{
+	GENERATED_BODY()
+
+	void Reset()
+	{
+		EndOfPathPosition.Reset();
+		EndOfPathIntent = EMassMovementAction::Move;
+	}
+	
+	/** Optional end of path location. */
+	TOptional<FVector> EndOfPathPosition;
+
+	/** Movement intent at the end of the path */
+	EMassMovementAction EndOfPathIntent = EMassMovementAction::Move;
+};
+
+/** FMassNavMeshPathFollowTask movement parameters */
+USTRUCT()
+struct MASSAIBEHAVIOR_API FMassNavMeshPathFollowTaskInstanceData
+{ 
+	GENERATED_BODY()
+	
+	UPROPERTY(EditAnywhere, Category = Input)
+	FMassTargetLocation TargetLocation;
+
+	UPROPERTY(EditAnywhere, Category = Parameter)
+	FMassMovementStyleRef MovementStyle;
+
+	UPROPERTY(EditAnywhere, Category = Parameter)
+	float SpeedScale = 1.f;
+};
+
+/** Finds a path to TargetLocation, requests a short path, starts a move action and follow the path by updating the short path when needed. */ 
+USTRUCT(meta = (DisplayName = "NavMesh Path Follow"))
+struct MASSAIBEHAVIOR_API FMassNavMeshPathFollowTask : public FMassStateTreeTaskBase
+{
+	GENERATED_BODY()
+
+protected:
+	virtual bool Link(FStateTreeLinker& Linker) override;
+	virtual const UStruct* GetInstanceDataType() const override { return FMassNavMeshPathFollowTaskInstanceData::StaticStruct(); };
+
+	bool RequestPath(FStateTreeExecutionContext& Context, const FMassTargetLocation& TargetLocation) const;
+	bool UpdateShortPath(FStateTreeExecutionContext& Context) const;
+
+	virtual EStateTreeRunStatus EnterState(FStateTreeExecutionContext& Context, const FStateTreeTransitionResult& Transition) const override;
+	virtual EStateTreeRunStatus Tick(FStateTreeExecutionContext& Context, const float DeltaTime) const override;
+
+	TStateTreeExternalDataHandle<FTransformFragment> TransformHandle;
+	TStateTreeExternalDataHandle<FMassMoveTargetFragment> MoveTargetHandle;
+	TStateTreeExternalDataHandle<FAgentRadiusFragment> AgentRadiusHandle;
+	TStateTreeExternalDataHandle<FMassMovementParameters> MovementParamsHandle;
+
+	// Hold a small part of a navmesh path
+	TStateTreeExternalDataHandle<FMassNavMeshShortPathFragment> ShortPathHandle;
+
+	TStateTreeExternalDataHandle<FMassNavMeshCachedPathFragment> CachedPathHandle;
+};
diff --git a/Engine/Plugins/AI/MassAI/Source/MassNavMeshNavigation/MassNavMeshNavigation.Build.cs b/Engine/Plugins/AI/MassAI/Source/MassNavMeshNavigation/MassNavMeshNavigation.Build.cs
new file mode 100644
--- /dev/null
+++ b/Engine/Plugins/AI/MassAI/Source/MassNavMeshNavigation/MassNavMeshNavigation.Build.cs
@@ -0,0 +1,28 @@
+﻿// Copyright Epic Games, Inc. All Rights Reserved.
+
+namespace UnrealBuildTool.Rules
+{
+	public class MassNavMeshNavigation : ModuleRules
+	{
+		public MassNavMeshNavigation(ReadOnlyTargetRules Target) : base(Target)
+		{
+			PCHUsage = ModuleRules.PCHUsageMode.UseExplicitOrSharedPCHs;
+			
+			PublicDependencyModuleNames.AddRange(
+				new string[] {
+					"AIModule",
+					"MassEntity",
+					"Core",
+					"CoreUObject",
+					"Engine",
+					"MassCommon",
+					"MassLOD",
+					"MassSignals",
+					"MassSpawner",
+					"MassMovement",
+					"MassNavigation"
+				}
+			);
+		}
+	}
+}
\ No newline at end of file
diff --git a/Engine/Plugins/AI/MassAI/Source/MassNavMeshNavigation/Private/MassNavMeshNavigationFragments.cpp b/Engine/Plugins/AI/MassAI/Source/MassNavMeshNavigation/Private/MassNavMeshNavigationFragments.cpp
new file mode 100644
--- /dev/null
+++ b/Engine/Plugins/AI/MassAI/Source/MassNavMeshNavigation/Private/MassNavMeshNavigationFragments.cpp
@@ -0,0 +1,47 @@
+﻿// Copyright Epic Games, Inc. All Rights Reserved.
+
+#include "MassNavMeshNavigationFragments.h"
+
+#include "AI/Navigation/NavigationTypes.h"
+
+// Make a new short path
+bool FMassNavMeshShortPathFragment::RequestShortPath(const TArray<FNavPathPoint>& NavPathPoints, const int32 NavPathStartIndex)
+{
+	Reset();
+
+	// Copy PathPoints from the StartIndex
+	for(; NumPoints < MaxPoints && (NavPathStartIndex + NumPoints) < NavPathPoints.Num(); NumPoints++)
+	{
+		FMassNavMeshPathPoint& Point = Points[NumPoints];
+		Point.Position = NavPathPoints[NavPathStartIndex + NumPoints];
+	}
+
+	checkf(NumPoints >= 2, TEXT("Path should have at least 2 points at this stage, has %d."), NumPoints);
+
+	// We are not done with the main path, mark as partial.
+	const int32 LastNavPathIndex = NavPathPoints.Num()-1;
+	if (NavPathStartIndex + (NumPoints-1) < LastNavPathIndex)
+	{
+		bPartialResult = true;
+	}
+
+	// Calculate movement distance at each point.
+	float PathDistance = 0.f;
+	Points[0].Distance.Set(PathDistance);
+	for (uint8 PointIndex = 1; PointIndex < NumPoints; PointIndex++)
+	{
+		FMassNavMeshPathPoint& PrevPoint = Points[PointIndex - 1];
+		FMassNavMeshPathPoint& Point = Points[PointIndex];
+		const FVector& PrevPosition = PrevPoint.Position;
+		const FVector& Position = Point.Position;
+		const float DeltaDistance = FloatCastChecked<float>(FVector::Dist(PrevPosition, Position), UE::LWC::DefaultFloatPrecision);
+		PathDistance += DeltaDistance;
+		Point.Distance.Set(PathDistance);
+
+		// Set tangent of PrevPoint
+		const FVector& Tangent = (Point.Position - PrevPoint.Position).GetSafeNormal();
+		PrevPoint.Tangent = FMassSnorm8Vector2D(Tangent);
+	}
+	
+	return true;
+}
diff --git a/Engine/Plugins/AI/MassAI/Source/MassNavMeshNavigation/Private/MassNavMeshNavigationModule.cpp b/Engine/Plugins/AI/MassAI/Source/MassNavMeshNavigation/Private/MassNavMeshNavigationModule.cpp
new file mode 100644
--- /dev/null
+++ b/Engine/Plugins/AI/MassAI/Source/MassNavMeshNavigation/Private/MassNavMeshNavigationModule.cpp
@@ -0,0 +1,25 @@
+﻿// Copyright Epic Games, Inc. All Rights Reserved.
+
+#include "IMassNavMeshNavigationModule.h"
+
+class FMassNavMeshNavigationModule : public IMassNavMeshNavigationModule
+{
+	/** IModuleInterface implementation */
+	virtual void StartupModule() override;
+	virtual void ShutdownModule() override;
+};
+
+IMPLEMENT_MODULE(FMassNavMeshNavigationModule, MassNavMeshNavigation)
+
+
+void FMassNavMeshNavigationModule::StartupModule()
+{
+	// This code will execute after your module is loaded into memory (but after global variables are initialized, of course.)
+}
+
+
+void FMassNavMeshNavigationModule::ShutdownModule()
+{
+	// This function may be called during shutdown to clean up your module.  For modules that support dynamic reloading,
+	// we call this function before unloading the module.
+}
diff --git a/Engine/Plugins/AI/MassAI/Source/MassNavMeshNavigation/Private/MassNavMeshNavigationProcessors.cpp b/Engine/Plugins/AI/MassAI/Source/MassNavMeshNavigation/Private/MassNavMeshNavigationProcessors.cpp
new file mode 100644
--- /dev/null
+++ b/Engine/Plugins/AI/MassAI/Source/MassNavMeshNavigation/Private/MassNavMeshNavigationProcessors.cpp
@@ -0,0 +1,198 @@
+﻿// Copyright Epic Games, Inc. All Rights Reserved.
+
+#include "MassNavMeshNavigationProcessors.h"
+#include "Avoidance/MassAvoidanceFragments.h"
+#include "MassCommonTypes.h"
+#include "MassDebugger.h"
+#include "MassExecutionContext.h"
+#include "MassNavMeshNavigationFragments.h"
+#include "MassNavMeshNavigationTypes.h"
+#include "MassNavigationFragments.h"
+#include "MassSignalSubsystem.h"
+#include "MassSimulationLOD.h"
+
+#define UNSAFE_FOR_MT 1
+
+UMassNavMeshPathFollowProcessor::UMassNavMeshPathFollowProcessor()
+	: EntityQuery_Conditional(*this)
+{
+	ExecutionFlags = (int32)EProcessorExecutionFlags::AllNetModes;
+	ExecutionOrder.ExecuteInGroup = UE::Mass::ProcessorGroupNames::Tasks;
+	ExecutionOrder.ExecuteBefore.Add(UE::Mass::ProcessorGroupNames::Avoidance);
+}
+
+void UMassNavMeshPathFollowProcessor::InitializeInternal(UObject& Owner, const TSharedRef<FMassEntityManager>& EntityManager)
+{
+	Super::InitializeInternal(Owner, EntityManager);
+	SignalSubsystem = UWorld::GetSubsystem<UMassSignalSubsystem>(Owner.GetWorld());
+}
+
+void UMassNavMeshPathFollowProcessor::ConfigureQueries(const TSharedRef<FMassEntityManager>& EntityManager)
+{
+	EntityQuery_Conditional.AddRequirement<FMassNavMeshShortPathFragment>(EMassFragmentAccess::ReadWrite);
+	EntityQuery_Conditional.AddRequirement<FMassMoveTargetFragment>(EMassFragmentAccess::ReadWrite);
+
+	// @todo: validate LOD and variable ticking
+	EntityQuery_Conditional.AddRequirement<FMassSimulationLODFragment>(EMassFragmentAccess::ReadOnly, EMassFragmentPresence::Optional);
+	EntityQuery_Conditional.AddRequirement<FMassSimulationVariableTickFragment>(EMassFragmentAccess::ReadOnly, EMassFragmentPresence::Optional);
+}
+
+void UMassNavMeshPathFollowProcessor::Execute(FMassEntityManager& EntityManager, FMassExecutionContext& Context)
+{
+	if (!SignalSubsystem)
+	{
+		return;
+	}
+	
+	TArray<FMassEntityHandle> EntitiesToSignalPathDone;
+
+	EntityQuery_Conditional.ForEachEntityChunk(EntityManager, Context, [this, &EntitiesToSignalPathDone](FMassExecutionContext& Context)
+	{
+		const int32 NumEntities = Context.GetNumEntities();
+		const TArrayView<FMassNavMeshShortPathFragment> ShortPathList = Context.GetMutableFragmentView<FMassNavMeshShortPathFragment>();
+		const TArrayView<FMassMoveTargetFragment> MoveTargetList = Context.GetMutableFragmentView<FMassMoveTargetFragment>();
+
+		// @todo: validate LOD and variable ticking
+		const TConstArrayView<FMassSimulationLODFragment> SimLODList = Context.GetFragmentView<FMassSimulationLODFragment>();
+		const bool bHasLOD = (SimLODList.Num() > 0);
+
+		const TConstArrayView<FMassSimulationVariableTickFragment> SimVariableTickList = Context.GetFragmentView<FMassSimulationVariableTickFragment>();
+		const bool bHasVariableTick = (SimVariableTickList.Num() > 0);
+		const float WorldDeltaTime = Context.GetDeltaTimeSeconds();
+
+		for (int32 EntityIndex = 0; EntityIndex < NumEntities; ++EntityIndex)
+		{
+			FMassNavMeshShortPathFragment& ShortPath = ShortPathList[EntityIndex];
+			FMassMoveTargetFragment& MoveTarget = MoveTargetList[EntityIndex];
+			const FMassEntityHandle Entity = Context.GetEntity(EntityIndex);
+			const float DeltaTime = bHasVariableTick ? SimVariableTickList[EntityIndex].DeltaTime : WorldDeltaTime;
+
+			bool bDisplayDebug = false;
+#if WITH_MASSGAMEPLAY_DEBUG && UNSAFE_FOR_MT // this will result in bDisplayDebug == false and disabling of all the vlogs below
+			bDisplayDebug = UE::Mass::Debug::IsDebuggingEntity(Entity);
+#endif // WITH_MASSGAMEPLAY_DEBUG
+
+			// Must have at least two points to interpolate.
+			if (MoveTarget.GetCurrentAction() == EMassMovementAction::Move && ShortPath.NumPoints >= 2)
+			{
+				UE_CVLOG_UELOG(bDisplayDebug, this, LogMassNavMeshNavigation, Verbose, TEXT("Entity [%s] Updating navmesh path following"), *Entity.DebugGetDescription());
+				
+				const bool bWasDone = ShortPath.IsDone();
+
+				// Note: this should be in sync with the logic in apply velocity.
+				const bool bHasSteering = (bHasLOD == false) || (SimLODList[EntityIndex].LOD != EMassLOD::Off);
+
+				if (!bHasSteering || !MoveTarget.bSteeringFallingBehind)
+				{
+					// Update progress
+					ShortPath.ProgressDistance += MoveTarget.DesiredSpeed.Get() * DeltaTime;
+				}
+				
+				if (!bWasDone)
+				{
+					const uint8 LastPointIndex = ShortPath.NumPoints - 1;
+
+					if (ShortPath.ProgressDistance <= 0.0f)
+					{
+						// Requested time before the start of the path
+						MoveTarget.Center = ShortPath.Points[0].Position;
+						MoveTarget.Forward = ShortPath.Points[0].Tangent.GetVector();
+						MoveTarget.DistanceToGoal = ShortPath.Points[LastPointIndex].Distance.Get();
+
+						UE_CVLOG_UELOG(bDisplayDebug, this, LogMassNavMeshNavigation, Log, TEXT("Entity [%s]   before start of path. Distance to goal: %.1f."),
+							*Entity.DebugGetDescription(),
+							MoveTarget.DistanceToGoal);
+					}
+					else if (ShortPath.ProgressDistance <= ShortPath.Points[LastPointIndex].Distance.Get())
+					{
+						// Requested time along the path, interpolate.
+						uint8 PointIndex = 0;
+						while (PointIndex < (ShortPath.NumPoints - 2))
+						{
+							const FMassNavMeshPathPoint& NextPoint = ShortPath.Points[PointIndex + 1];
+							if (ShortPath.ProgressDistance <= NextPoint.Distance.Get())
+							{
+								break;
+							}
+							PointIndex++;
+						}
+
+						const FMassNavMeshPathPoint& CurrPoint = ShortPath.Points[PointIndex];
+						const FMassNavMeshPathPoint& NextPoint = ShortPath.Points[PointIndex + 1];
+						const float T = (ShortPath.ProgressDistance - CurrPoint.Distance.Get()) / (NextPoint.Distance.Get() - CurrPoint.Distance.Get());
+
+						// Set move target location using the new progress distance.
+						MoveTarget.Center = FMath::Lerp(CurrPoint.Position, NextPoint.Position, T);
+						MoveTarget.Forward = FMath::Lerp(CurrPoint.Tangent.GetVector(), NextPoint.Tangent.GetVector(), T).GetSafeNormal();
+						MoveTarget.DistanceToGoal = ShortPath.Points[LastPointIndex].Distance.Get() - FMath::Lerp(CurrPoint.Distance.Get(), NextPoint.Distance.Get(), T);
+
+						UE_CVLOG_UELOG(bDisplayDebug, this, LogMassNavMeshNavigation, Verbose, TEXT("Entity [%s]   Distance to goal on short path: %.1f."),
+							*Entity.DebugGetDescription(),
+							MoveTarget.DistanceToGoal);
+					}
+					else
+					{
+						// @todo: probably need to do this earlier to make sure to avoid any discontinuity when updating the short path.
+						
+						// End of short path reached.
+						MoveTarget.Center = ShortPath.Points[LastPointIndex].Position;
+						MoveTarget.Forward = ShortPath.Points[LastPointIndex].Tangent.GetVector();
+						MoveTarget.DistanceToGoal = 0.f;
+
+						UE_CVLOG_UELOG(bDisplayDebug, this, LogMassNavMeshNavigation, Log, TEXT("Entity [%s]   Finished path follow on short path."), *Entity.DebugGetDescription());
+
+						ShortPath.bDone = true;
+					}
+				}
+
+				const bool bIsDone = ShortPath.IsDone();
+
+				// Signal path done.
+				if (!bWasDone && bIsDone)
+				{
+					EntitiesToSignalPathDone.Add(Entity);
+				}
+
+#if WITH_MASSGAMEPLAY_DEBUG && UNSAFE_FOR_MT
+				if (bDisplayDebug)
+				{
+					const FColor EntityColor = UE::Mass::Debug::GetEntityDebugColor(Entity);
+
+					const FVector ZOffset(0,0,10);
+					FLinearColor MixColor(EntityColor);
+					MixColor += FColor::White;
+					MixColor /= 2;
+					const FColor LightEntityColor = MixColor.ToFColorSRGB();
+
+					// Draw path
+					for (uint8 PointIndex = 0; PointIndex < ShortPath.NumPoints - 1; PointIndex++)
+					{
+						const FMassNavMeshPathPoint& CurrPoint = ShortPath.Points[PointIndex];
+						const FMassNavMeshPathPoint& NextPoint = ShortPath.Points[PointIndex + 1];
+						UE_VLOG_SEGMENT_THICK(this, LogMassNavMeshNavigation, Display, CurrPoint.Position + ZOffset, NextPoint.Position + ZOffset, EntityColor, /*Thickness*/3, TEXT(""));
+					}
+
+					// Draw tangents
+					for (uint8 PointIndex = 0; PointIndex < ShortPath.NumPoints; PointIndex++)
+					{
+						const FMassNavMeshPathPoint& CurrPoint = ShortPath.Points[PointIndex];
+						const FVector CurrBase = CurrPoint.Position + ZOffset + FVector(0,0,1);
+						UE_VLOG_SEGMENT_THICK(this, LogMassNavMeshNavigation, Display, CurrBase, CurrBase + CurrPoint.Tangent.GetVector() * 100.0f, LightEntityColor, /*Thickness*/1, TEXT(""));
+					}
+
+					// Draw MoveTarget
+					constexpr float Radius = 20.f;
+					UE_VLOG_CIRCLE(this, LogMassNavMeshNavigation, Display, MoveTarget.Center + ZOffset, FVector(0,0,1), Radius, LightEntityColor, TEXT(""));
+				}
+#endif // WITH_MASSGAMEPLAY_DEBUG
+			}
+		}
+	});
+
+	if (EntitiesToSignalPathDone.Num())
+	{
+		SignalSubsystem->SignalEntities(UE::Mass::Signals::FollowPointPathDone, EntitiesToSignalPathDone);
+	}
+}
+
+#undef UNSAFE_FOR_MT
\ No newline at end of file
diff --git a/Engine/Plugins/AI/MassAI/Source/MassNavMeshNavigation/Private/MassNavMeshNavigationTrait.cpp b/Engine/Plugins/AI/MassAI/Source/MassNavMeshNavigation/Private/MassNavMeshNavigationTrait.cpp
new file mode 100644
--- /dev/null
+++ b/Engine/Plugins/AI/MassAI/Source/MassNavMeshNavigation/Private/MassNavMeshNavigationTrait.cpp
@@ -0,0 +1,11 @@
+// Copyright Epic Games, Inc. All Rights Reserved.
+#include "MassNavMeshNavigationTrait.h"
+#include "MassEntityTemplateRegistry.h"
+#include "MassEntityUtils.h"
+#include "MassNavMeshNavigationFragments.h"
+
+void UMassNavMeshNavigationTrait::BuildTemplate(FMassEntityTemplateBuildContext& BuildContext, const UWorld& World) const
+{
+	BuildContext.AddFragment<FMassNavMeshCachedPathFragment>();
+	BuildContext.AddFragment<FMassNavMeshShortPathFragment>();
+}
diff --git a/Engine/Plugins/AI/MassAI/Source/MassNavMeshNavigation/Private/MassNavMeshNavigationTypes.cpp b/Engine/Plugins/AI/MassAI/Source/MassNavMeshNavigation/Private/MassNavMeshNavigationTypes.cpp
new file mode 100644
--- /dev/null
+++ b/Engine/Plugins/AI/MassAI/Source/MassNavMeshNavigation/Private/MassNavMeshNavigationTypes.cpp
@@ -0,0 +1,5 @@
+﻿// Copyright Epic Games, Inc. All Rights Reserved.
+
+#include "MassNavMeshNavigationTypes.h"
+
+DEFINE_LOG_CATEGORY(LogMassNavMeshNavigation);
diff --git a/Engine/Plugins/AI/MassAI/Source/MassNavMeshNavigation/Public/IMassNavMeshNavigationModule.h b/Engine/Plugins/AI/MassAI/Source/MassNavMeshNavigation/Public/IMassNavMeshNavigationModule.h
new file mode 100644
--- /dev/null
+++ b/Engine/Plugins/AI/MassAI/Source/MassNavMeshNavigation/Public/IMassNavMeshNavigationModule.h
@@ -0,0 +1,39 @@
+﻿// Copyright Epic Games, Inc. All Rights Reserved.
+
+#pragma once
+
+#include "Modules/ModuleInterface.h"
+#include "Modules/ModuleManager.h"
+
+
+/**
+* The public interface to this module.  In most cases, this interface is only public to sibling modules 
+* within this plugin.
+*/
+class IMassNavMeshNavigationModule : public IModuleInterface
+{
+
+public:
+
+	/**
+	* Singleton-like access to this module's interface.  This is just for convenience!
+	* Beware of calling this during the shutdown phase, though.  Your module might have been unloaded already.
+	*
+	* @return Returns singleton instance, loading the module on demand if needed
+	*/
+	static inline IMassNavMeshNavigationModule& Get()
+	{
+		return FModuleManager::LoadModuleChecked<IMassNavMeshNavigationModule>( "MassNavMeshNavigation" );
+	}
+
+	/**
+	* Checks to see if this module is loaded and ready.  It is only valid to call Get() if IsAvailable() returns true.
+	*
+	* @return True if the module is loaded and ready to use
+	*/
+	static inline bool IsAvailable()
+	{
+		return FModuleManager::Get().IsModuleLoaded( "MassNavMeshNavigation" );
+	}
+};
+
diff --git a/Engine/Plugins/AI/MassAI/Source/MassNavMeshNavigation/Public/MassNavMeshNavigationFragments.h b/Engine/Plugins/AI/MassAI/Source/MassNavMeshNavigation/Public/MassNavMeshNavigationFragments.h
new file mode 100644
--- /dev/null
+++ b/Engine/Plugins/AI/MassAI/Source/MassNavMeshNavigation/Public/MassNavMeshNavigationFragments.h
@@ -0,0 +1,86 @@
+﻿// Copyright Epic Games, Inc. All Rights Reserved.
+
+#pragma once
+
+#include "MassEntityTypes.h"
+#include "MassCommonTypes.h"
+#include "MassNavigationTypes.h"
+#include "AI/Navigation/NavigationTypes.h"
+#include "Containers/StaticArray.h"
+#include "MassNavMeshNavigationFragments.generated.h"
+
+/** Navmesh path point data */
+USTRUCT()
+struct MASSNAVMESHNAVIGATION_API FMassNavMeshPathPoint
+{
+	GENERATED_BODY()
+
+	/** Position of the path. */
+	FVector Position = FVector::ZeroVector;
+
+	/** Tangent direction of the path. */
+	FMassSnorm8Vector2D Tangent;
+
+	/** Distance along the path from first point. */
+	FMassInt16Real Distance = FMassInt16Real(0.f);
+};
+
+/** Short path used for navmesh navigation */
+USTRUCT()
+struct MASSNAVMESHNAVIGATION_API FMassNavMeshShortPathFragment : public FMassFragment
+{
+	GENERATED_BODY()
+
+	FMassNavMeshShortPathFragment() = default;
+	
+	static constexpr uint8 MaxPoints = 4;
+
+	void Reset()
+	{
+		ProgressDistance = 0.f;
+		NumPoints = 0;
+		EndOfPathIntent = EMassMovementAction::Stand;
+		bPartialResult = false;
+		bDone = false;
+	}
+
+	/** Returns true if we are done progressing on the short path. */
+	bool IsDone() const
+	{
+		return NumPoints == 0 || bDone;
+	}
+
+	/** Fill Points from an array of FNavPathPoint. */
+	bool RequestShortPath(const TArray<FNavPathPoint>& NavPathPoints, const int32 NavPathStartIndex);
+	
+	/** Path points */
+	TStaticArray<FMassNavMeshPathPoint, MaxPoints> Points;
+
+	/** Current progress distance along the short path. */
+	float ProgressDistance = 0.f;
+	
+	/** Number of points on path. */
+	uint8 NumPoints = 0;
+	
+	/** Intent at the end of the path. */
+	EMassMovementAction EndOfPathIntent = EMassMovementAction::Stand;
+
+	/** True if the path was partial. */
+	uint8 bPartialResult : 1 = false;
+
+	/** True when path follow is completed. */
+	uint8 bDone : 1 = false;
+};
+
+/** Current navmesh path */
+USTRUCT()
+struct MASSNAVMESHNAVIGATION_API FMassNavMeshCachedPathFragment : public FMassFragment
+{
+	GENERATED_BODY()
+
+	/** Reference to a FNavigationPath. */
+	FNavPathSharedPtr NavPath;
+
+	/** Index used to keep track of progression on the navmesh path. */
+	uint16 NavPathNextStartIndex = 0;
+};
diff --git a/Engine/Plugins/AI/MassAI/Source/MassNavMeshNavigation/Public/MassNavMeshNavigationProcessors.h b/Engine/Plugins/AI/MassAI/Source/MassNavMeshNavigation/Public/MassNavMeshNavigationProcessors.h
new file mode 100644
--- /dev/null
+++ b/Engine/Plugins/AI/MassAI/Source/MassNavMeshNavigation/Public/MassNavMeshNavigationProcessors.h
@@ -0,0 +1,27 @@
+﻿// Copyright Epic Games, Inc. All Rights Reserved.
+
+#pragma once
+
+#include "MassProcessor.h"
+#include "MassNavMeshNavigationProcessors.generated.h"
+
+class UMassSignalSubsystem;
+
+/** Processor for updating move target on a navmesh short path. */
+UCLASS()
+class MASSNAVMESHNAVIGATION_API UMassNavMeshPathFollowProcessor : public UMassProcessor
+{
+	GENERATED_BODY()
+
+protected:
+	UMassNavMeshPathFollowProcessor();
+
+	virtual void ConfigureQueries(const TSharedRef<FMassEntityManager>& EntityManager) override;
+	virtual void InitializeInternal(UObject& Owner, const TSharedRef<FMassEntityManager>& ) override;
+	virtual void Execute(FMassEntityManager& EntityManager, FMassExecutionContext& Context) override;
+
+	FMassEntityQuery EntityQuery_Conditional;
+
+	UPROPERTY(Transient)
+	TObjectPtr<UMassSignalSubsystem> SignalSubsystem = nullptr;
+};
diff --git a/Engine/Plugins/AI/MassAI/Source/MassNavMeshNavigation/Public/MassNavMeshNavigationTrait.h b/Engine/Plugins/AI/MassAI/Source/MassNavMeshNavigation/Public/MassNavMeshNavigationTrait.h
new file mode 100644
--- /dev/null
+++ b/Engine/Plugins/AI/MassAI/Source/MassNavMeshNavigation/Public/MassNavMeshNavigationTrait.h
@@ -0,0 +1,15 @@
+// Copyright Epic Games, Inc. All Rights Reserved.
+
+#pragma once 
+
+#include "MassEntityTraitBase.h"
+#include "MassNavMeshNavigationTrait.generated.h"
+
+UCLASS(meta = (DisplayName = "NavMesh Navigation"))
+class MASSNAVMESHNAVIGATION_API UMassNavMeshNavigationTrait : public UMassEntityTraitBase
+{
+	GENERATED_BODY()
+
+protected:
+	virtual void BuildTemplate(FMassEntityTemplateBuildContext& BuildContext, const UWorld& World) const override;
+};
diff --git a/Engine/Plugins/AI/MassAI/Source/MassNavMeshNavigation/Public/MassNavMeshNavigationTypes.h b/Engine/Plugins/AI/MassAI/Source/MassNavMeshNavigation/Public/MassNavMeshNavigationTypes.h
new file mode 100644
--- /dev/null
+++ b/Engine/Plugins/AI/MassAI/Source/MassNavMeshNavigation/Public/MassNavMeshNavigationTypes.h
@@ -0,0 +1,7 @@
+﻿// Copyright Epic Games, Inc. All Rights Reserved.
+
+#pragma once
+
+#include "Logging/LogMacros.h"
+
+MASSNAVMESHNAVIGATION_API DECLARE_LOG_CATEGORY_EXTERN(LogMassNavMeshNavigation, Warning, All);
