Temtative CIS fix

#rnx

diff --git a/Engine/Plugins/AI/MassAI/Source/MassAIBehavior/Public/Tasks/MassFindSmartObjectTargetTask.h b/Engine/Plugins/AI/MassAI/Source/MassAIBehavior/Public/Tasks/MassFindSmartObjectTargetTask.h
--- a/Engine/Plugins/AI/MassAI/Source/MassAIBehavior/Public/Tasks/MassFindSmartObjectTargetTask.h
+++ b/Engine/Plugins/AI/MassAI/Source/MassAIBehavior/Public/Tasks/MassFindSmartObjectTargetTask.h
@@ -2,7 +2,7 @@
 
 #pragma once
 
-#include "MassNavMeshPathFollowTask.h"
+#include "Tasks/MassNavMeshPathFollowTask.h"
 #include "MassStateTreeTypes.h"
 #include "SmartObjectRuntime.h"
 #include "MassFindSmartObjectTargetTask.generated.h"
