CIS fix

#rnx

diff --git a/Engine/Plugins/AI/MassAI/Source/MassAIBehavior/Private/Tasks/MassNavMeshPathFollowTask.cpp b/Engine/Plugins/AI/MassAI/Source/MassAIBehavior/Private/Tasks/MassNavMeshPathFollowTask.cpp
--- a/Engine/Plugins/AI/MassAI/Source/MassAIBehavior/Private/Tasks/MassNavMeshPathFollowTask.cpp
+++ b/Engine/Plugins/AI/MassAI/Source/MassAIBehavior/Private/Tasks/MassNavMeshPathFollowTask.cpp
@@ -1,6 +1,6 @@
 // Copyright Epic Games, Inc. All Rights Reserved.
 
-#include "Tasks/MassNavMeshPathFollowTask.h"
+#include "Tasks/MassNavMeshPathfollowTask.h"
 
 #include "MassAIBehaviorTypes.h"
 #include "MassCommonFragments.h"
diff --git a/Engine/Plugins/AI/MassAI/Source/MassAIBehavior/Public/Tasks/MassFindSmartObjectTargetTask.h b/Engine/Plugins/AI/MassAI/Source/MassAIBehavior/Public/Tasks/MassFindSmartObjectTargetTask.h
--- a/Engine/Plugins/AI/MassAI/Source/MassAIBehavior/Public/Tasks/MassFindSmartObjectTargetTask.h
+++ b/Engine/Plugins/AI/MassAI/Source/MassAIBehavior/Public/Tasks/MassFindSmartObjectTargetTask.h
@@ -2,7 +2,7 @@
 
 #pragma once
 
-#include "Tasks/MassNavMeshPathFollowTask.h"
+#include "MassNavigationTypes.h"
 #include "MassStateTreeTypes.h"
 #include "SmartObjectRuntime.h"
 #include "MassFindSmartObjectTargetTask.generated.h"
diff --git a/Engine/Plugins/AI/MassAI/Source/MassAIBehavior/Public/Tasks/MassNavMeshPathfollowTask.h b/Engine/Plugins/AI/MassAI/Source/MassAIBehavior/Public/Tasks/MassNavMeshPathfollowTask.h
--- a/Engine/Plugins/AI/MassAI/Source/MassAIBehavior/Public/Tasks/MassNavMeshPathfollowTask.h
+++ b/Engine/Plugins/AI/MassAI/Source/MassAIBehavior/Public/Tasks/MassNavMeshPathfollowTask.h
@@ -14,25 +14,6 @@ struct FMassMoveTargetFragment;
 struct FAgentRadiusFragment;
 struct FMassMovementParameters;
 
-/** Target location used by FMassNavMeshPathFollowTask */
-USTRUCT()
-struct MASSAIBEHAVIOR_API FMassTargetLocation
-{
-	GENERATED_BODY()
-
-	void Reset()
-	{
-		EndOfPathPosition.Reset();
-		EndOfPathIntent = EMassMovementAction::Move;
-	}
-	
-	/** Optional end of path location. */
-	TOptional<FVector> EndOfPathPosition;
-
-	/** Movement intent at the end of the path */
-	EMassMovementAction EndOfPathIntent = EMassMovementAction::Move;
-};
-
 /** FMassNavMeshPathFollowTask movement parameters */
 USTRUCT()
 struct MASSAIBEHAVIOR_API FMassNavMeshPathFollowTaskInstanceData
diff --git a/Engine/Plugins/AI/MassAI/Source/MassNavigation/Public/MassNavigationTypes.h b/Engine/Plugins/AI/MassAI/Source/MassNavigation/Public/MassNavigationTypes.h
--- a/Engine/Plugins/AI/MassAI/Source/MassNavigation/Public/MassNavigationTypes.h
+++ b/Engine/Plugins/AI/MassAI/Source/MassNavigation/Public/MassNavigationTypes.h
@@ -29,3 +29,22 @@ enum class EMassMovementAction : uint8
 	Move,		// Move or keep on moving.
 	Animate,	// Animation has control over the transform
 };
+
+/** Target location used by tasks */
+USTRUCT()
+struct MASSNAVIGATION_API FMassTargetLocation
+{
+	GENERATED_BODY()
+
+	void Reset()
+	{
+		EndOfPathPosition.Reset();
+		EndOfPathIntent = EMassMovementAction::Move;
+	}
+	
+	/** Optional end of path location. */
+	TOptional<FVector> EndOfPathPosition;
+
+	/** Movement intent at the end of the path */
+	EMassMovementAction EndOfPathIntent = EMassMovementAction::Move;
+};
\ No newline at end of file
