In FMassFindSmartObjectTask, addition of an option to find the SO from the lane location or the world location (keeping the default as it is currently, by lane).

diff --git a/Engine/Plugins/AI/MassAI/Source/MassAIBehavior/Private/Tasks/MassFindSmartObjectTask.cpp b/Engine/Plugins/AI/MassAI/Source/MassAIBehavior/Private/Tasks/MassFindSmartObjectTask.cpp
--- a/Engine/Plugins/AI/MassAI/Source/MassAIBehavior/Private/Tasks/MassFindSmartObjectTask.cpp
+++ b/Engine/Plugins/AI/MassAI/Source/MassAIBehavior/Private/Tasks/MassFindSmartObjectTask.cpp
@@ -97,7 +97,8 @@ EStateTreeRunStatus FMassFindSmartObjectTask::Tick(FStateTreeExecutionContext& C
 		// If the user is already using a SmartObject, or has used interaction recently, skip search and empty results. 
 		if (SOUser.InteractionHandle.IsValid() || SOUser.InteractionCooldownEndTime > World->GetTimeSeconds())
 		{
-			MASSBEHAVIOR_LOG(Verbose, TEXT("Skipped: Recently interacted (%s %.1f)"), SOUser.InteractionHandle.IsValid() ? TEXT("Interacting") : TEXT("Cooldown"), FMath::Max(0.0f, SOUser.InteractionCooldownEndTime - World->GetTimeSeconds()));
+			MASSBEHAVIOR_LOG(Verbose, TEXT("Skipped: Recently interacted (%s %.1f)"),
+				SOUser.InteractionHandle.IsValid() ? TEXT("Interacting") : TEXT("Cooldown"), FMath::Max(0.0f, SOUser.InteractionCooldownEndTime - World->GetTimeSeconds()));
 
 			// Do not offer any new candidates during cool down.
 			InstanceData.FoundCandidateSlots.Reset();
@@ -115,7 +116,7 @@ EStateTreeRunStatus FMassFindSmartObjectTask::Tick(FStateTreeExecutionContext& C
 		{
 			// Use lanes if possible for faster queries using zone graph annotations
 			const FMassEntityHandle RequestingEntity = MassStateTreeContext.GetEntity();
-			if (LaneLocation != nullptr)
+			if (bFindFromLaneLocation && LaneLocation != nullptr)
 			{
 				MASSBEHAVIOR_CLOG(!LaneLocation->LaneHandle.IsValid(), Error, TEXT("Always expecting a valid lane from the ZoneGraph movement"));
 				if (LaneLocation->LaneHandle.IsValid())
diff --git a/Engine/Plugins/AI/MassAI/Source/MassAIBehavior/Public/Tasks/MassFindSmartObjectTask.h b/Engine/Plugins/AI/MassAI/Source/MassAIBehavior/Public/Tasks/MassFindSmartObjectTask.h
--- a/Engine/Plugins/AI/MassAI/Source/MassAIBehavior/Public/Tasks/MassFindSmartObjectTask.h
+++ b/Engine/Plugins/AI/MassAI/Source/MassAIBehavior/Public/Tasks/MassFindSmartObjectTask.h
@@ -69,4 +69,8 @@ protected:
 	/** How frequently to search for new candidates. */
 	UPROPERTY(EditAnywhere, Category = Parameter)
 	float SearchInterval = 10.0f;
+	
+	/** If true, search smart objects using current lane position else, use world position. */
+	UPROPERTY(EditAnywhere, Category = Parameter)
+	bool bFindFromLaneLocation = true;
 };
\ No newline at end of file
