Addition of a FMassNavMeshStandTask and action.

diff --git a/Engine/Plugins/AI/MassAI/Source/MassAIBehavior/Private/Tasks/MassNavMeshStandTask.cpp b/Engine/Plugins/AI/MassAI/Source/MassAIBehavior/Private/Tasks/MassNavMeshStandTask.cpp
new file mode 100644
--- /dev/null
+++ b/Engine/Plugins/AI/MassAI/Source/MassAIBehavior/Private/Tasks/MassNavMeshStandTask.cpp
@@ -0,0 +1,69 @@
+﻿// Copyright Epic Games, Inc. All Rights Reserved.
+
+#include "Tasks/MassNavMeshStandTask.h"
+#include "StateTreeExecutionContext.h"
+#include "MassMovementFragments.h"
+#include "MassNavigationFragments.h"
+#include "MassNavMeshNavigationFragments.h"
+#include "MassNavMeshNavigationUtils.h"
+#include "MassStateTreeExecutionContext.h"
+#include "MassSignalSubsystem.h"
+#include "StateTreeLinker.h"
+
+bool FMassNavMeshStandTask::Link(FStateTreeLinker& Linker)
+{
+	Linker.LinkExternalData(MoveTargetHandle);
+	Linker.LinkExternalData(ShortPathHandle);
+	Linker.LinkExternalData(MovementParamsHandle);
+	Linker.LinkExternalData(MassSignalSubsystemHandle);
+
+	return true;
+}
+
+EStateTreeRunStatus FMassNavMeshStandTask::EnterState(FStateTreeExecutionContext& Context, const FStateTreeTransitionResult& Transition) const
+{
+	const FMassStateTreeExecutionContext& MassStateTreeContext = static_cast<FMassStateTreeExecutionContext&>(Context);
+
+	const FMassMovementParameters& MovementParams = Context.GetExternalData(MovementParamsHandle);
+
+	FInstanceDataType& InstanceData = Context.GetInstanceData(*this);
+
+	FMassNavMeshShortPathFragment& ShortPath = Context.GetExternalData(ShortPathHandle);
+	FMassMoveTargetFragment& MoveTarget = Context.GetExternalData(MoveTargetHandle);
+
+	const UWorld* World = Context.GetWorld();
+	checkf(World != nullptr, TEXT("A valid world is expected from the execution context."));
+
+	MoveTarget.CreateNewAction(EMassMovementAction::Stand, *World);
+	
+	const bool bSuccess = UE::MassNavigation::ActivateActionStand(Context.GetOwner(), MassStateTreeContext.GetEntity(), 
+		MovementParams.DefaultDesiredSpeed, MoveTarget, ShortPath);
+
+	if (!bSuccess)
+	{
+		return EStateTreeRunStatus::Failed;
+	}
+
+	InstanceData.Time = 0.f;
+
+	// A Duration <= 0 indicates that the task runs until a transition in the state tree stops it.
+	// Otherwise, we schedule a signal to end the task.
+	if (InstanceData.Duration > 0.f)
+	{
+		UMassSignalSubsystem& MassSignalSubsystem = Context.GetExternalData(MassSignalSubsystemHandle);
+		MassSignalSubsystem.DelaySignalEntityDeferred(MassStateTreeContext.GetMassEntityExecutionContext(),
+			UE::Mass::Signals::StandTaskFinished, MassStateTreeContext.GetEntity(), InstanceData.Duration);
+	}
+
+	return EStateTreeRunStatus::Running;
+}
+
+EStateTreeRunStatus FMassNavMeshStandTask::Tick(FStateTreeExecutionContext& Context, const float DeltaTime) const
+{
+	FInstanceDataType& InstanceData = Context.GetInstanceData(*this);
+	
+	InstanceData.Time += DeltaTime;
+	
+	return InstanceData.Duration <= 0.f ? EStateTreeRunStatus::Running :
+		(InstanceData.Time < InstanceData.Duration ? EStateTreeRunStatus::Running : EStateTreeRunStatus::Succeeded);
+}
diff --git a/Engine/Plugins/AI/MassAI/Source/MassAIBehavior/Public/Tasks/MassNavMeshStandTask.h b/Engine/Plugins/AI/MassAI/Source/MassAIBehavior/Public/Tasks/MassNavMeshStandTask.h
new file mode 100644
--- /dev/null
+++ b/Engine/Plugins/AI/MassAI/Source/MassAIBehavior/Public/Tasks/MassNavMeshStandTask.h
@@ -0,0 +1,53 @@
+﻿// Copyright Epic Games, Inc. All Rights Reserved.
+
+#pragma once
+
+#include "MassStateTreeTypes.h"
+#include "MassNavMeshStandTask.generated.h"
+
+struct FStateTreeExecutionContext;
+struct FMassMoveTargetFragment;
+struct FMassNavMeshShortPathFragment;
+struct FMassMovementParameters;
+class UZoneGraphSubsystem;
+class UMassSignalSubsystem;
+
+USTRUCT()
+struct MASSAIBEHAVIOR_API FMassNavMeshStandTaskInstanceData
+{
+	GENERATED_BODY()
+
+	/** Delay before the task ends. Default (0 or any negative) will run indefinitely, so it requires a transition in the state tree to stop it. */
+	UPROPERTY(EditAnywhere, Category = Parameter)
+	float Duration = 0.f;
+
+	UPROPERTY()
+	float Time = 0.f;
+};
+
+/**
+ * Stop, and stand on current navmesh location
+ */
+USTRUCT(meta = (DisplayName = "NavMesh Stand"))
+struct MASSAIBEHAVIOR_API FMassNavMeshStandTask : public FMassStateTreeTaskBase
+{
+	GENERATED_BODY()
+
+	using FInstanceDataType = FMassNavMeshStandTaskInstanceData;
+
+protected:
+	virtual bool Link(FStateTreeLinker& Linker) override;
+
+	virtual const UStruct* GetInstanceDataType() const override
+	{
+		return FInstanceDataType::StaticStruct();
+	}
+
+	virtual EStateTreeRunStatus EnterState(FStateTreeExecutionContext& Context, const FStateTreeTransitionResult& Transition) const override;
+	virtual EStateTreeRunStatus Tick(FStateTreeExecutionContext& Context, const float DeltaTime) const override;
+
+	TStateTreeExternalDataHandle<FMassMoveTargetFragment> MoveTargetHandle;
+	TStateTreeExternalDataHandle<FMassNavMeshShortPathFragment> ShortPathHandle;
+	TStateTreeExternalDataHandle<FMassMovementParameters> MovementParamsHandle;
+	TStateTreeExternalDataHandle<UMassSignalSubsystem> MassSignalSubsystemHandle;
+};
diff --git a/Engine/Plugins/AI/MassAI/Source/MassNavMeshNavigation/Private/MassNavMeshNavigationUtils.cpp b/Engine/Plugins/AI/MassAI/Source/MassNavMeshNavigation/Private/MassNavMeshNavigationUtils.cpp
new file mode 100644
--- /dev/null
+++ b/Engine/Plugins/AI/MassAI/Source/MassNavMeshNavigation/Private/MassNavMeshNavigationUtils.cpp
@@ -0,0 +1,34 @@
+﻿// Copyright Epic Games, Inc. All Rights Reserved.
+
+#include "MassNavMeshNavigationUtils.h"
+#include "MassCommonTypes.h"
+#include "MassNavigationFragments.h"
+#include "MassNavMeshNavigationFragments.h"
+#include "VisualLogger/VisualLogger.h"
+
+namespace UE::MassNavigation
+{
+	bool ActivateActionStand(const UObject* Requester,
+							 const FMassEntityHandle Entity,
+							 const float DesiredSpeed,
+							 FMassMoveTargetFragment& InOutMoveTarget,
+							 FMassNavMeshShortPathFragment& OutShortPath)
+	{
+		OutShortPath.Reset();
+		InOutMoveTarget.DistanceToGoal = 0.f;
+		InOutMoveTarget.DesiredSpeed.Set(0.f);
+		
+		if (!ensureMsgf(InOutMoveTarget.GetCurrentAction() == EMassMovementAction::Stand,
+				TEXT("Expecting action 'Stand': Invalid action %s."), *UEnum::GetDisplayValueAsText(InOutMoveTarget.GetCurrentAction()).ToString()))
+		{
+			return false;
+		}
+		
+		InOutMoveTarget.IntentAtGoal = EMassMovementAction::Stand;
+		InOutMoveTarget.DesiredSpeed.Set(DesiredSpeed);
+		
+		UE_VLOG(Requester, LogMassNavigation, Log, TEXT("Entity [%s] successfully requested %s."), *Entity.DebugGetDescription(), *InOutMoveTarget.ToString());
+
+		return true;
+	}	
+}
diff --git a/Engine/Plugins/AI/MassAI/Source/MassNavMeshNavigation/Public/MassNavMeshNavigationUtils.h b/Engine/Plugins/AI/MassAI/Source/MassNavMeshNavigation/Public/MassNavMeshNavigationUtils.h
new file mode 100644
--- /dev/null
+++ b/Engine/Plugins/AI/MassAI/Source/MassNavMeshNavigation/Public/MassNavMeshNavigationUtils.h
@@ -0,0 +1,16 @@
+﻿// Copyright Epic Games, Inc. All Rights Reserved.
+
+#pragma once
+
+#include "EngineDefines.h"
+
+class UObject;
+struct FMassEntityHandle;
+struct FMassMoveTargetFragment;
+struct FMassNavMeshShortPathFragment;
+
+namespace UE::MassNavigation
+{
+	MASSNAVMESHNAVIGATION_API bool ActivateActionStand(
+		const UObject* Requester, const FMassEntityHandle Entity, const float DesiredSpeed, FMassMoveTargetFragment& InOutMoveTarget, FMassNavMeshShortPathFragment& OutShortPath);
+};
