StateTree: Review feedback from CL 38076983.

diff --git a/Engine/Plugins/AI/MassAI/Source/MassAIBehavior/Private/MassStateTreeExecutionContext.cpp b/Engine/Plugins/AI/MassAI/Source/MassAIBehavior/Private/MassStateTreeExecutionContext.cpp
--- a/Engine/Plugins/AI/MassAI/Source/MassAIBehavior/Private/MassStateTreeExecutionContext.cpp
+++ b/Engine/Plugins/AI/MassAI/Source/MassAIBehavior/Private/MassStateTreeExecutionContext.cpp
@@ -133,6 +133,22 @@ EStateTreeRunStatus FMassStateTreeExecutionContext::Start()
 		});
 }
 
+EStateTreeRunStatus FMassStateTreeExecutionContext::Start(const FInstancedPropertyBag* InitialParameters, int32 RandomSeed)
+{
+	ensureMsgf(Entity.IsValid(), TEXT("The entity is not valid before starting the state tree instance."));
+
+	const TOptional<int32> ParamRandomSeed = RandomSeed == -1 ? TOptional<int32>() : RandomSeed;
+
+	FMassExecutionExtension Extension;
+	Extension.Entity = Entity;
+	return FStateTreeExecutionContext::Start(FStateTreeExecutionContext::FStartParameters
+		{
+			.GlobalParameters = InitialParameters,
+			.ExecutionExtension = TInstancedStruct<FMassExecutionExtension>::Make(MoveTemp(Extension)),
+			.RandomSeed = ParamRandomSeed
+		});
+}
+
 void FMassStateTreeExecutionContext::SetEntity(const FMassEntityHandle InEntity)
 {
 	Entity = InEntity;
diff --git a/Engine/Plugins/AI/MassAI/Source/MassAIBehavior/Public/MassStateTreeExecutionContext.h b/Engine/Plugins/AI/MassAI/Source/MassAIBehavior/Public/MassStateTreeExecutionContext.h
--- a/Engine/Plugins/AI/MassAI/Source/MassAIBehavior/Public/MassStateTreeExecutionContext.h
+++ b/Engine/Plugins/AI/MassAI/Source/MassAIBehavior/Public/MassStateTreeExecutionContext.h
@@ -49,6 +49,7 @@ struct MASSAIBEHAVIOR_API FMassStateTreeExecutionContext : public FStateTreeExec
 
 	/** Start executing. */
 	EStateTreeRunStatus Start();
+	EStateTreeRunStatus Start(const FInstancedPropertyBag* InitialParameters, int32 RandomSeed = -1);
 
 	FMassEntityManager& GetEntityManager() const
 	{
diff --git a/Engine/Plugins/Runtime/StateTree/Source/StateTreeEditorModule/Public/StateTreeViewModel.h b/Engine/Plugins/Runtime/StateTree/Source/StateTreeEditorModule/Public/StateTreeViewModel.h
--- a/Engine/Plugins/Runtime/StateTree/Source/StateTreeEditorModule/Public/StateTreeViewModel.h
+++ b/Engine/Plugins/Runtime/StateTree/Source/StateTreeEditorModule/Public/StateTreeViewModel.h
@@ -76,7 +76,7 @@ public:
 	int32 GetSubTreeCount() const;
 	void GetSubTrees(TArray<TWeakObjectPtr<UStateTreeState>>& OutSubtrees) const;
 
-	/** Find the state that link to the StateTree and the state the state can link to. */
+	/** Find the states that are linked to the provided StateID. */
 	void GetLinkStates(FGuid StateID, TArray<FGuid>& LinkingIn, TArray<FGuid>& LinkedOut) const;
 
 	// Gets and sets StateTree view expansion state store in the asset.
diff --git a/Engine/Plugins/Runtime/StateTree/Source/StateTreeModule/Private/StateTreeExecutionContext.cpp b/Engine/Plugins/Runtime/StateTree/Source/StateTreeModule/Private/StateTreeExecutionContext.cpp
--- a/Engine/Plugins/Runtime/StateTree/Source/StateTreeModule/Private/StateTreeExecutionContext.cpp
+++ b/Engine/Plugins/Runtime/StateTree/Source/StateTreeModule/Private/StateTreeExecutionContext.cpp
@@ -376,7 +376,7 @@ FStateTreeWeakMinimalExecutionContext FStateTreeExecutionContext::GetWeakMinimal
 
 EStateTreeRunStatus FStateTreeExecutionContext::Start(const FInstancedPropertyBag* InitialParameters, int32 RandomSeed)
 {
-	TOptional<int32> ParamRandomSeed = RandomSeed == -1 ? TOptional<int32>() : RandomSeed;
+	const TOptional<int32> ParamRandomSeed = RandomSeed == -1 ? TOptional<int32>() : RandomSeed;
 	return Start(FStartParameters{.GlobalParameters = InitialParameters, .RandomSeed = ParamRandomSeed });
 }
 
diff --git a/Engine/Plugins/Runtime/StateTree/Source/StateTreeModule/Public/StateTreeExecutionContext.h b/Engine/Plugins/Runtime/StateTree/Source/StateTreeModule/Public/StateTreeExecutionContext.h
--- a/Engine/Plugins/Runtime/StateTree/Source/StateTreeModule/Public/StateTreeExecutionContext.h
+++ b/Engine/Plugins/Runtime/StateTree/Source/StateTreeModule/Public/StateTreeExecutionContext.h
@@ -566,7 +566,7 @@ protected:
 
 	/** @return Prefix that will be used by STATETREE_LOG and STATETREE_CLOG, Owner name by default. */
 	UE_DEPRECATED(5.6, "Use FStateTreeExecutionExtension::GetInstanceDescription instead")
-	virtual FString GetInstanceDescription() const;
+	virtual FString GetInstanceDescription() const final;
 
 	/** Callback when delayed transition is triggered. Contexts that are event based can use this to trigger a future event. */
 	virtual void BeginDelayedTransition(const FStateTreeTransitionDelayedState& DelayedState) {};
