Allow UMassMovingAvoidanceProcessor to account for environment edges.
Addition of environment boundaries in FMassNavMeshShortPathFragment.
Those boundaries are added when requesting a short path and are made from a FNavCorridor of the path.
Fix in avoidance processor to ignore environment edges that are not yet relevant.
Addition of UMassNavMeshNavigationBoundaryProcessor to pass the boundaries to the avoidance through the FMassNavigationEdgesFragment.

diff --git a/Engine/Plugins/AI/MassAI/MassAI.uplugin b/Engine/Plugins/AI/MassAI/MassAI.uplugin
--- a/Engine/Plugins/AI/MassAI/MassAI.uplugin
+++ b/Engine/Plugins/AI/MassAI/MassAI.uplugin
@@ -70,6 +70,10 @@
 		{
 			"Name": "StateTree",
 			"Enabled": true
+		},
+		{
+			"Name": "NavCorridor",
+			"Enabled": true
 		}
 	]
 }
\ No newline at end of file
diff --git a/Engine/Plugins/AI/MassAI/Source/MassAIBehavior/MassAIBehavior.Build.cs b/Engine/Plugins/AI/MassAI/Source/MassAIBehavior/MassAIBehavior.Build.cs
--- a/Engine/Plugins/AI/MassAI/Source/MassAIBehavior/MassAIBehavior.Build.cs
+++ b/Engine/Plugins/AI/MassAI/Source/MassAIBehavior/MassAIBehavior.Build.cs
@@ -38,7 +38,8 @@ namespace UnrealBuildTool.Rules
 					"StateTreeModule",
 					"ZoneGraph",
 					"ZoneGraphAnnotations",
-					"MassGameplayExternalTraits"
+					"MassGameplayExternalTraits",
+					"NavCorridor"
 				}
 			);
 
diff --git a/Engine/Plugins/AI/MassAI/Source/MassAIBehavior/Private/Tasks/MassNavMeshPathFollowTask.cpp b/Engine/Plugins/AI/MassAI/Source/MassAIBehavior/Private/Tasks/MassNavMeshPathFollowTask.cpp
--- a/Engine/Plugins/AI/MassAI/Source/MassAIBehavior/Private/Tasks/MassNavMeshPathFollowTask.cpp
+++ b/Engine/Plugins/AI/MassAI/Source/MassAIBehavior/Private/Tasks/MassNavMeshPathFollowTask.cpp
@@ -8,6 +8,7 @@
 #include "MassMovementFragments.h"
 #include "MassNavigationFragments.h"
 #include "MassStateTreeExecutionContext.h"
+#include "NavCorridor.h"
 #include "NavigationSystem.h"
 #include "StateTreeLinker.h"
 
@@ -75,11 +76,19 @@ bool FMassNavMeshPathFollowTask::RequestPath(FStateTreeExecutionContext& Context
 		FMassNavMeshCachedPathFragment& CachedPathFragment = Context.GetExternalData(CachedPathHandle);
 		CachedPathFragment.NavPath = Result.Path;
 
-		// Update short path
-		const TArray<FNavPathPoint>& PathPoints = Result.Path.Get()->GetPathPoints();
+		// Build corridor
+		CachedPathFragment.Corridor = MakeShared<FNavCorridor>();
+		const FSharedConstNavQueryFilter NavQueryFilter = Query.QueryFilter ? Query.QueryFilter : NavData->GetDefaultQueryFilter();
+
+		FNavCorridorParams CorridorParams;
+		CorridorParams.SetFromWidth(InstanceData.CorridorWidth);
+		CorridorParams.PathOffsetFromBoundaries = InstanceData.OffsetFromBoundaries;
+
+		CachedPathFragment.Corridor->BuildFromPath(*CachedPathFragment.NavPath, NavQueryFilter, CorridorParams);
 
+		// Update short path
 		FMassNavMeshShortPathFragment& ShortPathFragment = Context.GetExternalData(ShortPathHandle);
-		ShortPathFragment.RequestShortPath(PathPoints, 0);
+		ShortPathFragment.RequestShortPath(CachedPathFragment.Corridor, 0);
 
 		CachedPathFragment.NavPathNextStartIndex += ShortPathFragment.NumPoints-1;
 		
@@ -107,7 +116,8 @@ bool FMassNavMeshPathFollowTask::UpdateShortPath(FStateTreeExecutionContext& Con
 	MASSBEHAVIOR_LOG(Verbose, TEXT("updating short path, starting at index %i"), CachedPathFragment.NavPathNextStartIndex);
 	
 	FMassNavMeshShortPathFragment& ShortPathFragment = Context.GetExternalData(ShortPathHandle);
-	ShortPathFragment.RequestShortPath(CachedPathFragment.NavPath.Get()->GetPathPoints(), CachedPathFragment.NavPathNextStartIndex);
+
+	ShortPathFragment.RequestShortPath(CachedPathFragment.Corridor, CachedPathFragment.NavPathNextStartIndex);
 
 	CachedPathFragment.NavPathNextStartIndex += ShortPathFragment.NumPoints-1;
 
diff --git a/Engine/Plugins/AI/MassAI/Source/MassAIBehavior/Public/Tasks/MassNavMeshPathfollowTask.h b/Engine/Plugins/AI/MassAI/Source/MassAIBehavior/Public/Tasks/MassNavMeshPathfollowTask.h
--- a/Engine/Plugins/AI/MassAI/Source/MassAIBehavior/Public/Tasks/MassNavMeshPathfollowTask.h
+++ b/Engine/Plugins/AI/MassAI/Source/MassAIBehavior/Public/Tasks/MassNavMeshPathfollowTask.h
@@ -28,6 +28,14 @@ struct MASSAIBEHAVIOR_API FMassNavMeshPathFollowTaskInstanceData
 
 	UPROPERTY(EditAnywhere, Category = Parameter)
 	float SpeedScale = 1.f;
+	
+	/** Maximum width of the corridor to use. */
+	UPROPERTY(EditAnywhere, Category = Parameter)
+    float CorridorWidth = 300.f;
+
+	/** Amount to offset corridor sides from navmesh borders. Used to keep movement away for navmesh borders. */
+	UPROPERTY(EditAnywhere, Category = Parameter)
+	float OffsetFromBoundaries = 10.f;
 };
 
 /** Finds a path to TargetLocation, requests a short path, starts a move action and follow the path by updating the short path when needed. */ 
diff --git a/Engine/Plugins/AI/MassAI/Source/MassNavMeshNavigation/MassNavMeshNavigation.Build.cs b/Engine/Plugins/AI/MassAI/Source/MassNavMeshNavigation/MassNavMeshNavigation.Build.cs
--- a/Engine/Plugins/AI/MassAI/Source/MassNavMeshNavigation/MassNavMeshNavigation.Build.cs
+++ b/Engine/Plugins/AI/MassAI/Source/MassNavMeshNavigation/MassNavMeshNavigation.Build.cs
@@ -20,7 +20,8 @@ namespace UnrealBuildTool.Rules
 					"MassSignals",
 					"MassSpawner",
 					"MassMovement",
-					"MassNavigation"
+					"MassNavigation",
+					"NavCorridor"
 				}
 			);
 		}
diff --git a/Engine/Plugins/AI/MassAI/Source/MassNavMeshNavigation/Private/MassNavMeshNavigationBoundaryProcessor.cpp b/Engine/Plugins/AI/MassAI/Source/MassNavMeshNavigation/Private/MassNavMeshNavigationBoundaryProcessor.cpp
new file mode 100644
--- /dev/null
+++ b/Engine/Plugins/AI/MassAI/Source/MassNavMeshNavigation/Private/MassNavMeshNavigationBoundaryProcessor.cpp
@@ -0,0 +1,100 @@
+﻿// Copyright Epic Games, Inc. All Rights Reserved.
+
+#include "MassNavMeshNavigationBoundaryProcessor.h"
+
+#include "MassNavMeshNavigationTypes.h"
+#include "MassCommonTypes.h"
+#include "MassLODFragments.h"
+#include "MassNavigationFragments.h"
+#include "MassNavMeshNavigationFragments.h"
+#include "Avoidance/MassAvoidanceFragments.h"
+#include "MassDebugger.h"
+#include "VisualLogger/VisualLogger.h"
+
+UMassNavMeshNavigationBoundaryProcessor::UMassNavMeshNavigationBoundaryProcessor()
+	: EntityQuery(*this)
+{
+	ExecutionFlags = (int32)EProcessorExecutionFlags::AllNetModes;
+
+	bAutoRegisterWithProcessingPhases = true;
+	ExecutionOrder.ExecuteBefore.Add(UE::Mass::ProcessorGroupNames::Avoidance);
+}
+
+void UMassNavMeshNavigationBoundaryProcessor::ConfigureQueries(const TSharedRef<FMassEntityManager>& EntityManager)
+{
+	EntityQuery.AddRequirement<FMassNavMeshShortPathFragment>(EMassFragmentAccess::ReadOnly);
+	EntityQuery.AddRequirement<FMassMoveTargetFragment>(EMassFragmentAccess::ReadOnly);
+	
+	EntityQuery.AddRequirement<FMassNavMeshBoundaryFragment>(EMassFragmentAccess::ReadWrite);
+	EntityQuery.AddRequirement<FMassNavigationEdgesFragment>(EMassFragmentAccess::ReadWrite);	// to output edges for avoidance
+}
+
+void UMassNavMeshNavigationBoundaryProcessor::Execute(FMassEntityManager& EntityManager, FMassExecutionContext& Context)
+{
+	QUICK_SCOPE_CYCLE_COUNTER(UMassNavMeshNavigationBoundaryProcessor);
+
+	EntityQuery.ForEachEntityChunk(EntityManager, Context, [this](FMassExecutionContext& Context)
+	{
+		const int32 NumEntities = Context.GetNumEntities();
+
+		TConstArrayView<FMassNavMeshShortPathFragment> ShortPathList = Context.GetFragmentView<FMassNavMeshShortPathFragment>();
+		TConstArrayView<FMassMoveTargetFragment> MovementTargetList = Context.GetFragmentView<FMassMoveTargetFragment>();
+
+		TArrayView<FMassNavMeshBoundaryFragment> NavmeshBoundaryList = Context.GetMutableFragmentView<FMassNavMeshBoundaryFragment>();
+		TArrayView<FMassNavigationEdgesFragment> EdgesList = Context.GetMutableFragmentView<FMassNavigationEdgesFragment>();
+
+		for (int32 EntityIndex = 0; EntityIndex < NumEntities; ++EntityIndex)
+		{
+			const FMassNavMeshShortPathFragment& ShortPath = ShortPathList[EntityIndex];
+			const FMassMoveTargetFragment& MovementTarget = MovementTargetList[EntityIndex];
+			FMassNavMeshBoundaryFragment& NavmeshBoundary = NavmeshBoundaryList[EntityIndex];
+			FMassNavigationEdgesFragment& Edges = EdgesList[EntityIndex];
+
+			const FMassEntityHandle Entity = Context.GetEntity(EntityIndex);
+
+			// First check if we moved enough for an update.
+			const FVector::FReal DeltaDistSquared = FVector::DistSquared(MovementTarget.Center, NavmeshBoundary.LastUpdatePosition);
+			constexpr FVector::FReal UpdateDistanceThresholdSquared = FMath::Square(50.);
+
+			bool bDisplayDebug = false;
+#if WITH_MASSGAMEPLAY_DEBUG
+			FColor EntityColor = FColor::Black;
+			bDisplayDebug = UE::Mass::Debug::IsDebuggingEntity(Entity, &EntityColor);
+
+			if (bDisplayDebug)
+			{
+				// Draw MoveTarget
+				const FVector ZOffset(0,0,10);
+				constexpr float Radius = 20.f;
+				UE_VLOG_CIRCLE(this, LogMassNavMeshNavigation, Display, NavmeshBoundary.LastUpdatePosition + ZOffset, FVector(0,0,1), Radius, EntityColor, TEXT(""));
+			}
+#endif // WITH_MASSGAMEPLAY_DEBUG
+
+			if (DeltaDistSquared < UpdateDistanceThresholdSquared)
+			{
+				// Not moved enough
+				continue;
+			}
+
+			NavmeshBoundary.LastUpdatePosition = MovementTarget.Center;
+			Edges.AvoidanceEdges.Reset();
+
+			if (ShortPath.NumPoints < 2)
+			{
+				// Nothing to do
+				continue;
+			}
+
+			// Make environment edges from the short path.
+			for (int32 Index = 0; Index+1 < ShortPath.NumPoints; Index++)
+			{
+				const FMassNavMeshPathPoint& Portal = ShortPath.Points[Index];
+				const FMassNavMeshPathPoint& NextPortal = ShortPath.Points[Index+1];
+
+				// Left side: reverse start and end to keep the normal inside.
+				Edges.AvoidanceEdges.Add(FNavigationAvoidanceEdge(NextPortal.Left, Portal.Left));
+				Edges.AvoidanceEdges.Add(FNavigationAvoidanceEdge(Portal.Right, NextPortal.Right));
+			}
+		}
+	});
+}
diff --git a/Engine/Plugins/AI/MassAI/Source/MassNavMeshNavigation/Private/MassNavMeshNavigationFragments.cpp b/Engine/Plugins/AI/MassAI/Source/MassNavMeshNavigation/Private/MassNavMeshNavigationFragments.cpp
--- a/Engine/Plugins/AI/MassAI/Source/MassNavMeshNavigation/Private/MassNavMeshNavigationFragments.cpp
+++ b/Engine/Plugins/AI/MassAI/Source/MassNavMeshNavigation/Private/MassNavMeshNavigationFragments.cpp
@@ -3,29 +3,37 @@
 #include "MassNavMeshNavigationFragments.h"
 
 #include "AI/Navigation/NavigationTypes.h"
+#include "NavCorridor.h"
 
-// Make a new short path
-bool FMassNavMeshShortPathFragment::RequestShortPath(const TArray<FNavPathPoint>& NavPathPoints, const int32 NavPathStartIndex)
+// Make a new short corridor
+bool FMassNavMeshShortPathFragment::RequestShortPath(const TSharedPtr<FNavCorridor>& NavCorridor, const int32 NavCorridorStartIndex)
 {
 	Reset();
 
-	// Copy PathPoints from the StartIndex
-	for(; NumPoints < MaxPoints && (NavPathStartIndex + NumPoints) < NavPathPoints.Num(); NumPoints++)
+	TArray<FNavCorridorPortal>& NavCorridorPortals = NavCorridor->Portals;
+	int32 PortalNum = NavCorridorPortals.Num();
+
+	// Copy Portals from the StartIndex
+	for(; NumPoints < MaxPoints && (NavCorridorStartIndex + NumPoints) < PortalNum; NumPoints++)
 	{
+		const FNavCorridorPortal& SourcePortal = NavCorridorPortals[NavCorridorStartIndex + NumPoints];
+		
 		FMassNavMeshPathPoint& Point = Points[NumPoints];
-		Point.Position = NavPathPoints[NavPathStartIndex + NumPoints];
+		Point.Left = SourcePortal.Left;
+		Point.Right = SourcePortal.Right;
+		Point.Position = SourcePortal.Location;
 	}
 
 	checkf(NumPoints >= 2, TEXT("Path should have at least 2 points at this stage, has %d."), NumPoints);
 
 	// We are not done with the main path, mark as partial.
-	const int32 LastNavPathIndex = NavPathPoints.Num()-1;
-	if (NavPathStartIndex + (NumPoints-1) < LastNavPathIndex)
+	const int32 LastNavPathIndex = PortalNum-1;
+	if (NavCorridorStartIndex + (NumPoints-1) < LastNavPathIndex)
 	{
 		bPartialResult = true;
 	}
 
-	// Calculate movement distance at each point.
+	// Calculate movement distance and tangent at each point.
 	float PathDistance = 0.f;
 	Points[0].Distance.Set(PathDistance);
 	for (uint8 PointIndex = 1; PointIndex < NumPoints; PointIndex++)
diff --git a/Engine/Plugins/AI/MassAI/Source/MassNavMeshNavigation/Private/MassNavMeshNavigationProcessors.cpp b/Engine/Plugins/AI/MassAI/Source/MassNavMeshNavigation/Private/MassNavMeshNavigationProcessors.cpp
--- a/Engine/Plugins/AI/MassAI/Source/MassNavMeshNavigation/Private/MassNavMeshNavigationProcessors.cpp
+++ b/Engine/Plugins/AI/MassAI/Source/MassNavMeshNavigation/Private/MassNavMeshNavigationProcessors.cpp
@@ -10,6 +10,7 @@
 #include "MassNavigationFragments.h"
 #include "MassSignalSubsystem.h"
 #include "MassSimulationLOD.h"
+#include "NavCorridor.h"
 
 #define UNSAFE_FOR_MT 1
 
@@ -164,12 +165,19 @@ void UMassNavMeshPathFollowProcessor::Execute(FMassEntityManager& EntityManager,
 					MixColor /= 2;
 					const FColor LightEntityColor = MixColor.ToFColorSRGB();
 
+					FLinearColor BorderMixColor(EntityColor);
+					BorderMixColor += FColor::Black;
+					BorderMixColor /= 2;
+					const FColor DarkEntityColor = BorderMixColor.ToFColorSRGB();
+
 					// Draw path
 					for (uint8 PointIndex = 0; PointIndex < ShortPath.NumPoints - 1; PointIndex++)
 					{
 						const FMassNavMeshPathPoint& CurrPoint = ShortPath.Points[PointIndex];
 						const FMassNavMeshPathPoint& NextPoint = ShortPath.Points[PointIndex + 1];
 						UE_VLOG_SEGMENT_THICK(this, LogMassNavMeshNavigation, Display, CurrPoint.Position + ZOffset, NextPoint.Position + ZOffset, EntityColor, /*Thickness*/3, TEXT(""));
+						UE_VLOG_SEGMENT_THICK(this, LogMassNavMeshNavigation, Display, CurrPoint.Left + ZOffset, NextPoint.Left + ZOffset, DarkEntityColor, /*Thickness*/2, TEXT(""));
+						UE_VLOG_SEGMENT_THICK(this, LogMassNavMeshNavigation, Display, CurrPoint.Right + ZOffset, NextPoint.Right + ZOffset, DarkEntityColor, /*Thickness*/2, TEXT(""));
 					}
 
 					// Draw tangents
diff --git a/Engine/Plugins/AI/MassAI/Source/MassNavMeshNavigation/Private/MassNavMeshNavigationTrait.cpp b/Engine/Plugins/AI/MassAI/Source/MassNavMeshNavigation/Private/MassNavMeshNavigationTrait.cpp
--- a/Engine/Plugins/AI/MassAI/Source/MassNavMeshNavigation/Private/MassNavMeshNavigationTrait.cpp
+++ b/Engine/Plugins/AI/MassAI/Source/MassNavMeshNavigation/Private/MassNavMeshNavigationTrait.cpp
@@ -1,11 +1,17 @@
 // Copyright Epic Games, Inc. All Rights Reserved.
+
 #include "MassNavMeshNavigationTrait.h"
+
 #include "MassEntityTemplateRegistry.h"
-#include "MassEntityUtils.h"
+#include "MassNavigationFragments.h"
 #include "MassNavMeshNavigationFragments.h"
 
 void UMassNavMeshNavigationTrait::BuildTemplate(FMassEntityTemplateBuildContext& BuildContext, const UWorld& World) const
 {
+	BuildContext.RequireFragment<FMassMoveTargetFragment>();
+
 	BuildContext.AddFragment<FMassNavMeshCachedPathFragment>();
 	BuildContext.AddFragment<FMassNavMeshShortPathFragment>();
+
+	BuildContext.AddFragment<FMassNavMeshBoundaryFragment>();	
 }
diff --git a/Engine/Plugins/AI/MassAI/Source/MassNavMeshNavigation/Public/MassNavMeshNavigationBoundaryProcessor.h b/Engine/Plugins/AI/MassAI/Source/MassNavMeshNavigation/Public/MassNavMeshNavigationBoundaryProcessor.h
new file mode 100644
--- /dev/null
+++ b/Engine/Plugins/AI/MassAI/Source/MassNavMeshNavigation/Public/MassNavMeshNavigationBoundaryProcessor.h
@@ -0,0 +1,23 @@
+﻿// Copyright Epic Games, Inc. All Rights Reserved.
+
+#pragma once
+
+#include "MassProcessor.h"
+#include "MassNavMeshNavigationBoundaryProcessor.generated.h"
+
+/** Fills FMassNavigationEdgesFragment using FMassNavMeshShortPathFragment. */
+UCLASS()
+class MASSNAVMESHNAVIGATION_API UMassNavMeshNavigationBoundaryProcessor : public UMassProcessor
+{
+	GENERATED_BODY()
+
+public:
+	UMassNavMeshNavigationBoundaryProcessor();
+
+protected:
+	virtual void ConfigureQueries(const TSharedRef<FMassEntityManager>& EntityManager) override;
+	virtual void Execute(FMassEntityManager& EntityManager, FMassExecutionContext& Context) override;
+
+private:
+	FMassEntityQuery EntityQuery;
+};
diff --git a/Engine/Plugins/AI/MassAI/Source/MassNavMeshNavigation/Public/MassNavMeshNavigationFragments.h b/Engine/Plugins/AI/MassAI/Source/MassNavMeshNavigation/Public/MassNavMeshNavigationFragments.h
--- a/Engine/Plugins/AI/MassAI/Source/MassNavMeshNavigation/Public/MassNavMeshNavigationFragments.h
+++ b/Engine/Plugins/AI/MassAI/Source/MassNavMeshNavigation/Public/MassNavMeshNavigationFragments.h
@@ -9,12 +9,20 @@
 #include "Containers/StaticArray.h"
 #include "MassNavMeshNavigationFragments.generated.h"
 
+struct FNavCorridor;
+
 /** Navmesh path point data */
 USTRUCT()
 struct MASSNAVMESHNAVIGATION_API FMassNavMeshPathPoint
 {
 	GENERATED_BODY()
 
+	/** Portal left limit. */
+	FVector Left = FVector::ZeroVector;
+	
+	/** Portal right limit. */
+	FVector Right = FVector::ZeroVector;
+
 	/** Position of the path. */
 	FVector Position = FVector::ZeroVector;
 
@@ -50,10 +58,10 @@ struct MASSNAVMESHNAVIGATION_API FMassNavMeshShortPathFragment : public FMassFra
 		return NumPoints == 0 || bDone;
 	}
 
-	/** Fill Points from an array of FNavPathPoint. */
-	bool RequestShortPath(const TArray<FNavPathPoint>& NavPathPoints, const int32 NavPathStartIndex);
-	
-	/** Path points */
+	/** Fill Points from Corridor. */
+	bool RequestShortPath(const TSharedPtr<FNavCorridor>& Corridor, const int32 NavCorridorStartIndex);
+
+	/** Portal points */
 	TStaticArray<FMassNavMeshPathPoint, MaxPoints> Points;
 
 	/** Current progress distance along the short path. */
@@ -81,6 +89,18 @@ struct MASSNAVMESHNAVIGATION_API FMassNavMeshCachedPathFragment : public FMassFr
 	/** Reference to a FNavigationPath. */
 	FNavPathSharedPtr NavPath;
 
+	/** Reference to an FNavCorridor. Built out of a navigation path. */
+	TSharedPtr<FNavCorridor> Corridor;
+
 	/** Index used to keep track of progression on the navmesh path. */
 	uint16 NavPathNextStartIndex = 0;
 };
+
+USTRUCT()
+struct MASSNAVMESHNAVIGATION_API FMassNavMeshBoundaryFragment : public FMassFragment
+{
+	GENERATED_BODY()
+
+	/** MovementTarget position when UMassNavMeshNavigationBoundaryProcessor was last updated. Used to identify when a new update is needed. */
+	FVector LastUpdatePosition = FVector::ZeroVector;
+};
diff --git a/Engine/Plugins/AI/MassAI/Source/MassNavigation/Private/MassAvoidanceProcessors.cpp b/Engine/Plugins/AI/MassAI/Source/MassNavigation/Private/MassAvoidanceProcessors.cpp
--- a/Engine/Plugins/AI/MassAI/Source/MassNavigation/Private/MassAvoidanceProcessors.cpp
+++ b/Engine/Plugins/AI/MassAI/Source/MassNavigation/Private/MassAvoidanceProcessors.cpp
@@ -30,7 +30,7 @@ namespace UE::MassAvoidance
 	namespace Tweakables
 	{
 		bool bEnableEnvironmentAvoidance = true;
-		bool bEnableSettingsforExtendingColliders = true;
+		bool bEnableSettingsForExtendingColliders = true;
 		bool bUseAdjacentCorridors = true;
 		bool bUseDrawDebugHelpers = false;
 	} // Tweakables
@@ -38,7 +38,7 @@ namespace UE::MassAvoidance
 	FAutoConsoleVariableRef Vars[] = 
 	{
 		FAutoConsoleVariableRef(TEXT("ai.mass.avoidance.EnableEnvironmentAvoidance"), Tweakables::bEnableEnvironmentAvoidance, TEXT("Set to false to disable avoidance forces for environment (for debug purposes)."), ECVF_Cheat),
-		FAutoConsoleVariableRef(TEXT("ai.mass.avoidance.EnableSettingsforExtendingColliders"), Tweakables::bEnableSettingsforExtendingColliders, TEXT("Set to false to disable using different settings for extending obstacles (for debug purposes)."), ECVF_Cheat),
+		FAutoConsoleVariableRef(TEXT("ai.mass.avoidance.EnableSettingsForExtendingColliders"), Tweakables::bEnableSettingsForExtendingColliders, TEXT("Set to false to disable using different settings for extending obstacles (for debug purposes)."), ECVF_Cheat),
 		FAutoConsoleVariableRef(TEXT("ai.mass.avoidance.UseAdjacentCorridors"), Tweakables::bUseAdjacentCorridors, TEXT("Set to false to disable usage of adjacent lane width."), ECVF_Cheat),
 		FAutoConsoleVariableRef(TEXT("ai.mass.avoidance.UseDrawDebugHelpers"), Tweakables::bUseDrawDebugHelpers, TEXT("Use debug draw helpers in addition to visual logs."), ECVF_Cheat)
 	};
@@ -164,7 +164,6 @@ namespace UE::MassAvoidance
 	static const FColor CurrentAgentColor = FColor::Emerald;
 
 	static const FColor VelocityColor = FColor::Black;
-	static const FColor PrefVelocityColor = FColor::Red;
 	static const FColor DesiredVelocityColor = FColor::Yellow;
 	static const FColor FinalSteeringForceColor = FColor::Cyan;
 	static constexpr float BigArrowThickness = 6.f;
@@ -172,14 +171,14 @@ namespace UE::MassAvoidance
 
 	// Agents colors
 	static const FColor AgentsColor = FColor::Orange;
-	static const FColor AgentSeparationForceColor = FColor(255, 145, 71);	// Orange red
+	static constexpr FColor AgentSeparationForceColor = FColor(255, 145, 71);	// Orange red
 	static const FColor AgentAvoidForceColor = AgentsColor;
 	
 	// Obstacles colors
 	static const FColor ObstacleColor = FColor::Blue;
 	static const FColor ObstacleContactNormalColor = FColor::Silver;
 	static const FColor ObstacleAvoidForceColor = FColor::Magenta;
-	static const FColor ObstacleSeparationForceColor = FColor(255, 66, 66);	// Bright red
+	static constexpr FColor ObstacleSeparationForceColor = FColor(255, 66, 66);	// Bright red
 	
 	static const FVector DebugAgentHeightOffset = FVector(0., 0., 185.);
 	static const FVector DebugLowCylinderOffset = FVector(0., 0., 20.);
@@ -231,7 +230,7 @@ namespace UE::MassAvoidance
 			return;
 		}
 
-		const FVector::FReal Pointyness = 1.8;
+		constexpr FVector::FReal Pointyness = 1.8;
 		const FVector Line = End - Start;
 		const FVector UnitV = Line.GetSafeNormal();
 		const FVector Perp = FVector::CrossProduct(UnitV, FVector::UpVector);
@@ -307,8 +306,8 @@ namespace UE::MassAvoidance
 			return;
 		}
 
-		const float Thickness = 3.f;
-		const FVector::FReal Pointyness = 1.8;
+		constexpr float Thickness = 3.f;
+		constexpr FVector::FReal Pointyness = 1.8;
 		const FVector Line = End - Start;
 		const FVector UnitV = Line.GetSafeNormal();
 		const FVector Perp = FVector::CrossProduct(UnitV, FVector::UpVector);
@@ -424,6 +423,7 @@ void UMassMovingAvoidanceProcessor::Execute(FMassEntityManager& EntityManager, F
 			FVector Position = FVector::ZeroVector;
 			FVector Normal = FVector::ZeroVector;
 			FVector::FReal Distance = 0.;
+			bool bRelevantForSeparation = false;
 		};
 		TArray<FEnvironmentContact, TInlineAllocator<16>> Contacts;
 
@@ -500,18 +500,23 @@ void UMassMovingAvoidanceProcessor::Execute(FMassEntityManager& EntityManager, F
 			{
 				// Draw agent
 				const FString Text = FString::Printf(TEXT("%i"), Entity.Index);
-				DebugDrawCylinder(BaseDebugContext, AgentLocation, AgentLocation + UE::MassAvoidance::DebugAgentHeightOffset, static_cast<float>(AgentRadius+1.), UE::MassAvoidance::CurrentAgentColor, Text);
+				DebugDrawCylinder(BaseDebugContext, AgentLocation, AgentLocation + UE::MassAvoidance::DebugAgentHeightOffset, static_cast<float>(AgentRadius+1.),
+					UE::MassAvoidance::CurrentAgentColor, Text);
 
 				DebugDrawSphere(BaseDebugContext, AgentLocation, 10.f, UE::MassAvoidance::CurrentAgentColor);
 
 				// Draw current velocity (black)
-				DebugDrawVelocity(VelocitiesDebugContext, AgentLocation + UE::MassAvoidance::DebugAgentHeightOffset, AgentLocation + UE::MassAvoidance::DebugAgentHeightOffset + AgentVelocity, UE::MassAvoidance::VelocityColor);
+				DebugDrawVelocity(VelocitiesDebugContext, AgentLocation + UE::MassAvoidance::DebugAgentHeightOffset,
+					AgentLocation + UE::MassAvoidance::DebugAgentHeightOffset + AgentVelocity, UE::MassAvoidance::VelocityColor);
 
 				// Draw preferred velocity (red)
-//				DebugDrawVelocity(VelocitiesDebugContext, AgentLocation + UE::MassAvoidance::DebugAgentHeightOffset, AgentLocation + UE::MassAvoidance::DebugAgentHeightOffset + PrefVelocity, UE::MassAvoidance::PrefVelocityColor);
+//				DebugDrawVelocity(VelocitiesDebugContext, AgentLocation + UE::MassAvoidance::DebugAgentHeightOffset,
+//					AgentLocation + UE::MassAvoidance::DebugAgentHeightOffset + PrefVelocity, UE::MassAvoidance::PrefVelocityColor);
 
 				// Draw initial steering force
-				DebugDrawArrow(BaseDebugContext, AgentLocation + UE::MassAvoidance::DebugAgentHeightOffset, AgentLocation +UE::MassAvoidance:: DebugAgentHeightOffset + SteeringForce, UE::MassAvoidance::CurrentAgentColor, UE::MassAvoidance::BigArrowHeadSize, UE::MassAvoidance::BigArrowThickness);
+				DebugDrawArrow(BaseDebugContext, AgentLocation + UE::MassAvoidance::DebugAgentHeightOffset,
+					AgentLocation +UE::MassAvoidance:: DebugAgentHeightOffset + SteeringForce,
+					UE::MassAvoidance::CurrentAgentColor, UE::MassAvoidance::BigArrowHeadSize, UE::MassAvoidance::BigArrowThickness);
 
 				// Draw center
 				DebugDrawSphere(BaseDebugContext, AgentLocation, /*Radius*/2.f, UE::MassAvoidance::CurrentAgentColor);
@@ -531,7 +536,8 @@ void UMassMovingAvoidanceProcessor::Execute(FMassEntityManager& EntityManager, F
 
 #if WITH_MASSGAMEPLAY_DEBUG && UNSAFE_FOR_MT
 				// Draw desired velocity (yellow)
-				UE::MassAvoidance::DebugDrawVelocity(VelocitiesDebugContext, AgentLocation + UE::MassAvoidance::DebugAgentHeightOffset, AgentLocation + UE::MassAvoidance::DebugAgentHeightOffset + DesiredVelocity, UE::MassAvoidance::DesiredVelocityColor);
+				UE::MassAvoidance::DebugDrawVelocity(VelocitiesDebugContext, AgentLocation + UE::MassAvoidance::DebugAgentHeightOffset,
+					AgentLocation + UE::MassAvoidance::DebugAgentHeightOffset + DesiredVelocity, UE::MassAvoidance::DesiredVelocityColor);
 #endif // WITH_MASSGAMEPLAY_DEBUG
 
 				OldSteeringForce = SteeringForce;
@@ -545,18 +551,19 @@ void UMassMovingAvoidanceProcessor::Execute(FMassEntityManager& EntityManager, F
 					FVector::FReal EdgeLength = 0.;
 					EdgeDiff.ToDirectionAndLength(EdgeDir, EdgeLength);
 
-					const FVector AgentToEdgeStart = AgentLocation - Edge.Start;
-					const FVector::FReal DistAlongEdge = FVector::DotProduct(EdgeDir, AgentToEdgeStart);
-					const FVector::FReal DistAwayFromEdge = FVector::DotProduct(Edge.LeftDir, AgentToEdgeStart);
+					const FVector EdgeStartToAgent = AgentLocation - Edge.Start;
+					const FVector::FReal DistAlongEdge = FVector::DotProduct(EdgeDir, EdgeStartToAgent);
+					const FVector::FReal DistAwayFromEdge = FVector::DotProduct(Edge.LeftDir, EdgeStartToAgent);
 
 					FVector::FReal ConDist = 0.;
 					FVector ConNorm = FVector::ForwardVector;
 					FVector ConPos = FVector::ZeroVector;
 					bool bDirectlyBehindEdge = false;
+					bool bAlongTheEdge = false;
 					
 					if (DistAwayFromEdge < 0.)
 					{
-						// Inside or behind the edge
+						// Behind the edge
 						if (DistAlongEdge < 0.)
 						{
 							ConPos = Edge.Start;
@@ -575,6 +582,7 @@ void UMassMovingAvoidanceProcessor::Execute(FMassEntityManager& EntityManager, F
 							ConNorm = Edge.LeftDir;
 							ConDist = 0.;
 							bDirectlyBehindEdge = true;
+							bAlongTheEdge = true;
 						}
 					}
 					else
@@ -583,28 +591,30 @@ void UMassMovingAvoidanceProcessor::Execute(FMassEntityManager& EntityManager, F
 						{
 							// Start Corner
 							ConPos = Edge.Start;
-							EdgeDiff.ToDirectionAndLength(ConNorm, ConDist);
 						}
 						else if (DistAlongEdge > EdgeLength)
 						{
 							// End Corner
 							ConPos = Edge.End;
-							EdgeDiff.ToDirectionAndLength(ConNorm, ConDist);
 						}
 						else
 						{
 							// Front
 							ConPos = Edge.Start + EdgeDir * DistAlongEdge;
-							ConNorm = Edge.LeftDir;
-							ConDist = DistAwayFromEdge;
+							bAlongTheEdge = true;
 						}
+
+						ConNorm = Edge.LeftDir;
+						ConDist = DistAwayFromEdge;
+
 					}
 					
 					// Check to merge contacts
 					bool bAdd = true;
 					for (int ContactIndex = 0; ContactIndex < Contacts.Num(); ContactIndex++)
 					{
-						if (FVector::DotProduct(Contacts[ContactIndex].Normal, ConNorm) > 0. && FMath::Abs(FVector::DotProduct(ConNorm, Contacts[ContactIndex].Position - ConPos)) < (10./*cm*/))
+						if (FVector::DotProduct(Contacts[ContactIndex].Normal, ConNorm) > 0. &&
+							FMath::Abs(FVector::DotProduct(ConNorm, Contacts[ContactIndex].Position - ConPos)) < (10./*cm*/))
 						{
 							// Contacts are on same place, merge
 							if (ConDist < Contacts[ContactIndex].Distance)
@@ -613,6 +623,7 @@ void UMassMovingAvoidanceProcessor::Execute(FMassEntityManager& EntityManager, F
 								Contacts[ContactIndex].Position = ConPos;
 								Contacts[ContactIndex].Normal = ConNorm;
 								Contacts[ContactIndex].Distance = ConDist;
+								Contacts[ContactIndex].bRelevantForSeparation = bAlongTheEdge;
 							}
 							bAdd = false;
 							break;
@@ -626,6 +637,7 @@ void UMassMovingAvoidanceProcessor::Execute(FMassEntityManager& EntityManager, F
 						Contact.Position = ConPos;
 						Contact.Normal = ConNorm;
 						Contact.Distance = ConDist;
+						Contact.bRelevantForSeparation = bAlongTheEdge;
 						Contacts.Add(Contact);
 					}
 
@@ -648,18 +660,21 @@ void UMassMovingAvoidanceProcessor::Execute(FMassEntityManager& EntityManager, F
 						const FVector::FReal AvoidPen = (PredictiveAvoidanceAgentRadius + MovingAvoidanceParams.PredictiveAvoidanceDistance) - AvoidDist;
 						const FVector::FReal AvoidMag = FMath::Square(FMath::Clamp(AvoidPen / MovingAvoidanceParams.PredictiveAvoidanceDistance, 0., 1.));
 						const FVector::FReal AvoidMagDist = 1. + FMath::Square(1. - CPA * InvPredictiveAvoidanceTime);
-						const FVector AvoidForce = AvoidNormal * AvoidMag * AvoidMagDist * MovingAvoidanceParams.EnvironmentPredictiveAvoidanceStiffness * NearEndScaling; // Predictive avoidance against environment is tuned down towards the end of the path
+
+						// Predictive avoidance against environment is tuned down towards the end of the path
+						const FVector AvoidForce = AvoidNormal * AvoidMag * AvoidMagDist * MovingAvoidanceParams.EnvironmentPredictiveAvoidanceStiffness * NearEndScaling; 
 
 						SteeringForce += AvoidForce;
 
 #if WITH_MASSGAMEPLAY_DEBUG && UNSAFE_FOR_MT
 						// Draw contact normal
-						UE::MassAvoidance::DebugDrawArrow(ObstacleDebugContext, ConPos, ConPos + 50. * ConNorm, UE::MassAvoidance::ObstacleContactNormalColor, /*HeadSize=*/ 5.f);
+						UE::MassAvoidance::DebugDrawArrow(ObstacleDebugContext, ConPos, ConPos + (ConDist * ConNorm), UE::MassAvoidance::ObstacleContactNormalColor, /*HeadSize=*/ 5.f);
 						UE::MassAvoidance::DebugDrawSphere(ObstacleDebugContext, ConPos, 2.5f, UE::MassAvoidance::ObstacleContactNormalColor);
 
 						// Draw hit pos with edge
 						UE::MassAvoidance::DebugDrawLine(ObstacleDebugContext, AgentLocation, HitAgentPos, UE::MassAvoidance::ObstacleAvoidForceColor);
-						UE::MassAvoidance::DebugDrawCylinder(ObstacleDebugContext, HitAgentPos, HitAgentPos + UE::MassAvoidance::DebugAgentHeightOffset, static_cast<float>(AgentRadius), UE::MassAvoidance::ObstacleAvoidForceColor);
+						UE::MassAvoidance::DebugDrawCylinder(ObstacleDebugContext, HitAgentPos, HitAgentPos + UE::MassAvoidance::DebugAgentHeightOffset,
+							static_cast<float>(AgentRadius), UE::MassAvoidance::ObstacleAvoidForceColor);
 
 						// Draw avoid obstacle force
 						UE::MassAvoidance::DebugDrawForce(ObstacleDebugContext, HitObPos, HitObPos + AvoidForce, UE::MassAvoidance::ObstacleAvoidForceColor);
@@ -678,23 +693,26 @@ void UMassMovingAvoidanceProcessor::Execute(FMassEntityManager& EntityManager, F
 
 				// Process contacts to add edge separation force
 				const FVector SteeringForceBeforeSeparation = SteeringForce;
-				for (int ContactIndex = 0; ContactIndex < Contacts.Num(); ContactIndex++) 
+				for (const FEnvironmentContact& Contact : Contacts) 
 				{
-					const FVector ConNorm = Contacts[ContactIndex].Normal.GetSafeNormal();
-					const FVector::FReal ContactDist = Contacts[ContactIndex].Distance;
+					if (Contact.bRelevantForSeparation)
+					{
+						const FVector ConNorm = Contact.Normal.GetSafeNormal();
+						const FVector::FReal ContactDist = Contact.Distance;
 
-					// Separation force (stay away from obstacles if possible)
-					const FVector::FReal SeparationPenalty = (SeparationAgentRadius + MovingAvoidanceParams.EnvironmentSeparationDistance) - ContactDist;
-					const FVector::FReal SeparationMag = UE::MassNavigation::Smooth(FMath::Clamp(SeparationPenalty / MovingAvoidanceParams.EnvironmentSeparationDistance, 0., 1.));
-					const FVector SeparationForce = ConNorm * MovingAvoidanceParams.EnvironmentSeparationStiffness * SeparationMag;
+						// Separation force (stay away from obstacles if possible)
+						const FVector::FReal SeparationPenalty = (SeparationAgentRadius + MovingAvoidanceParams.EnvironmentSeparationDistance) - ContactDist;
+						const FVector::FReal SeparationMag = UE::MassNavigation::Smooth(FMath::Clamp(SeparationPenalty / MovingAvoidanceParams.EnvironmentSeparationDistance, 0., 1.));
+						const FVector SeparationForce = ConNorm * MovingAvoidanceParams.EnvironmentSeparationStiffness * SeparationMag;
 
-					SteeringForce += SeparationForce;
+						SteeringForce += SeparationForce;
 
-#if WITH_MASSGAMEPLAY_DEBUG && UNSAFE_FOR_MT
-					// Draw individual contact forces
-					DebugDrawForce(ObstacleDebugContext, Contacts[ContactIndex].Position + UE::MassAvoidance::DebugAgentHeightOffset,
-					Contacts[ContactIndex].Position + SeparationForce + UE::MassAvoidance::DebugAgentHeightOffset, UE::MassAvoidance::ObstacleSeparationForceColor);
-#endif // WITH_MASSGAMEPLAY_DEBUG
+	#if WITH_MASSGAMEPLAY_DEBUG && UNSAFE_FOR_MT
+						// Draw individual contact forces
+						DebugDrawForce(ObstacleDebugContext, Contact.Position + UE::MassAvoidance::DebugAgentHeightOffset,
+						Contact.Position + SeparationForce + UE::MassAvoidance::DebugAgentHeightOffset, UE::MassAvoidance::ObstacleSeparationForceColor);
+	#endif // WITH_MASSGAMEPLAY_DEBUG
+					}
 				}
 				
 #if WITH_MASSGAMEPLAY_DEBUG && UNSAFE_FOR_MT
@@ -713,7 +731,8 @@ void UMassMovingAvoidanceProcessor::Execute(FMassEntityManager& EntityManager, F
 						DebugDrawLine(ObstacleDebugContext, UE::MassAvoidance::DebugAgentHeightOffset + Edge.Start,
 							UE::MassAvoidance::DebugAgentHeightOffset + Edge.End, UE::MassAvoidance::ObstacleColor, /*Thickness=*/2.f);
 						const FVector Middle = UE::MassAvoidance::DebugAgentHeightOffset + 0.5f * (Edge.Start + Edge.End);
-						DebugDrawArrow(ObstacleDebugContext, Middle, Middle + 10. * FVector::CrossProduct((Edge.End - Edge.Start), FVector::UpVector).GetSafeNormal(), UE::MassAvoidance::ObstacleColor, /*HeadSize=*/2.f);
+						DebugDrawArrow(ObstacleDebugContext, Middle, Middle + 10. * FVector::CrossProduct((Edge.End - Edge.Start), FVector::UpVector).GetSafeNormal(),
+							UE::MassAvoidance::ObstacleColor, /*HeadSize=*/2.f);
 					}
 				}
 #endif // WITH_MASSGAMEPLAY_DEBUG
@@ -728,7 +747,9 @@ void UMassMovingAvoidanceProcessor::Execute(FMassEntityManager& EntityManager, F
 
 			// Find close obstacles
 			const FNavigationObstacleHashGrid2D& AvoidanceObstacleGrid = NavigationSubsystem->GetObstacleGridMutable();
-			UE::MassAvoidance::FindCloseObstacles(AgentLocation, MovingAvoidanceParams.ObstacleDetectionDistance, AvoidanceObstacleGrid, CloseEntities, UE::MassAvoidance::MaxObstacleResults);
+
+			UE::MassAvoidance::FindCloseObstacles(AgentLocation, MovingAvoidanceParams.ObstacleDetectionDistance,
+				AvoidanceObstacleGrid, CloseEntities, UE::MassAvoidance::MaxObstacleResults);
 
 			// Remove unwanted and find the closests in the CloseEntities
 			const FVector::FReal DistanceCutOffSqr = FMath::Square(MovingAvoidanceParams.ObstacleDetectionDistance);
@@ -773,9 +794,9 @@ void UMassMovingAvoidanceProcessor::Execute(FMassEntityManager& EntityManager, F
 
 			// Fill collider list from close agents
 			Colliders.Reset();
-			constexpr int32 MaxColliders = 6;
 			for (int32 Index = 0; Index < ClosestObstacles.Num(); Index++)
 			{
+				constexpr int32 MaxColliders = 6;
 				if (Colliders.Num() >= MaxColliders)
 				{
 					break;
@@ -901,8 +922,9 @@ void UMassMovingAvoidanceProcessor::Execute(FMassEntityManager& EntityManager, F
 					const FVector::FReal Blend = FMath::Max(0., -FVector::DotProduct(ConNorm, RelVelNorm));
 					SeparationNormal = FMath::Lerp(ConNorm, ForcedNormal, Blend).GetSafeNormal();
 				}
-				
-				const FVector::FReal StandingScaling = Collider.bIsMoving ? 1. : MovingAvoidanceParams.StandingObstacleAvoidanceScale; // Care less about standing agents so that we can push through standing crowd.
+
+				// Care less about standing agents so that we can push through standing crowd.
+				const FVector::FReal StandingScaling = Collider.bIsMoving ? 1. : MovingAvoidanceParams.StandingObstacleAvoidanceScale; 
 				
 				// Separation force (stay away from agents if possible)
 				const FVector::FReal PenSep = (SeparationAgentRadius + Collider.Radius + MovingAvoidanceParams.ObstacleSeparationDistance) - ConDist;
@@ -913,7 +935,7 @@ void UMassMovingAvoidanceProcessor::Execute(FMassEntityManager& EntityManager, F
 				SteeringForce += SeparationForce;
 				TotalAgentSeparationForce += SeparationForce;
 
-				// Calculate closest point of approach based on relative agent positions and velocities.
+				// Calculate the closest point of approach based on relative agent positions and velocities.
 				const FVector::FReal CPA = UE::MassAvoidance::ComputeClosestPointOfApproach(RelPos, RelVel, PredictiveAvoidanceAgentRadius + Collider.Radius, MovingAvoidanceParams.PredictiveAvoidanceTime);
 
 				// Calculate penetration at CPA
@@ -939,11 +961,13 @@ void UMassMovingAvoidanceProcessor::Execute(FMassEntityManager& EntityManager, F
 
 #if WITH_MASSGAMEPLAY_DEBUG && UNSAFE_FOR_MT
 				// Display close agent
-				UE::MassAvoidance::DebugDrawCylinder(AgentDebugContext, Collider.Location, Collider.Location + UE::MassAvoidance::DebugLowCylinderOffset, Collider.Radius, UE::MassAvoidance::AgentsColor);
+				UE::MassAvoidance::DebugDrawCylinder(AgentDebugContext, Collider.Location,
+					Collider.Location + UE::MassAvoidance::DebugLowCylinderOffset, Collider.Radius, UE::MassAvoidance::AgentsColor);
 
 				if (bHasForcedNormal)
 				{
-					UE::MassAvoidance::DebugDrawCylinder(BaseDebugContext, Collider.Location, Collider.Location + UE::MassAvoidance::DebugAgentHeightOffset, Collider.Radius, FColor::Red);
+					UE::MassAvoidance::DebugDrawCylinder(BaseDebugContext, Collider.Location,
+						Collider.Location + UE::MassAvoidance::DebugAgentHeightOffset, Collider.Radius, FColor::Red);
 				}
 
 				// Draw agent contact separation force
@@ -957,9 +981,12 @@ void UMassMovingAvoidanceProcessor::Execute(FMassEntityManager& EntityManager, F
 					// Draw agent vs agent hit positions
 					const FVector HitPosition = AgentLocation + (DesVel * CPA);
 					const FVector LeftOffset = PredictiveAvoidanceAgentRadius * UE::MassNavigation::GetLeftDirection(DesVel.GetSafeNormal(), FVector::UpVector);
-					UE::MassAvoidance::DebugDrawLine(AgentDebugContext, AgentLocation + UE::MassAvoidance::DebugAgentHeightOffset + LeftOffset, HitPosition + UE::MassAvoidance::DebugAgentHeightOffset + LeftOffset, UE::MassAvoidance::CurrentAgentColor, 1.5f);
-					UE::MassAvoidance::DebugDrawLine(AgentDebugContext, AgentLocation + UE::MassAvoidance::DebugAgentHeightOffset - LeftOffset, HitPosition + UE::MassAvoidance::DebugAgentHeightOffset - LeftOffset, UE::MassAvoidance::CurrentAgentColor, 1.5f);
-					UE::MassAvoidance::DebugDrawCylinder(AgentDebugContext, HitPosition, HitPosition + UE::MassAvoidance::DebugAgentHeightOffset, static_cast<float>(PredictiveAvoidanceAgentRadius), UE::MassAvoidance::CurrentAgentColor);
+					UE::MassAvoidance::DebugDrawLine(AgentDebugContext, AgentLocation + UE::MassAvoidance::DebugAgentHeightOffset + LeftOffset,
+						HitPosition + UE::MassAvoidance::DebugAgentHeightOffset + LeftOffset, UE::MassAvoidance::CurrentAgentColor, 1.5f);
+					UE::MassAvoidance::DebugDrawLine(AgentDebugContext, AgentLocation + UE::MassAvoidance::DebugAgentHeightOffset - LeftOffset,
+						HitPosition + UE::MassAvoidance::DebugAgentHeightOffset - LeftOffset, UE::MassAvoidance::CurrentAgentColor, 1.5f);
+					UE::MassAvoidance::DebugDrawCylinder(AgentDebugContext, HitPosition,
+						HitPosition + UE::MassAvoidance::DebugAgentHeightOffset, static_cast<float>(PredictiveAvoidanceAgentRadius), UE::MassAvoidance::CurrentAgentColor);
 
 					const FVector OtherHitPosition = Collider.Location + (Collider.Velocity * CPA);
 					const FVector OtherLeftOffset = Collider.Radius * UE::MassNavigation::GetLeftDirection(Collider.Velocity.GetSafeNormal(), FVector::UpVector);
@@ -967,8 +994,10 @@ void UMassMovingAvoidanceProcessor::Execute(FMassEntityManager& EntityManager, F
 					const FVector Right = UE::MassAvoidance::DebugAgentHeightOffset - OtherLeftOffset;
 					UE::MassAvoidance::DebugDrawLine(AgentDebugContext, Collider.Location + Left, OtherHitPosition + Left, UE::MassAvoidance::AgentsColor, 1.5f);
 					UE::MassAvoidance::DebugDrawLine(AgentDebugContext, Collider.Location + Right, OtherHitPosition + Right, UE::MassAvoidance::AgentsColor, 1.5f);
-					UE::MassAvoidance::DebugDrawCylinder(AgentDebugContext, Collider.Location, Collider.Location + UE::MassAvoidance::DebugAgentHeightOffset, static_cast<float>(AgentRadius), UE::MassAvoidance::AgentsColor);
-					UE::MassAvoidance::DebugDrawCylinder(AgentDebugContext, OtherHitPosition, OtherHitPosition + UE::MassAvoidance::DebugAgentHeightOffset, static_cast<float>(AgentRadius), UE::MassAvoidance::AgentsColor);
+					UE::MassAvoidance::DebugDrawCylinder(AgentDebugContext, Collider.Location, Collider.Location + UE::MassAvoidance::DebugAgentHeightOffset,
+						static_cast<float>(AgentRadius), UE::MassAvoidance::AgentsColor);
+					UE::MassAvoidance::DebugDrawCylinder(AgentDebugContext, OtherHitPosition, OtherHitPosition + UE::MassAvoidance::DebugAgentHeightOffset,
+						static_cast<float>(AgentRadius), UE::MassAvoidance::AgentsColor);
 
 					// Draw agent avoid force
 					UE::MassAvoidance::DebugDrawForce(AgentDebugContext,
@@ -1075,7 +1104,12 @@ void UMassStandingAvoidanceProcessor::Execute(FMassEntityManager& EntityManager,
 		struct FSortedObstacle
 		{
 			FSortedObstacle() = default;
-			FSortedObstacle(const FMassEntityHandle InEntity, const FVector InLocation, const FVector InForward, const FVector::FReal InDistSq) : Entity(InEntity), Location(InLocation), Forward(InForward), DistSq(InDistSq) {}
+			FSortedObstacle(const FMassEntityHandle InEntity, const FVector InLocation, const FVector InForward, const FVector::FReal InDistSq)
+				: Entity(InEntity)
+				, Location(InLocation)
+				, Forward(InForward)
+				, DistSq(InDistSq)
+			{}
 			
 			FMassEntityHandle Entity;
 			FVector Location = FVector::ZeroVector;
@@ -1185,7 +1219,7 @@ void UMassStandingAvoidanceProcessor::Execute(FMassEntityManager& EntityManager,
 				// If other has ghost active, avoid that, else avoid the actual agent.
 				if (bOtherHasGhost)
 				{
-					// Avoid the other agent more, when it is further away from it's goal location.
+					// Avoid the other agent more, when it is further away from its goal location.
 					const FVector::FReal OtherDistanceToGoal = FVector::Distance(OtherGhost->Location, OtherMoveTarget->Center);
 					const FVector::FReal OtherSteerFade = FMath::Clamp(OtherDistanceToGoal / StandingParams.GhostToTargetMaxDeviation, 0., 1.);
 					const FVector::FReal SeparationStiffness = FMath::Lerp(GhostSeparationStiffness, MovingSeparationStiffness, OtherSteerFade);
@@ -1239,7 +1273,7 @@ void UMassStandingAvoidanceProcessor::Execute(FMassEntityManager& EntityManager,
 			
 			Ghost.Location += Ghost.Velocity * DeltaTime;
 
-			// Dont let the ghost location too far from move target center.
+			// Don't let the ghost location too far from move target center.
 			const FVector DirToCenter = Ghost.Location - MoveTarget.Center;
 			const FVector::FReal DistToCenter = DirToCenter.Length();
 			if (DistToCenter > StandingParams.GhostToTargetMaxDeviation)
