Fix bug when setting NavPathNextStartIndex on FMassNavMeshCachedPathFragment
#rnx

diff --git a/Engine/Plugins/AI/MassAI/Source/MassAIBehavior/Private/Tasks/MassNavMeshPathFollowTask.cpp b/Engine/Plugins/AI/MassAI/Source/MassAIBehavior/Private/Tasks/MassNavMeshPathFollowTask.cpp
--- a/Engine/Plugins/AI/MassAI/Source/MassAIBehavior/Private/Tasks/MassNavMeshPathFollowTask.cpp
+++ b/Engine/Plugins/AI/MassAI/Source/MassAIBehavior/Private/Tasks/MassNavMeshPathFollowTask.cpp
@@ -90,7 +90,7 @@ bool FMassNavMeshPathFollowTask::RequestPath(FStateTreeExecutionContext& Context
 		FMassNavMeshShortPathFragment& ShortPathFragment = Context.GetExternalData(ShortPathHandle);
 		ShortPathFragment.RequestShortPath(CachedPathFragment.Corridor, 0);
 
-		CachedPathFragment.NavPathNextStartIndex += ShortPathFragment.NumPoints-1;
+		CachedPathFragment.NavPathNextStartIndex = ShortPathFragment.NumPoints-1;
 		
 		// Update MoveTarget
 		FMassMoveTargetFragment& MoveTarget = Context.GetExternalData(MoveTargetHandle);
