Change LogMassBehavior to be UE_VLOG_ALWAYS_UELOG (not just vlog anymore)

diff --git a/Engine/Plugins/AI/MassAI/Source/MassAIBehavior/Public/MassAIBehaviorTypes.h b/Engine/Plugins/AI/MassAI/Source/MassAIBehavior/Public/MassAIBehaviorTypes.h
--- a/Engine/Plugins/AI/MassAI/Source/MassAIBehavior/Public/MassAIBehaviorTypes.h
+++ b/Engine/Plugins/AI/MassAI/Source/MassAIBehavior/Public/MassAIBehaviorTypes.h
@@ -28,9 +28,9 @@ namespace UE::Mass::ProcessorGroupNames
  *		MASSBEHAVIOR_CLOG(bDisplayDebug, Log, TEXT("Starting action: %s"), *StaticEnum<ESomeActionEnum>()->GetValueAsString(SomeActionEnumValue));
  */
 #if WITH_MASSGAMEPLAY_DEBUG
-#define MASSBEHAVIOR_LOG(Verbosity, Format, ...) UE_VLOG(static_cast<FMassStateTreeExecutionContext&>(Context).GetOwner(), LogMassBehavior, Verbosity, \
+#define MASSBEHAVIOR_LOG(Verbosity, Format, ...) UE_VLOG_ALWAYS_UELOG(static_cast<FMassStateTreeExecutionContext&>(Context).GetOwner(), LogMassBehavior, Verbosity, \
 	TEXT("Entity [%s][%s] ") Format, *static_cast<FMassStateTreeExecutionContext&>(Context).GetEntity().DebugGetDescription(), *StaticStruct()->GetName(), ##__VA_ARGS__)
-#define MASSBEHAVIOR_CLOG(Condition, Verbosity, Format, ...) UE_CVLOG((Condition), static_cast<FMassStateTreeExecutionContext&>(Context).GetOwner(), LogMassBehavior, Verbosity, \
+#define MASSBEHAVIOR_CLOG(Condition, Verbosity, Format, ...) UE_CVLOG_ALWAYS_UELOG((Condition), static_cast<FMassStateTreeExecutionContext&>(Context).GetOwner(), LogMassBehavior, Verbosity, \
 	TEXT("Entity [%s][%s] ") Format, *static_cast<FMassStateTreeExecutionContext&>(Context).GetEntity().DebugGetDescription(), *StaticStruct()->GetName(), ##__VA_ARGS__)
 #else
 #define MASSBEHAVIOR_LOG(Verbosity, Format, ...)
