[Mass] Added LookAt properties related to target location (main and gaze)

diff --git a/Engine/Plugins/AI/MassAI/Source/MassAIBehavior/Private/MassLookAtProcessors.cpp b/Engine/Plugins/AI/MassAI/Source/MassAIBehavior/Private/MassLookAtProcessors.cpp
--- a/Engine/Plugins/AI/MassAI/Source/MassAIBehavior/Private/MassLookAtProcessors.cpp
+++ b/Engine/Plugins/AI/MassAI/Source/MassAIBehavior/Private/MassLookAtProcessors.cpp
@@ -166,6 +166,7 @@ void UMassLookAtProcessor::Execute(FMassEntityManager& EntityManager, FMassExecu
 
 				// Update specific look at mode.
 				LookAt.Direction = FVector::ForwardVector;
+				LookAt.MainTargetLocation = FVector::ZeroVector;
 
 				switch (LookAt.LookAtMode)
 				{
@@ -247,6 +248,7 @@ void UMassLookAtProcessor::FindNewGazeTarget(const UMassNavigationSubsystem& Mas
 	
 	LookAt.GazeTrackedEntity.Reset();
 	LookAt.GazeDirection = FVector::ForwardVector;
+	LookAt.GazeTargetLocation = FVector::ZeroVector;
 
 	// Search for potential targets in front
 	bool bTargetFound = false;
@@ -298,6 +300,7 @@ void UMassLookAtProcessor::FindNewGazeTarget(const UMassNavigationSubsystem& Mas
 			{
 				LookAt.GazeDirection = Direction;
 				LookAt.GazeTrackedEntity = NearbyEntity.Entity;
+				LookAt.GazeTargetLocation = TargetLocation;
 				bTargetFound = true;
 				break;
 			}
@@ -312,6 +315,7 @@ void UMassLookAtProcessor::FindNewGazeTarget(const UMassNavigationSubsystem& Mas
 				{
 					LookAt.GazeDirection = Direction;
 					LookAt.GazeTrackedEntity = NearbyEntity.Entity;
+					LookAt.GazeTargetLocation = TargetLocation;
 					bTargetFound = true;
 					break;
 				}
@@ -325,6 +329,7 @@ void UMassLookAtProcessor::FindNewGazeTarget(const UMassNavigationSubsystem& Mas
 	{
 		const FRotator Rot(FMath::FRandRange(-(float)LookAt.RandomGazePitchVariation, (float)LookAt.RandomGazePitchVariation),FMath::FRandRange(-(float)LookAt.RandomGazeYawVariation, (float)LookAt.RandomGazeYawVariation), 0.f);
 		LookAt.GazeDirection = UE::MassBehavior::ClampDirectionToXAxisCone(Rot.Vector(), FMath::DegreesToRadians(AngleThresholdInDegrees));
+		LookAt.GazeTargetLocation = FVector::ZeroVector;
 	}
 
 	// @todo: This does not currently carry over time. It's intentional, since there might be big gaps between updates.
@@ -346,7 +351,8 @@ void UMassLookAtProcessor::UpdateLookAtTrajectory(const FTransform& Transform, c
 		const FVector NewGlobalDirection = (AnticipatedPosition - AgentPosition).GetSafeNormal();
 		LookAt.Direction = Transform.InverseTransformVector(NewGlobalDirection);
 		LookAt.Direction.Z = 0.0f;
-						
+		LookAt.MainTargetLocation = AnticipatedPosition;
+
 #if WITH_MASSGAMEPLAY_DEBUG && UNSAFE_FOR_MT
 		if (bDisplayDebug)
 		{
@@ -365,8 +371,10 @@ void UMassLookAtProcessor::UpdateLookAtTrackedEntity(const FMassEntityManager& E
 		if (const FTransformFragment* TargetTransform = EntityManager.GetFragmentDataPtr<FTransformFragment>(LookAt.TrackedEntity))
 		{
 			const FVector AgentPosition = Transform.GetLocation();
-			const FVector NewGlobalDirection = (TargetTransform->GetTransform().GetLocation() - AgentPosition).GetSafeNormal();
+			const FVector TargetLocation = TargetTransform->GetTransform().GetLocation();
+			const FVector NewGlobalDirection = (TargetLocation - AgentPosition).GetSafeNormal();
 			LookAt.Direction = Transform.InverseTransformVector(NewGlobalDirection);
+			LookAt.MainTargetLocation = TargetLocation;
 
 #if WITH_MASSGAMEPLAY_DEBUG && UNSAFE_FOR_MT
 			if (bDisplayDebug)
@@ -389,9 +397,10 @@ bool UMassLookAtProcessor::UpdateGazeTrackedEntity(const FMassEntityManager& Ent
 		if (const FTransformFragment* TargetTransform = EntityManager.GetFragmentDataPtr<FTransformFragment>(LookAt.GazeTrackedEntity))
 		{
 			const FVector AgentPosition = Transform.GetLocation();
-			const FVector NewGlobalDirection = (TargetTransform->GetTransform().GetLocation() - AgentPosition).GetSafeNormal();
+			const FVector TargetLocation = TargetTransform->GetTransform().GetLocation();
+			const FVector NewGlobalDirection = (TargetLocation - AgentPosition).GetSafeNormal();
 			LookAt.GazeDirection = Transform.InverseTransformVector(NewGlobalDirection);
-
+			LookAt.GazeTargetLocation = TargetLocation;
 			bHasTarget = true;
 			
 #if WITH_MASSGAMEPLAY_DEBUG && UNSAFE_FOR_MT
diff --git a/Engine/Plugins/AI/MassAI/Source/MassAIBehavior/Public/MassLookAtFragments.h b/Engine/Plugins/AI/MassAI/Source/MassAIBehavior/Public/MassLookAtFragments.h
--- a/Engine/Plugins/AI/MassAI/Source/MassAIBehavior/Public/MassLookAtFragments.h
+++ b/Engine/Plugins/AI/MassAI/Source/MassAIBehavior/Public/MassLookAtFragments.h
@@ -48,6 +48,8 @@ struct MASSAIBEHAVIOR_API FMassLookAtFragment : public FMassFragment
 	
 	void Reset()
 	{
+		MainTargetLocation = FVector::ZeroVector;
+		GazeTargetLocation = FVector::ZeroVector;
 		Direction = FVector::ForwardVector;
 		GazeDirection = FVector::ForwardVector;
 		TrackedEntity.Reset();
@@ -62,6 +64,14 @@ struct MASSAIBEHAVIOR_API FMassLookAtFragment : public FMassFragment
 		bRandomGazeEntities = false;
 	}
 
+	/** Current look at main target location. */
+	UPROPERTY(Transient)
+	FVector MainTargetLocation = FVector::ZeroVector;
+
+	/** Current gaze target location. */
+	UPROPERTY(Transient)
+	FVector GazeTargetLocation = FVector::ZeroVector;
+
 	/** Current look at direction (with gaze applied). */
 	UPROPERTY(Transient)
 	FVector Direction = FVector::ForwardVector;
diff --git a/Engine/Plugins/AI/MassAI/Source/MassAIDebug/Private/GameplayDebuggerCategory_Mass.cpp b/Engine/Plugins/AI/MassAI/Source/MassAIDebug/Private/GameplayDebuggerCategory_Mass.cpp
--- a/Engine/Plugins/AI/MassAI/Source/MassAIDebug/Private/GameplayDebuggerCategory_Mass.cpp
+++ b/Engine/Plugins/AI/MassAI/Source/MassAIDebug/Private/GameplayDebuggerCategory_Mass.cpp
@@ -606,26 +606,20 @@ void FGameplayDebuggerCategory_Mass::CollectData(APlayerController* OwnerPC, AAc
 						bool bLookArrowDrawn = false;
 						if (LookAt.LookAtMode == EMassLookAtMode::LookAtEntity && EntityManager.IsEntityValid(LookAt.TrackedEntity))
 						{
-							if (const FTransformFragment* TargetTransform = EntityManager.GetFragmentDataPtr<FTransformFragment>(LookAt.TrackedEntity))
-							{
-								FVector TargetPosition = TargetTransform->GetTransform().GetLocation();
-								TargetPosition.Z = BasePos.Z;
-								AddShape(FGameplayDebuggerShape::MakeCircle(TargetPosition, FVector::UpVector, Radius.Radius, FColor::Red));
+							FVector TargetPosition = LookAt.MainTargetLocation;
+							TargetPosition.Z = BasePos.Z;
+							AddShape(FGameplayDebuggerShape::MakeCircle(TargetPosition, FVector::UpVector, Radius.Radius, FColor::Red));
 
-								const FVector::FReal TargetDistance = FMath::Max(LookArrowLength, FVector::DotProduct(WorldLookDirection, TargetPosition - BasePos));
-								AddShape(FGameplayDebuggerShape::MakeSegment(BasePos, BasePos + WorldLookDirection * TargetDistance, FColorList::LightGrey));
-								bLookArrowDrawn = true;
-							}
+							const FVector::FReal TargetDistance = FMath::Max(LookArrowLength, FVector::DotProduct(WorldLookDirection, TargetPosition - BasePos));
+							AddShape(FGameplayDebuggerShape::MakeSegment(BasePos, BasePos + WorldLookDirection * TargetDistance, FColorList::LightGrey));
+							bLookArrowDrawn = true;
 						}
 
 						if (LookAt.bRandomGazeEntities && EntityManager.IsEntityValid(LookAt.GazeTrackedEntity))
 						{
-							if (const FTransformFragment* TargetTransform = EntityManager.GetFragmentDataPtr<FTransformFragment>(LookAt.GazeTrackedEntity))
-							{
-								FVector TargetPosition = TargetTransform->GetTransform().GetLocation();
-								TargetPosition.Z = BasePos.Z;
-								AddShape(FGameplayDebuggerShape::MakeCircle(TargetPosition, FVector::UpVector, Radius.Radius, FColor::Turquoise));
-							}
+							FVector TargetPosition = LookAt.GazeTargetLocation;
+							TargetPosition.Z = BasePos.Z;
+							AddShape(FGameplayDebuggerShape::MakeCircle(TargetPosition, FVector::UpVector, Radius.Radius, FColor::Turquoise));
 						}
 
 						if (!bLookArrowDrawn)
