Mass: State tree dependencies. Add link function to schema, similar to link for nodes. Add a dependency builder to help users add read or write dependencies. The dependencies are built at runtime to support hotfixes. The dependencies are built once, when the state tree asset loads or when it compiles.

diff --git a/Engine/Plugins/AI/MassAI/Source/MassAIBehavior/Private/MassStateTreeDependency.cpp b/Engine/Plugins/AI/MassAI/Source/MassAIBehavior/Private/MassStateTreeDependency.cpp
new file mode 100644
--- /dev/null
+++ b/Engine/Plugins/AI/MassAI/Source/MassAIBehavior/Private/MassStateTreeDependency.cpp
@@ -0,0 +1,34 @@
+// Copyright Epic Games, Inc. All Rights Reserved.
+
+#include "MassStateTreeDependency.h"
+
+namespace UE::MassBehavior
+{
+
+	FStateTreeDependencyBuilder::FStateTreeDependencyBuilder(TArray<FMassStateTreeDependency>& InDependencies)
+		: Dependencies(InDependencies)
+	{
+	}
+
+	void FStateTreeDependencyBuilder::Add(TNotNull<const UStruct*> Struct, FStateTreeDependencyBuilder::EAccessType Access)
+	{
+		static_assert((int32)EAccessType::ReadOnly == (int32)EMassFragmentAccess::ReadOnly);
+		static_assert((int32)EAccessType::ReadWrite == (int32)EMassFragmentAccess::ReadWrite);
+
+		const EMassFragmentAccess FragmentAccess = static_cast<EMassFragmentAccess>(Access);
+
+		if (FMassStateTreeDependency* Found = Dependencies.FindByPredicate([Struct](const FMassStateTreeDependency& Other) { return Other.Type == Struct; }))
+		{
+			// set the worst case
+			if (Found->Access != FragmentAccess && Access == EAccessType::ReadWrite)
+			{
+				Found->Access = EMassFragmentAccess::ReadWrite;
+			}
+		}
+		else
+		{
+			Dependencies.Emplace(Struct, FragmentAccess);
+		}
+	}
+
+} //namespace
diff --git a/Engine/Plugins/AI/MassAI/Source/MassAIBehavior/Private/MassStateTreeSchema.cpp b/Engine/Plugins/AI/MassAI/Source/MassAIBehavior/Private/MassStateTreeSchema.cpp
--- a/Engine/Plugins/AI/MassAI/Source/MassAIBehavior/Private/MassStateTreeSchema.cpp
+++ b/Engine/Plugins/AI/MassAI/Source/MassAIBehavior/Private/MassStateTreeSchema.cpp
@@ -2,11 +2,19 @@
 
 #include "MassStateTreeSchema.h"
 #include "MassEntityTypes.h"
+#include "MassStateTreeDependency.h"
 #include "MassStateTreeTypes.h"
+#include "StateTree.h"
 #include "StateTreeConditionBase.h"
 #include "StateTreeConsiderationBase.h"
+#include "StateTreeExecutionTypes.h"
+#include "StateTreeLinker.h"
 #include "Subsystems/WorldSubsystem.h"
 
+#ifndef UE_WITH_MASS_STATETREE_DEPENDENCIES_DEBUG
+#define UE_WITH_MASS_STATETREE_DEPENDENCIES_DEBUG 0
+#endif
+
 bool UMassStateTreeSchema::IsStructAllowed(const UScriptStruct* InScriptStruct) const
 {
 	// Only allow Mass evals and tasks,and common conditions.
@@ -25,3 +33,32 @@ bool UMassStateTreeSchema::IsExternalItemAllowed(const UStruct& InStruct) const
 			|| InStruct.IsChildOf(FMassSharedFragment::StaticStruct())
 			|| InStruct.IsChildOf(FMassConstSharedFragment::StaticStruct());
 }
+
+bool UMassStateTreeSchema::Link(FStateTreeLinker& Linker)
+{
+	Dependencies.Empty();
+	UE::MassBehavior::FStateTreeDependencyBuilder Builder(Dependencies);
+
+	const UStateTree* StateTree = CastChecked<UStateTree>(GetOuter());
+	for (FConstStructView Node : StateTree->GetNodes())
+	{
+		if (const FMassStateTreeEvaluatorBase* Evaluator = Node.GetPtr<const FMassStateTreeEvaluatorBase>())
+		{
+			Evaluator->GetDependencies(Builder);
+		}
+		else if (const FMassStateTreeTaskBase* Task = Node.GetPtr<const FMassStateTreeTaskBase>())
+		{
+			Task->GetDependencies(Builder);
+		}
+	}
+
+#if UE_WITH_MASS_STATETREE_DEPENDENCIES_DEBUG
+	for (const FStateTreeExternalDataDesc& Desc : Linker.GetExternalDataDescs())
+	{
+		const bool bContains = Dependencies.ContainsByPredicate([&Desc](const FMassStateTreeDependency& Other) { return Desc.Struct == Other.Type; });
+		ensureMsgf(bContains, TEXT("Tree %s is missing a mass dependency"), *StateTree->GetPathName());
+	}
+#endif
+
+	return true;
+}
\ No newline at end of file
diff --git a/Engine/Plugins/AI/MassAI/Source/MassAIBehavior/Public/MassStateTreeDependency.h b/Engine/Plugins/AI/MassAI/Source/MassAIBehavior/Public/MassStateTreeDependency.h
new file mode 100644
--- /dev/null
+++ b/Engine/Plugins/AI/MassAI/Source/MassAIBehavior/Public/MassStateTreeDependency.h
@@ -0,0 +1,133 @@
+// Copyright Epic Games, Inc. All Rights Reserved.
+
+#pragma once
+
+#include "MassEntityConcepts.h"
+#include "MassRequirements.h"
+#include "Misc/NotNull.h"
+#include "MassStateTreeDependency.generated.h"
+
+/** A dependency the state tree has on other fragment or system used by mass. */
+USTRUCT()
+struct FMassStateTreeDependency
+{
+	GENERATED_BODY()
+
+	FMassStateTreeDependency() = default;
+	FMassStateTreeDependency(TNotNull<const UStruct*> InType, EMassFragmentAccess InAccess)
+		: Type(InType)
+		, Access(InAccess)
+	{}
+
+	UPROPERTY()
+	TObjectPtr<const UStruct> Type;
+
+	UPROPERTY()
+	EMassFragmentAccess Access;
+};
+
+namespace UE::MassBehavior
+{
+	
+/**
+ * FStateTreeDependencyBuilder is a utility struct to build a list of dependencies used by the state tree.
+ * The add dependency can be chain for ease of use.
+ * Example Usage:
+ * {
+ *		FStateTreeDependencyBuilder Builder()
+ *		.AddReadOnly<FTransformFragment>()
+ *		.AddReadWrite(ComponentHitSubsystemHandle);
+ *		Builder.AddReadOnly<FMassZoneGraphAnnotationFragment>();
+ * } 
+ */
+struct FStateTreeDependencyBuilder 
+{
+public:
+	enum class EAccessType : uint8
+	{
+		ReadOnly = (uint8)EMassFragmentAccess::ReadOnly,
+		ReadWrite = (uint8)EMassFragmentAccess::ReadWrite,
+	};
+
+public:
+	MASSAIBEHAVIOR_API FStateTreeDependencyBuilder(TArray<FMassStateTreeDependency>& Dependencies);
+
+	/** Add a read only dependency. */
+	template<typename T>
+	FStateTreeDependencyBuilder& AddReadOnly() requires TIsDerivedFrom<T, UObject>::IsDerived
+	{
+		Add(T::StaticClass(), EAccessType::ReadOnly);
+		return *this;
+	}
+
+	/** Add a read only dependency. */
+	template<UE::Mass::CNonTag T>
+	FStateTreeDependencyBuilder& AddReadOnly()
+	{
+		Add(T::StaticStruct(), EAccessType::ReadOnly);
+		return *this;
+	}
+
+	/** Add a read only dependency from a handle. */
+	template<typename T>
+	FStateTreeDependencyBuilder& AddReadOnly(const T& Handle) requires TIsDerivedFrom<typename T::DataType, UObject>::IsDerived
+	{
+		Add(T::DataType::StaticClass(), EAccessType::ReadOnly);
+		return *this;
+	}
+	
+	/** Add a read only dependency from a handle. */
+	template<UE::Mass::CNonTag T>
+	FStateTreeDependencyBuilder& AddReadOnly(const T& Handle)
+	{
+		Add(T::DataType::StaticStruct(), EAccessType::ReadOnly);
+		return *this;
+	}
+
+	/** Add a read write dependency. */
+	template<typename T>
+	FStateTreeDependencyBuilder& AddReadWrite() requires TIsDerivedFrom<T, UObject>::IsDerived
+	{
+		Add(T::StaticClass(), EAccessType::ReadWrite);
+		return *this;
+	}
+	
+	/** Add a read write dependency. */
+	template<UE::Mass::CNonTag T>
+	FStateTreeDependencyBuilder& AddReadWrite()
+	{
+		Add(T::StaticStruct(), EAccessType::ReadWrite);
+		return *this;
+	}
+
+	/** Add a read write dependency from a handle. */
+	template<typename T>
+	FStateTreeDependencyBuilder& AddReadWrite(const T& Handle) requires TIsDerivedFrom<typename T::DataType, UObject>::IsDerived
+	{
+		Add(T::DataType::StaticClass(), EAccessType::ReadWrite);
+		return *this;
+	}
+
+	/** Add a read write dependency from a handle. */
+	template<UE::Mass::CNonTag T>
+	FStateTreeDependencyBuilder& AddReadWrite(const T& Handle)
+	{
+		Add(T::DataType::StaticStruct(), EAccessType::ReadWrite);
+		return *this;
+	}
+
+	/** Add a dependency for a specific type. */
+	MASSAIBEHAVIOR_API void Add(TNotNull<const UStruct*> Struct, EAccessType Access);
+
+	/** @return the list of dependencies. */
+	const TConstArrayView<FMassStateTreeDependency> GetDependencies() const
+	{
+		return Dependencies;
+	}
+
+private:
+	/** The list of dependencies. */
+	TArray<FMassStateTreeDependency>& Dependencies;
+};
+
+}
diff --git a/Engine/Plugins/AI/MassAI/Source/MassAIBehavior/Public/MassStateTreeSchema.h b/Engine/Plugins/AI/MassAI/Source/MassAIBehavior/Public/MassStateTreeSchema.h
--- a/Engine/Plugins/AI/MassAI/Source/MassAIBehavior/Public/MassStateTreeSchema.h
+++ b/Engine/Plugins/AI/MassAI/Source/MassAIBehavior/Public/MassStateTreeSchema.h
@@ -3,6 +3,7 @@
 #pragma once
 
 #include "StateTreeSchema.h"
+#include "MassStateTreeDependency.h"
 #include "MassStateTreeSchema.generated.h"
 
 /**
@@ -17,5 +18,10 @@ protected:
 
 	virtual bool IsStructAllowed(const UScriptStruct* InScriptStruct) const override;
 	virtual bool IsExternalItemAllowed(const UStruct& InStruct) const override;
+	virtual bool Link(FStateTreeLinker& Linker) override;
+
+protected:
+	UPROPERTY(Transient)
+	TArray<FMassStateTreeDependency> Dependencies;
 };
 
diff --git a/Engine/Plugins/AI/MassAI/Source/MassAIBehavior/Public/MassStateTreeTypes.h b/Engine/Plugins/AI/MassAI/Source/MassAIBehavior/Public/MassStateTreeTypes.h
--- a/Engine/Plugins/AI/MassAI/Source/MassAIBehavior/Public/MassStateTreeTypes.h
+++ b/Engine/Plugins/AI/MassAI/Source/MassAIBehavior/Public/MassStateTreeTypes.h
@@ -9,6 +9,11 @@
 #endif // UE_ENABLE_INCLUDE_ORDER_DEPRECATED_IN_5_6
 #include "MassStateTreeTypes.generated.h"
 
+namespace UE::MassBehavior
+{
+struct FStateTreeDependencyBuilder;
+}
+
 /**
  * Signals used by the StateTree framework in Mass
  */
@@ -30,6 +35,14 @@ USTRUCT(meta = (DisplayName = "Mass Evaluator Base"))
 struct MASSAIBEHAVIOR_API FMassStateTreeEvaluatorBase : public FStateTreeEvaluatorBase
 {
 	GENERATED_BODY()
+
+	/**
+	 * Appends this evaluator's Mass dependencies to the given Builder.
+	 * This is done once once for every evaluator instance, when the state tree asset is loaded or compiled.
+	 */
+	virtual void GetDependencies(UE::MassBehavior::FStateTreeDependencyBuilder& Builder) const
+	{
+	}
 };
 
 /**
@@ -39,6 +52,14 @@ USTRUCT(meta = (DisplayName = "Mass Task Base"))
 struct MASSAIBEHAVIOR_API FMassStateTreeTaskBase : public FStateTreeTaskBase
 {
 	GENERATED_BODY()
+
+	/**
+	 * Appends this task's Mass dependencies to the given Builder.
+	 * This is done once once for every evaluator instance, when the state tree asset is loaded or compiled.
+	 */
+	virtual void GetDependencies(UE::MassBehavior::FStateTreeDependencyBuilder& Builder) const
+	{
+	}
 };
 
 /**
diff --git a/Engine/Plugins/Runtime/StateTree/Source/StateTreeModule/Private/StateTree.cpp b/Engine/Plugins/Runtime/StateTree/Source/StateTreeModule/Private/StateTree.cpp
--- a/Engine/Plugins/Runtime/StateTree/Source/StateTreeModule/Private/StateTree.cpp
+++ b/Engine/Plugins/Runtime/StateTree/Source/StateTreeModule/Private/StateTree.cpp
@@ -490,7 +490,7 @@ void UStateTree::Serialize(const FStructuredArchiveRecord Record)
 	// because property bindings may get invalid, and instance data potentially needs refreshed.
 	if (Record.GetUnderlyingArchive().IsModifyingWeakAndStrongReferences())
 	{
-		if (!Link())
+		if (!Link() && !HasAnyFlags(RF_ClassDefaultObject))
 		{
 			UE_LOG(LogStateTree, Log, TEXT("%s failed to link. Asset will not be usable at runtime."), *GetName());	
 		}
@@ -615,36 +615,49 @@ bool UStateTree::Link()
 		return false;
 	}
 
-	// Resolves nodes references to other StateTree data
-	FStateTreeLinker Linker(Schema);
-
-	for (int32 Index = 0; Index < Nodes.Num(); Index++)
+	if (States.Num() > 0 && Nodes.Num() > 0)
 	{
-		FStructView Node = Nodes[Index];
-		if (FStateTreeNodeBase* NodePtr = Node.GetPtr<FStateTreeNodeBase>())
+		// Check that all nodes are valid.
+		for (FConstStructView Node : Nodes)
 		{
-			const bool bLinkSucceeded = NodePtr->Link(Linker);
-			if (!bLinkSucceeded || Linker.GetStatus() == EStateTreeLinkerStatus::Failed)
+			if (!Node.IsValid())
 			{
-				UE_LOG(LogStateTree, Error, TEXT("%s: node '%s' failed to resolve its references."), *GetFullName(), *NodePtr->StaticStruct()->GetName());
+				UE_LOG(LogStateTree, Error, TEXT("%s: State Tree asset was not properly loaded (missing node). See log for loading failures, or recompile the StateTree asset."), *GetFullName());
 				return false;
 			}
 		}
 	}
 
-	ExternalDataDescs = Linker.GetExternalDataDescs();
-
-	if (States.Num() > 0 && Nodes.Num() > 0)
+	// Resolves nodes references to other StateTree data
 	{
-		// Check that all nodes are valid.
-		for (FConstStructView Node : Nodes)
+		FStateTreeLinker Linker(Schema);
+
+		for (int32 Index = 0; Index < Nodes.Num(); Index++)
 		{
-			if (!Node.IsValid())
+			FStructView Node = Nodes[Index];
+			FStateTreeNodeBase* NodePtr = Node.GetPtr<FStateTreeNodeBase>();
+			if (ensure(NodePtr))
 			{
-				UE_LOG(LogStateTree, Error, TEXT("%s: State Tree asset was not properly loaded (missing node). See log for loading failures, or recompile the StateTree asset."), *GetFullName());
+				const bool bLinkSucceeded = NodePtr->Link(Linker);
+				if (!bLinkSucceeded || Linker.GetStatus() == EStateTreeLinkerStatus::Failed)
+				{
+					UE_LOG(LogStateTree, Error, TEXT("%s: node '%s' failed to resolve its references."), *GetFullName(), *NodePtr->StaticStruct()->GetName());
+					return false;
+				}
+			}
+		}
+
+		if (Schema)
+		{
+			const bool bSchemaLinkSucceeded = Schema->Link(Linker);
+			if (!bSchemaLinkSucceeded || Linker.GetStatus() == EStateTreeLinkerStatus::Failed)
+			{
+				UE_LOG(LogStateTree, Error, TEXT("%s: schema failed to resolve its references."), *GetFullName());
 				return false;
 			}
 		}
+
+		ExternalDataDescs = Linker.GetExternalDataDescs();
 	}
 
 	UpdateRuntimeFlags();
diff --git a/Engine/Plugins/Runtime/StateTree/Source/StateTreeModule/Public/StateTreeLinker.h b/Engine/Plugins/Runtime/StateTree/Source/StateTreeModule/Public/StateTreeLinker.h
--- a/Engine/Plugins/Runtime/StateTree/Source/StateTreeModule/Public/StateTreeLinker.h
+++ b/Engine/Plugins/Runtime/StateTree/Source/StateTreeModule/Public/StateTreeLinker.h
@@ -89,11 +89,13 @@ struct FStateTreeLinker
 	}
 
 	/** @return linked external data descriptors. */
-	TConstArrayView<FStateTreeExternalDataDesc> GetExternalDataDescs() const { return ExternalDataDescs; }
+	const TArrayView<const FStateTreeExternalDataDesc> GetExternalDataDescs() const
+	{
+		return ExternalDataDescs;
+	}
 
 protected:
-
-	const UStateTreeSchema* Schema = nullptr;
+	TStrongObjectPtr<const UStateTreeSchema> Schema;
 	EStateTreeLinkerStatus Status = EStateTreeLinkerStatus::Succeeded;
 	TArray<FStateTreeExternalDataDesc> ExternalDataDescs;
 };
diff --git a/Engine/Plugins/Runtime/StateTree/Source/StateTreeModule/Public/StateTreeSchema.h b/Engine/Plugins/Runtime/StateTree/Source/StateTreeModule/Public/StateTreeSchema.h
--- a/Engine/Plugins/Runtime/StateTree/Source/StateTreeModule/Public/StateTreeSchema.h
+++ b/Engine/Plugins/Runtime/StateTree/Source/StateTreeModule/Public/StateTreeSchema.h
@@ -6,6 +6,7 @@
 #include "StateTreeSchema.generated.h"
 
 struct FStateTreeExternalDataDesc;
+struct FStateTreeLinker;
 
 /**
  * Schema describing which inputs, evaluators, and tasks a StateTree can contain.
@@ -60,6 +61,12 @@ public:
 		return {};
 	}
 
+	/** Resolves schema references to other StateTree data. */
+	virtual bool Link(FStateTreeLinker& Linker)
+	{
+		return true;
+	}
+
 #if WITH_EDITOR
 	/** @return True if enter conditions are allowed. */
 	virtual bool AllowEnterConditions() const
