Fix variable not initialized
#rnx

diff --git a/Engine/Plugins/AI/MassAI/Source/MassAIBehavior/Public/MassStateTreeDependency.h b/Engine/Plugins/AI/MassAI/Source/MassAIBehavior/Public/MassStateTreeDependency.h
--- a/Engine/Plugins/AI/MassAI/Source/MassAIBehavior/Public/MassStateTreeDependency.h
+++ b/Engine/Plugins/AI/MassAI/Source/MassAIBehavior/Public/MassStateTreeDependency.h
@@ -23,7 +23,7 @@ struct FMassStateTreeDependency
 	TObjectPtr<const UStruct> Type;
 
 	UPROPERTY()
-	EMassFragmentAccess Access;
+	EMassFragmentAccess Access = EMassFragmentAccess::None;
 };
 
 namespace UE::MassBehavior
