[MassAIBehavior] changed LogMassBehavior default verbosity to Warning

diff --git a/Engine/Plugins/AI/MassAI/Source/MassAIBehavior/Public/MassAIBehaviorTypes.h b/Engine/Plugins/AI/MassAI/Source/MassAIBehavior/Public/MassAIBehaviorTypes.h
--- a/Engine/Plugins/AI/MassAI/Source/MassAIBehavior/Public/MassAIBehaviorTypes.h
+++ b/Engine/Plugins/AI/MassAI/Source/MassAIBehavior/Public/MassAIBehaviorTypes.h
@@ -6,7 +6,7 @@
 #include "MassCommonTypes.h"
 #include "VisualLogger/VisualLogger.h"
 
-MASSAIBEHAVIOR_API DECLARE_LOG_CATEGORY_EXTERN(LogMassBehavior, Log, All);
+MASSAIBEHAVIOR_API DECLARE_LOG_CATEGORY_EXTERN(LogMassBehavior, Warning, All);
 
 namespace UE::Mass::ProcessorGroupNames
 {
