[MassGameplay] Fix to UE::MassBehavior::FStateTreeDependencyBuilder's API to properly handle TStateTreeExternalDataHandle of non-uobject types

#rnx

diff --git a/Engine/Plugins/AI/MassAI/Source/MassAIBehavior/Public/MassStateTreeDependency.h b/Engine/Plugins/AI/MassAI/Source/MassAIBehavior/Public/MassStateTreeDependency.h
--- a/Engine/Plugins/AI/MassAI/Source/MassAIBehavior/Public/MassStateTreeDependency.h
+++ b/Engine/Plugins/AI/MassAI/Source/MassAIBehavior/Public/MassStateTreeDependency.h
@@ -5,6 +5,7 @@
 #include "MassEntityConcepts.h"
 #include "MassRequirements.h"
 #include "Misc/NotNull.h"
+#include "StateTreeExecutionTypes.h"
 #include "MassStateTreeDependency.generated.h"
 
 /** A dependency the state tree has on other fragment or system used by mass. */
@@ -70,15 +71,15 @@ public:
 
 	/** Add a read only dependency from a handle. */
 	template<typename T>
-	FStateTreeDependencyBuilder& AddReadOnly(const T& Handle) requires TIsDerivedFrom<typename T::DataType, UObject>::IsDerived
+	FStateTreeDependencyBuilder& AddReadOnly(const T& /*Handle*/) requires TIsDerivedFrom<typename T::DataType, UObject>::IsDerived
 	{
 		Add(T::DataType::StaticClass(), EAccessType::ReadOnly);
 		return *this;
 	}
 	
 	/** Add a read only dependency from a handle. */
-	template<UE::Mass::CNonTag T>
-	FStateTreeDependencyBuilder& AddReadOnly(const T& Handle)
+	template<typename T>
+	FStateTreeDependencyBuilder& AddReadOnly(const T& /*Handle*/) requires UE::Mass::CNonTag<typename T::DataType>
 	{
 		Add(T::DataType::StaticStruct(), EAccessType::ReadOnly);
 		return *this;
@@ -102,15 +103,15 @@ public:
 
 	/** Add a read write dependency from a handle. */
 	template<typename T>
-	FStateTreeDependencyBuilder& AddReadWrite(const T& Handle) requires TIsDerivedFrom<typename T::DataType, UObject>::IsDerived
+	FStateTreeDependencyBuilder& AddReadWrite(const T& /*Handle*/) requires TIsDerivedFrom<typename T::DataType, UObject>::IsDerived
 	{
 		Add(T::DataType::StaticClass(), EAccessType::ReadWrite);
 		return *this;
 	}
 
 	/** Add a read write dependency from a handle. */
-	template<UE::Mass::CNonTag T>
-	FStateTreeDependencyBuilder& AddReadWrite(const T& Handle)
+	template<typename T>
+	FStateTreeDependencyBuilder& AddReadWrite(const T& /*Handle*/) requires UE::Mass::CNonTag<typename T::DataType>
 	{
 		Add(T::DataType::StaticStruct(), EAccessType::ReadWrite);
 		return *this;
