Reverted header inclusion added to MassStateTreeDependency.h

#rnx

diff --git a/Engine/Plugins/AI/MassAI/Source/MassAIBehavior/Public/MassStateTreeDependency.h b/Engine/Plugins/AI/MassAI/Source/MassAIBehavior/Public/MassStateTreeDependency.h
--- a/Engine/Plugins/AI/MassAI/Source/MassAIBehavior/Public/MassStateTreeDependency.h
+++ b/Engine/Plugins/AI/MassAI/Source/MassAIBehavior/Public/MassStateTreeDependency.h
@@ -5,7 +5,6 @@
 #include "MassEntityConcepts.h"
 #include "MassRequirements.h"
 #include "Misc/NotNull.h"
-#include "StateTreeExecutionTypes.h"
 #include "MassStateTreeDependency.generated.h"
 
 /** A dependency the state tree has on other fragment or system used by mass. */
