[MassGameplay] Added GetDependencies for a bunch of MassST tasks

diff --git a/Engine/Plugins/AI/MassAI/Source/MassAIBehavior/Private/Evaluators/MassComponentHitEvaluator.cpp b/Engine/Plugins/AI/MassAI/Source/MassAIBehavior/Private/Evaluators/MassComponentHitEvaluator.cpp
--- a/Engine/Plugins/AI/MassAI/Source/MassAIBehavior/Private/Evaluators/MassComponentHitEvaluator.cpp
+++ b/Engine/Plugins/AI/MassAI/Source/MassAIBehavior/Private/Evaluators/MassComponentHitEvaluator.cpp
@@ -3,11 +3,11 @@
 #include "Evaluators/MassComponentHitEvaluator.h"
 #include "MassAIBehaviorTypes.h"
 #include "MassComponentHitSubsystem.h"
+#include "MassStateTreeDependency.h"
 #include "MassStateTreeExecutionContext.h"
 #include "Engine/World.h"
 #include "StateTreeLinker.h"
 
-
 bool FMassComponentHitEvaluator::Link(FStateTreeLinker& Linker)
 {
 	Linker.LinkExternalData(ComponentHitSubsystemHandle);
@@ -15,10 +15,15 @@ bool FMassComponentHitEvaluator::Link(FStateTreeLinker& Linker)
 	return true;
 }
 
+void FMassComponentHitEvaluator::GetDependencies(UE::MassBehavior::FStateTreeDependencyBuilder& Builder) const
+{
+	Builder.AddReadOnly(ComponentHitSubsystemHandle);
+}
+
 void FMassComponentHitEvaluator::Tick(FStateTreeExecutionContext &Context, const float DeltaTime) const
 {
 	// Look for recent hits
-	UMassComponentHitSubsystem& HitSubsystem = Context.GetExternalData(ComponentHitSubsystemHandle);
+	const UMassComponentHitSubsystem& HitSubsystem = Context.GetExternalData(ComponentHitSubsystemHandle);
 	const FMassHitResult* HitResult = HitSubsystem.GetLastHit(static_cast<FMassStateTreeExecutionContext&>(Context).GetEntity());
 
 	FInstanceDataType& InstanceData = Context.GetInstanceData(*this);
diff --git a/Engine/Plugins/AI/MassAI/Source/MassAIBehavior/Private/Evaluators/MassZoneGraphAnnotationEvaluator.cpp b/Engine/Plugins/AI/MassAI/Source/MassAIBehavior/Private/Evaluators/MassZoneGraphAnnotationEvaluator.cpp
--- a/Engine/Plugins/AI/MassAI/Source/MassAIBehavior/Private/Evaluators/MassZoneGraphAnnotationEvaluator.cpp
+++ b/Engine/Plugins/AI/MassAI/Source/MassAIBehavior/Private/Evaluators/MassZoneGraphAnnotationEvaluator.cpp
@@ -4,6 +4,7 @@
 #include "MassStateTreeExecutionContext.h"
 #include "MassZoneGraphAnnotationFragments.h"
 #include "StateTreeLinker.h"
+#include "MassStateTreeDependency.h"
 
 
 FMassZoneGraphAnnotationEvaluator::FMassZoneGraphAnnotationEvaluator()
@@ -17,6 +18,11 @@ bool FMassZoneGraphAnnotationEvaluator::Link(FStateTreeLinker& Linker)
 	return true;
 }
 
+void FMassZoneGraphAnnotationEvaluator::GetDependencies(UE::MassBehavior::FStateTreeDependencyBuilder& Builder) const
+{
+	Builder.AddReadOnly<FMassZoneGraphAnnotationFragment>();
+}
+
 void FMassZoneGraphAnnotationEvaluator::Tick(FStateTreeExecutionContext& Context, const float DeltaTime) const
 {
 	const FMassZoneGraphAnnotationFragment& AnnotationTagsFragment = Context.GetExternalData(AnnotationTagsFragmentHandle);
diff --git a/Engine/Plugins/AI/MassAI/Source/MassAIBehavior/Private/Tasks/MassClaimSmartObjectTask.cpp b/Engine/Plugins/AI/MassAI/Source/MassAIBehavior/Private/Tasks/MassClaimSmartObjectTask.cpp
--- a/Engine/Plugins/AI/MassAI/Source/MassAIBehavior/Private/Tasks/MassClaimSmartObjectTask.cpp
+++ b/Engine/Plugins/AI/MassAI/Source/MassAIBehavior/Private/Tasks/MassClaimSmartObjectTask.cpp
@@ -11,6 +11,7 @@
 #include "MassStateTreeExecutionContext.h"
 #include "SmartObjectSubsystem.h"
 #include "StateTreeLinker.h"
+#include "MassStateTreeDependency.h"
 
 //----------------------------------------------------------------------//
 // FMassClaimSmartObjectTask
@@ -31,6 +32,16 @@ bool FMassClaimSmartObjectTask::Link(FStateTreeLinker& Linker)
 	return true;
 }
 
+void FMassClaimSmartObjectTask::GetDependencies(UE::MassBehavior::FStateTreeDependencyBuilder& Builder) const
+{
+	Builder.AddReadWrite(SmartObjectUserHandle);
+	// @todo SmartObjectSubsystemHandle is being used in a RW fashion, but we need this task to be able to run in parallel with
+	// everything else, so we need to ensure TMassExternalSubsystemTraits<USmartObjectSubsystem> is marked up for parallel access
+	// and that this information is properly utilized during Mass processing graph creation
+	Builder.AddReadOnly(SmartObjectSubsystemHandle);
+	Builder.AddReadOnly(MassSignalSubsystemHandle);
+}
+
 EStateTreeRunStatus FMassClaimSmartObjectTask::EnterState(FStateTreeExecutionContext& Context, const FStateTreeTransitionResult& Transition) const
 {
 	// Retrieve fragments and subsystems
diff --git a/Engine/Plugins/AI/MassAI/Source/MassAIBehavior/Private/Tasks/MassFindSmartObjectTask.cpp b/Engine/Plugins/AI/MassAI/Source/MassAIBehavior/Private/Tasks/MassFindSmartObjectTask.cpp
--- a/Engine/Plugins/AI/MassAI/Source/MassAIBehavior/Private/Tasks/MassFindSmartObjectTask.cpp
+++ b/Engine/Plugins/AI/MassAI/Source/MassAIBehavior/Private/Tasks/MassFindSmartObjectTask.cpp
@@ -12,6 +12,7 @@
 #include "SmartObjectZoneAnnotations.h"
 #include "StateTreeLinker.h"
 #include "ZoneGraphAnnotationSubsystem.h"
+#include "MassStateTreeDependency.h"
 
 FMassFindSmartObjectTask::FMassFindSmartObjectTask()
 {
@@ -32,6 +33,18 @@ bool FMassFindSmartObjectTask::Link(FStateTreeLinker& Linker)
 	return true;
 }
 
+void FMassFindSmartObjectTask::GetDependencies(UE::MassBehavior::FStateTreeDependencyBuilder& Builder) const
+{
+	// note that we're marking this dependency as ReadOnly, even though the subsystem itself is being fetched in
+	// the RW mode. No RW operations are performed, however. This comment is here to serve as reminder, that
+	// if the "no rw operations being performed" state changes the call below needs to change to `Builder.AddReadWrite`
+	Builder.AddReadOnly(SmartObjectSubsystemHandle);
+	Builder.AddReadWrite(MassSignalSubsystemHandle);
+	Builder.AddReadOnly(EntityTransformHandle);
+	Builder.AddReadWrite(SmartObjectUserHandle);
+	Builder.AddReadOnly(LocationHandle);
+}
+
 void FMassFindSmartObjectTask::ExitState(FStateTreeExecutionContext& Context, const FStateTreeTransitionResult& Transition) const
 {
 	FInstanceDataType& InstanceData = Context.GetInstanceData(*this);
@@ -158,8 +171,7 @@ EStateTreeRunStatus FMassFindSmartObjectTask::Tick(FStateTreeExecutionContext& C
 			const float DelayInSeconds = SearchInterval * FMath::FRandRange(1.0f - SearchIntervalDeviation, 1.0f + SearchIntervalDeviation);
 				
 			InstanceData.NextUpdate = World->GetTimeSeconds() + DelayInSeconds;
-			UMassSignalSubsystem& MassSignalSubsystem = Context.GetExternalData(MassSignalSubsystemHandle);
-			MassSignalSubsystem.DelaySignalEntityDeferred(MassStateTreeContext.GetMassEntityExecutionContext(), UE::Mass::Signals::SmartObjectRequestCandidates, Entity, DelayInSeconds);
+			SignalSubsystem.DelaySignalEntityDeferred(MassStateTreeContext.GetMassEntityExecutionContext(), UE::Mass::Signals::SmartObjectRequestCandidates, Entity, DelayInSeconds);
 		}
 	}
 	
diff --git a/Engine/Plugins/AI/MassAI/Source/MassAIBehavior/Private/Tasks/MassLookAtTask.cpp b/Engine/Plugins/AI/MassAI/Source/MassAIBehavior/Private/Tasks/MassLookAtTask.cpp
--- a/Engine/Plugins/AI/MassAI/Source/MassAIBehavior/Private/Tasks/MassLookAtTask.cpp
+++ b/Engine/Plugins/AI/MassAI/Source/MassAIBehavior/Private/Tasks/MassLookAtTask.cpp
@@ -6,6 +6,7 @@
 #include "MassSignalSubsystem.h"
 #include "MassStateTreeExecutionContext.h"
 #include "StateTreeLinker.h"
+#include "MassStateTreeDependency.h"
 
 bool FMassLookAtTask::Link(FStateTreeLinker& Linker)
 {
@@ -15,6 +16,13 @@ bool FMassLookAtTask::Link(FStateTreeLinker& Linker)
 	return true;
 }
 
+void FMassLookAtTask::GetDependencies(UE::MassBehavior::FStateTreeDependencyBuilder& Builder) const
+{
+	Builder.AddReadWrite<FMassLookAtFragment>();
+	// @todo this is only needed when InstanceData.Duration, but there's no way to tell at this point.
+	Builder.AddReadWrite<UMassSignalSubsystem>();
+}
+
 EStateTreeRunStatus FMassLookAtTask::EnterState(FStateTreeExecutionContext& Context, const FStateTreeTransitionResult& Transition) const
 {
 	FInstanceDataType& InstanceData = Context.GetInstanceData(*this);
diff --git a/Engine/Plugins/AI/MassAI/Source/MassAIBehavior/Private/Tasks/MassUseSmartObjectTask.cpp b/Engine/Plugins/AI/MassAI/Source/MassAIBehavior/Private/Tasks/MassUseSmartObjectTask.cpp
--- a/Engine/Plugins/AI/MassAI/Source/MassAIBehavior/Private/Tasks/MassUseSmartObjectTask.cpp
+++ b/Engine/Plugins/AI/MassAI/Source/MassAIBehavior/Private/Tasks/MassUseSmartObjectTask.cpp
@@ -3,7 +3,6 @@
 #include "Tasks/MassUseSmartObjectTask.h"
 #include "Engine/World.h"
 #include "MassAIBehaviorTypes.h"
-#include "MassCommonFragments.h"
 #include "MassNavigationFragments.h"
 #include "MassSignalSubsystem.h"
 #include "MassSmartObjectFragments.h"
@@ -13,6 +12,7 @@
 #include "MassZoneGraphNavigationUtils.h"
 #include "SmartObjectSubsystem.h"
 #include "StateTreeLinker.h"
+#include "MassStateTreeDependency.h"
 
 //----------------------------------------------------------------------//
 // FMassUseSmartObjectTask
@@ -28,13 +28,20 @@ bool FMassUseSmartObjectTask::Link(FStateTreeLinker& Linker)
 {
 	Linker.LinkExternalData(SmartObjectSubsystemHandle);
 	Linker.LinkExternalData(MassSignalSubsystemHandle);
-	Linker.LinkExternalData(EntityTransformHandle);
 	Linker.LinkExternalData(SmartObjectUserHandle);
 	Linker.LinkExternalData(MoveTargetHandle);
 
 	return true;
 }
 
+void FMassUseSmartObjectTask::GetDependencies(UE::MassBehavior::FStateTreeDependencyBuilder& Builder) const
+{
+	Builder.AddReadWrite(SmartObjectSubsystemHandle);
+	Builder.AddReadWrite(MassSignalSubsystemHandle);
+	Builder.AddReadWrite(SmartObjectUserHandle);
+	Builder.AddReadWrite(MoveTargetHandle);
+}
+
 EStateTreeRunStatus FMassUseSmartObjectTask::EnterState(FStateTreeExecutionContext& Context, const FStateTreeTransitionResult& Transition) const
 {
 	FMassSmartObjectUserFragment& SOUser = Context.GetExternalData(SmartObjectUserHandle);
diff --git a/Engine/Plugins/AI/MassAI/Source/MassAIBehavior/Private/Tasks/MassZoneGraphFindEscapeTarget.cpp b/Engine/Plugins/AI/MassAI/Source/MassAIBehavior/Private/Tasks/MassZoneGraphFindEscapeTarget.cpp
--- a/Engine/Plugins/AI/MassAI/Source/MassAIBehavior/Private/Tasks/MassZoneGraphFindEscapeTarget.cpp
+++ b/Engine/Plugins/AI/MassAI/Source/MassAIBehavior/Private/Tasks/MassZoneGraphFindEscapeTarget.cpp
@@ -2,7 +2,6 @@
 
 #include "Tasks/MassZoneGraphFindEscapeTarget.h"
 #include "StateTreeExecutionContext.h"
-#include "MassStateTreeSubsystem.h"
 #include "ZoneGraphSubsystem.h"
 #include "ZoneGraphAnnotationSubsystem.h"
 #include "MassZoneGraphNavigationFragments.h"
@@ -12,6 +11,7 @@
 #include "ZoneGraphSettings.h"
 #include "StateTreeLinker.h"
 #include "MassDebugger.h"
+#include "MassStateTreeDependency.h"
 
 
 FMassZoneGraphFindEscapeTarget::FMassZoneGraphFindEscapeTarget()
@@ -27,6 +27,13 @@ bool FMassZoneGraphFindEscapeTarget::Link(FStateTreeLinker& Linker)
 	return true;
 }
 
+void FMassZoneGraphFindEscapeTarget::GetDependencies(UE::MassBehavior::FStateTreeDependencyBuilder& Builder) const
+{
+	Builder.AddReadOnly(LocationHandle);
+	Builder.AddReadOnly(ZoneGraphSubsystemHandle);
+	Builder.AddReadOnly(ZoneGraphAnnotationSubsystemHandle);
+}
+
 EStateTreeRunStatus FMassZoneGraphFindEscapeTarget::EnterState(FStateTreeExecutionContext& Context, const FStateTreeTransitionResult& Transition) const
 {
 	FMassStateTreeExecutionContext& MassContext = static_cast<FMassStateTreeExecutionContext&>(Context);
diff --git a/Engine/Plugins/AI/MassAI/Source/MassAIBehavior/Private/Tasks/MassZoneGraphFindSmartObjectTarget.cpp b/Engine/Plugins/AI/MassAI/Source/MassAIBehavior/Private/Tasks/MassZoneGraphFindSmartObjectTarget.cpp
--- a/Engine/Plugins/AI/MassAI/Source/MassAIBehavior/Private/Tasks/MassZoneGraphFindSmartObjectTarget.cpp
+++ b/Engine/Plugins/AI/MassAI/Source/MassAIBehavior/Private/Tasks/MassZoneGraphFindSmartObjectTarget.cpp
@@ -4,12 +4,12 @@
 #include "ZoneGraphAnnotationSubsystem.h"
 #include "MassAIBehaviorTypes.h"
 #include "SmartObjectSubsystem.h"
-#include "MassSmartObjectFragments.h"
 #include "MassSmartObjectSettings.h"
 #include "MassStateTreeExecutionContext.h"
 #include "MassZoneGraphNavigationFragments.h"
 #include "SmartObjectZoneAnnotations.h"
 #include "StateTreeLinker.h"
+#include "MassStateTreeDependency.h"
 
 bool FMassZoneGraphFindSmartObjectTarget::Link(FStateTreeLinker& Linker)
 {
@@ -20,6 +20,13 @@ bool FMassZoneGraphFindSmartObjectTarget::Link(FStateTreeLinker& Linker)
 	return true;
 }
 
+void FMassZoneGraphFindSmartObjectTarget::GetDependencies(UE::MassBehavior::FStateTreeDependencyBuilder& Builder) const
+{
+	Builder.AddReadOnly(LocationHandle);
+	Builder.AddReadOnly(AnnotationSubsystemHandle);
+	Builder.AddReadOnly(SmartObjectSubsystemHandle);
+}
+
 EStateTreeRunStatus FMassZoneGraphFindSmartObjectTarget::EnterState(FStateTreeExecutionContext& Context, const FStateTreeTransitionResult& Transition) const
 {
 	const UZoneGraphAnnotationSubsystem& AnnotationSubsystem = Context.GetExternalData(AnnotationSubsystemHandle);
diff --git a/Engine/Plugins/AI/MassAI/Source/MassAIBehavior/Private/Tasks/MassZoneGraphPathFollowTask.cpp b/Engine/Plugins/AI/MassAI/Source/MassAIBehavior/Private/Tasks/MassZoneGraphPathFollowTask.cpp
--- a/Engine/Plugins/AI/MassAI/Source/MassAIBehavior/Private/Tasks/MassZoneGraphPathFollowTask.cpp
+++ b/Engine/Plugins/AI/MassAI/Source/MassAIBehavior/Private/Tasks/MassZoneGraphPathFollowTask.cpp
@@ -1,7 +1,6 @@
 // Copyright Epic Games, Inc. All Rights Reserved.
 
 #include "Tasks/MassZoneGraphPathFollowTask.h"
-
 #include "MassAIBehaviorTypes.h"
 #include "MassCommonFragments.h"
 #include "MassNavigationFragments.h"
@@ -12,6 +11,7 @@
 #include "ZoneGraphSubsystem.h"
 #include "StateTreeLinker.h"
 #include "MassDebugger.h"
+#include "MassStateTreeDependency.h"
 
 
 bool FMassZoneGraphPathFollowTask::Link(FStateTreeLinker& Linker)
@@ -28,6 +28,18 @@ bool FMassZoneGraphPathFollowTask::Link(FStateTreeLinker& Linker)
 	return true;
 }
 
+void FMassZoneGraphPathFollowTask::GetDependencies(UE::MassBehavior::FStateTreeDependencyBuilder& Builder) const
+{
+	Builder.AddReadOnly(LocationHandle);
+	Builder.AddReadWrite(MoveTargetHandle);
+	Builder.AddReadWrite(PathRequestHandle);
+	Builder.AddReadWrite(ShortPathHandle);
+	Builder.AddReadWrite(CachedLaneHandle);
+	Builder.AddReadOnly(AgentRadiusHandle);
+	Builder.AddReadOnly(MovementParamsHandle);
+	Builder.AddReadOnly(ZoneGraphSubsystemHandle);
+}
+
 bool FMassZoneGraphPathFollowTask::RequestPath(FMassStateTreeExecutionContext& Context, const FMassZoneGraphTargetLocation& RequestedTargetLocation) const
 {
 	const UZoneGraphSubsystem& ZoneGraphSubsystem = Context.GetExternalData(ZoneGraphSubsystemHandle);
diff --git a/Engine/Plugins/AI/MassAI/Source/MassAIBehavior/Private/Tasks/MassZoneGraphStandTask.cpp b/Engine/Plugins/AI/MassAI/Source/MassAIBehavior/Private/Tasks/MassZoneGraphStandTask.cpp
--- a/Engine/Plugins/AI/MassAI/Source/MassAIBehavior/Private/Tasks/MassZoneGraphStandTask.cpp
+++ b/Engine/Plugins/AI/MassAI/Source/MassAIBehavior/Private/Tasks/MassZoneGraphStandTask.cpp
@@ -12,6 +12,7 @@
 #include "MassSignalSubsystem.h"
 #include "MassSimulationLOD.h"
 #include "StateTreeLinker.h"
+#include "MassStateTreeDependency.h"
 
 bool FMassZoneGraphStandTask::Link(FStateTreeLinker& Linker)
 {
@@ -26,6 +27,17 @@ bool FMassZoneGraphStandTask::Link(FStateTreeLinker& Linker)
 	return true;
 }
 
+void FMassZoneGraphStandTask::GetDependencies(UE::MassBehavior::FStateTreeDependencyBuilder& Builder) const
+{
+	Builder.AddReadOnly(LocationHandle);
+	Builder.AddReadWrite(MoveTargetHandle);
+	Builder.AddReadWrite(ShortPathHandle);
+	Builder.AddReadWrite(CachedLaneHandle);
+	Builder.AddReadOnly(MovementParamsHandle);
+	Builder.AddReadOnly(ZoneGraphSubsystemHandle);
+	Builder.AddReadWrite(MassSignalSubsystemHandle);
+}
+
 EStateTreeRunStatus FMassZoneGraphStandTask::EnterState(FStateTreeExecutionContext& Context, const FStateTreeTransitionResult& Transition) const
 {
 	const FMassStateTreeExecutionContext& MassStateTreeContext = static_cast<FMassStateTreeExecutionContext&>(Context);
diff --git a/Engine/Plugins/AI/MassAI/Source/MassAIBehavior/Public/Evaluators/MassComponentHitEvaluator.h b/Engine/Plugins/AI/MassAI/Source/MassAIBehavior/Public/Evaluators/MassComponentHitEvaluator.h
--- a/Engine/Plugins/AI/MassAI/Source/MassAIBehavior/Public/Evaluators/MassComponentHitEvaluator.h
+++ b/Engine/Plugins/AI/MassAI/Source/MassAIBehavior/Public/Evaluators/MassComponentHitEvaluator.h
@@ -7,6 +7,10 @@
 #include "MassComponentHitEvaluator.generated.h"
 
 class UMassComponentHitSubsystem;
+namespace UE::MassBehavior
+{
+	struct FStateTreeDependencyBuilder;
+}
 
 /**
  * Evaluator to extract last hit from the MassComponentHitSubsystem and expose it for tasks and transitions
@@ -35,6 +39,7 @@ protected:
 	virtual bool Link(FStateTreeLinker& Linker) override;
 	virtual const UStruct* GetInstanceDataType() const override { return FInstanceDataType::StaticStruct(); }
 	virtual void Tick(FStateTreeExecutionContext& Context, const float DeltaTime) const override;
+	virtual void GetDependencies(UE::MassBehavior::FStateTreeDependencyBuilder& Builder) const override;
 
 	TStateTreeExternalDataHandle<UMassComponentHitSubsystem> ComponentHitSubsystemHandle;
 };
diff --git a/Engine/Plugins/AI/MassAI/Source/MassAIBehavior/Public/Evaluators/MassZoneGraphAnnotationEvaluator.h b/Engine/Plugins/AI/MassAI/Source/MassAIBehavior/Public/Evaluators/MassZoneGraphAnnotationEvaluator.h
--- a/Engine/Plugins/AI/MassAI/Source/MassAIBehavior/Public/Evaluators/MassZoneGraphAnnotationEvaluator.h
+++ b/Engine/Plugins/AI/MassAI/Source/MassAIBehavior/Public/Evaluators/MassZoneGraphAnnotationEvaluator.h
@@ -8,6 +8,11 @@
 
 struct FStateTreeExecutionContext;
 struct FMassZoneGraphAnnotationFragment;
+namespace UE::MassBehavior
+{
+	struct FStateTreeDependencyBuilder;
+};
+
 /**
  * Evaluator to expose ZoneGraph Annotation Tags for decision-making.
  */
@@ -33,6 +38,7 @@ protected:
 	virtual bool Link(FStateTreeLinker& Linker) override;
 	virtual const UStruct* GetInstanceDataType() const override { return FInstanceDataType::StaticStruct(); }
 	virtual void Tick(FStateTreeExecutionContext& Context, const float DeltaTime) const override;
+	virtual void GetDependencies(UE::MassBehavior::FStateTreeDependencyBuilder& Builder) const override;
 
 	TStateTreeExternalDataHandle<FMassZoneGraphAnnotationFragment> AnnotationTagsFragmentHandle;
 };
diff --git a/Engine/Plugins/AI/MassAI/Source/MassAIBehavior/Public/Tasks/MassClaimSmartObjectTask.h b/Engine/Plugins/AI/MassAI/Source/MassAIBehavior/Public/Tasks/MassClaimSmartObjectTask.h
--- a/Engine/Plugins/AI/MassAI/Source/MassAIBehavior/Public/Tasks/MassClaimSmartObjectTask.h
+++ b/Engine/Plugins/AI/MassAI/Source/MassAIBehavior/Public/Tasks/MassClaimSmartObjectTask.h
@@ -15,6 +15,10 @@ class USmartObjectSubsystem;
 class UMassSignalSubsystem;
 struct FTransformFragment;
 struct FMassMoveTargetFragment;
+namespace UE::MassBehavior
+{
+	struct FStateTreeDependencyBuilder;
+};
 
 /**
  * Tasks to claim a smart object from search results and release it when done.
@@ -47,6 +51,7 @@ protected:
 	virtual EStateTreeRunStatus EnterState(FStateTreeExecutionContext& Context, const FStateTreeTransitionResult& Transition) const override;
 	virtual void ExitState(FStateTreeExecutionContext& Context, const FStateTreeTransitionResult& Transition) const override;
 	virtual EStateTreeRunStatus Tick(FStateTreeExecutionContext& Context, const float DeltaTime) const override;
+	virtual void GetDependencies(UE::MassBehavior::FStateTreeDependencyBuilder& Builder) const override;
 
 	TStateTreeExternalDataHandle<FMassSmartObjectUserFragment> SmartObjectUserHandle;
 	TStateTreeExternalDataHandle<USmartObjectSubsystem> SmartObjectSubsystemHandle;
diff --git a/Engine/Plugins/AI/MassAI/Source/MassAIBehavior/Public/Tasks/MassFindSmartObjectTask.h b/Engine/Plugins/AI/MassAI/Source/MassAIBehavior/Public/Tasks/MassFindSmartObjectTask.h
--- a/Engine/Plugins/AI/MassAI/Source/MassAIBehavior/Public/Tasks/MassFindSmartObjectTask.h
+++ b/Engine/Plugins/AI/MassAI/Source/MassAIBehavior/Public/Tasks/MassFindSmartObjectTask.h
@@ -14,6 +14,10 @@ class USmartObjectSubsystem;
 struct FMassSmartObjectUserFragment;
 struct FMassZoneGraphLaneLocationFragment;
 struct FTransformFragment;
+namespace UE::MassBehavior
+{
+	struct FStateTreeDependencyBuilder;
+};
 
 USTRUCT()
 struct MASSAIBEHAVIOR_API FMassFindSmartObjectTaskInstanceData
@@ -55,6 +59,7 @@ protected:
 	virtual void ExitState(FStateTreeExecutionContext& Context, const FStateTreeTransitionResult& Transition) const override;
 	virtual void StateCompleted(FStateTreeExecutionContext& Context, const EStateTreeRunStatus CompletionStatus, const FStateTreeActiveStates& CompletedActiveStates) const override;
 	virtual EStateTreeRunStatus Tick(FStateTreeExecutionContext& Context, const float DeltaTime) const override;
+	virtual void GetDependencies(UE::MassBehavior::FStateTreeDependencyBuilder& Builder) const override;
 
 	TStateTreeExternalDataHandle<USmartObjectSubsystem> SmartObjectSubsystemHandle;
 	TStateTreeExternalDataHandle<UMassSignalSubsystem> MassSignalSubsystemHandle;
diff --git a/Engine/Plugins/AI/MassAI/Source/MassAIBehavior/Public/Tasks/MassLookAtTask.h b/Engine/Plugins/AI/MassAI/Source/MassAIBehavior/Public/Tasks/MassLookAtTask.h
--- a/Engine/Plugins/AI/MassAI/Source/MassAIBehavior/Public/Tasks/MassLookAtTask.h
+++ b/Engine/Plugins/AI/MassAI/Source/MassAIBehavior/Public/Tasks/MassLookAtTask.h
@@ -8,6 +8,10 @@
 #include "MassLookAtTask.generated.h"
 
 class UMassSignalSubsystem;
+namespace UE::MassBehavior
+{
+	struct FStateTreeDependencyBuilder;
+};
 
 /**
  * Task to assign a LookAt target for mass processing
@@ -43,6 +47,7 @@ protected:
 	virtual EStateTreeRunStatus EnterState(FStateTreeExecutionContext& Context, const FStateTreeTransitionResult& Transition) const override;
 	virtual void ExitState(FStateTreeExecutionContext& Context, const FStateTreeTransitionResult& Transition) const override;
 	virtual EStateTreeRunStatus Tick(FStateTreeExecutionContext& Context, const float DeltaTime) const override;
+	virtual void GetDependencies(UE::MassBehavior::FStateTreeDependencyBuilder& Builder) const override;
 
 	TStateTreeExternalDataHandle<UMassSignalSubsystem> MassSignalSubsystemHandle;
 	TStateTreeExternalDataHandle<FMassLookAtFragment> LookAtHandle;
diff --git a/Engine/Plugins/AI/MassAI/Source/MassAIBehavior/Public/Tasks/MassUseSmartObjectTask.h b/Engine/Plugins/AI/MassAI/Source/MassAIBehavior/Public/Tasks/MassUseSmartObjectTask.h
--- a/Engine/Plugins/AI/MassAI/Source/MassAIBehavior/Public/Tasks/MassUseSmartObjectTask.h
+++ b/Engine/Plugins/AI/MassAI/Source/MassAIBehavior/Public/Tasks/MassUseSmartObjectTask.h
@@ -15,6 +15,10 @@ class USmartObjectSubsystem;
 class UMassSignalSubsystem;
 struct FTransformFragment;
 struct FMassMoveTargetFragment;
+namespace UE::MassBehavior
+{
+	struct FStateTreeDependencyBuilder;
+};
 
 /**
  * Task to tell an entity to start using a claimed smart object.
@@ -44,10 +48,10 @@ protected:
 	virtual void ExitState(FStateTreeExecutionContext& Context, const FStateTreeTransitionResult& Transition) const override;
 	virtual void StateCompleted(FStateTreeExecutionContext& Context, const EStateTreeRunStatus CompletionStatus, const FStateTreeActiveStates& CompletedActiveStates) const override;
 	virtual EStateTreeRunStatus Tick(FStateTreeExecutionContext& Context, const float DeltaTime) const override;
+	virtual void GetDependencies(UE::MassBehavior::FStateTreeDependencyBuilder& Builder) const override;
 
 	TStateTreeExternalDataHandle<USmartObjectSubsystem> SmartObjectSubsystemHandle;
 	TStateTreeExternalDataHandle<UMassSignalSubsystem> MassSignalSubsystemHandle;
-	TStateTreeExternalDataHandle<FTransformFragment> EntityTransformHandle;
 	TStateTreeExternalDataHandle<FMassSmartObjectUserFragment> SmartObjectUserHandle;
 	TStateTreeExternalDataHandle<FMassMoveTargetFragment> MoveTargetHandle;
 };
diff --git a/Engine/Plugins/AI/MassAI/Source/MassAIBehavior/Public/Tasks/MassZoneGraphFindEscapeTarget.h b/Engine/Plugins/AI/MassAI/Source/MassAIBehavior/Public/Tasks/MassZoneGraphFindEscapeTarget.h
--- a/Engine/Plugins/AI/MassAI/Source/MassAIBehavior/Public/Tasks/MassZoneGraphFindEscapeTarget.h
+++ b/Engine/Plugins/AI/MassAI/Source/MassAIBehavior/Public/Tasks/MassZoneGraphFindEscapeTarget.h
@@ -10,6 +10,10 @@ struct FStateTreeExecutionContext;
 struct FMassZoneGraphLaneLocationFragment;
 class UZoneGraphSubsystem;
 class UZoneGraphAnnotationSubsystem;
+namespace UE::MassBehavior
+{
+	struct FStateTreeDependencyBuilder;
+};
 
 /**
  * Updates TargetLocation to a escape target based on the agents current location on ZoneGraph, and disturbance annotation.
@@ -36,6 +40,7 @@ protected:
 	virtual bool Link(FStateTreeLinker& Linker) override;
 	virtual const UStruct* GetInstanceDataType() const override { return FInstanceDataType::StaticStruct(); }
 	virtual EStateTreeRunStatus EnterState(FStateTreeExecutionContext& Context, const FStateTreeTransitionResult& Transition) const override;
+	virtual void GetDependencies(UE::MassBehavior::FStateTreeDependencyBuilder& Builder) const override;
 
 	TStateTreeExternalDataHandle<FMassZoneGraphLaneLocationFragment> LocationHandle;
 	TStateTreeExternalDataHandle<UZoneGraphSubsystem> ZoneGraphSubsystemHandle;
diff --git a/Engine/Plugins/AI/MassAI/Source/MassAIBehavior/Public/Tasks/MassZoneGraphFindSmartObjectTarget.h b/Engine/Plugins/AI/MassAI/Source/MassAIBehavior/Public/Tasks/MassZoneGraphFindSmartObjectTarget.h
--- a/Engine/Plugins/AI/MassAI/Source/MassAIBehavior/Public/Tasks/MassZoneGraphFindSmartObjectTarget.h
+++ b/Engine/Plugins/AI/MassAI/Source/MassAIBehavior/Public/Tasks/MassZoneGraphFindSmartObjectTarget.h
@@ -11,6 +11,10 @@ struct FMassSmartObjectUserFragment;
 struct FMassZoneGraphLaneLocationFragment;
 class UZoneGraphAnnotationSubsystem;
 class USmartObjectSubsystem;
+namespace UE::MassBehavior
+{
+	struct FStateTreeDependencyBuilder;
+};
 
 /**
 * Computes move target to a smart object based on current location on ZoneGraph.
@@ -38,6 +42,7 @@ protected:
 	virtual bool Link(FStateTreeLinker& Linker) override;
 	virtual const UStruct* GetInstanceDataType() const override { return FInstanceDataType::StaticStruct(); }
 	virtual EStateTreeRunStatus EnterState(FStateTreeExecutionContext& Context, const FStateTreeTransitionResult& Transition) const override;
+	virtual void GetDependencies(UE::MassBehavior::FStateTreeDependencyBuilder& Builder) const override;
 
 	TStateTreeExternalDataHandle<FMassZoneGraphLaneLocationFragment> LocationHandle;
 	TStateTreeExternalDataHandle<UZoneGraphAnnotationSubsystem> AnnotationSubsystemHandle;
diff --git a/Engine/Plugins/AI/MassAI/Source/MassAIBehavior/Public/Tasks/MassZoneGraphPathFollowTask.h b/Engine/Plugins/AI/MassAI/Source/MassAIBehavior/Public/Tasks/MassZoneGraphPathFollowTask.h
--- a/Engine/Plugins/AI/MassAI/Source/MassAIBehavior/Public/Tasks/MassZoneGraphPathFollowTask.h
+++ b/Engine/Plugins/AI/MassAI/Source/MassAIBehavior/Public/Tasks/MassZoneGraphPathFollowTask.h
@@ -18,6 +18,10 @@ struct FMassZoneGraphCachedLaneFragment;
 struct FAgentRadiusFragment;
 struct FMassMovementParameters;
 class UZoneGraphSubsystem;
+namespace UE::MassBehavior
+{
+	struct FStateTreeDependencyBuilder;
+};
 
 USTRUCT()
 struct MASSAIBEHAVIOR_API FMassZoneGraphTargetLocation
@@ -94,6 +98,7 @@ protected:
 	virtual const UStruct* GetInstanceDataType() const override { return FInstanceDataType::StaticStruct(); };
 	virtual EStateTreeRunStatus EnterState(FStateTreeExecutionContext& Context, const FStateTreeTransitionResult& Transition) const override;
 	virtual EStateTreeRunStatus Tick(FStateTreeExecutionContext& Context, const float DeltaTime) const override;
+	virtual void GetDependencies(UE::MassBehavior::FStateTreeDependencyBuilder& Builder) const override;
 
 	bool RequestPath(FMassStateTreeExecutionContext& Context, const FMassZoneGraphTargetLocation& TargetLocation) const;
 
diff --git a/Engine/Plugins/AI/MassAI/Source/MassAIBehavior/Public/Tasks/MassZoneGraphStandTask.h b/Engine/Plugins/AI/MassAI/Source/MassAIBehavior/Public/Tasks/MassZoneGraphStandTask.h
--- a/Engine/Plugins/AI/MassAI/Source/MassAIBehavior/Public/Tasks/MassZoneGraphStandTask.h
+++ b/Engine/Plugins/AI/MassAI/Source/MassAIBehavior/Public/Tasks/MassZoneGraphStandTask.h
@@ -13,6 +13,10 @@ struct FMassZoneGraphCachedLaneFragment;
 struct FMassMovementParameters;
 class UZoneGraphSubsystem;
 class UMassSignalSubsystem;
+namespace UE::MassBehavior
+{
+	struct FStateTreeDependencyBuilder;
+};
 
 /**
  * Stop, and stand on current ZoneGraph location
@@ -43,6 +47,7 @@ protected:
 	virtual const UStruct* GetInstanceDataType() const override { return FInstanceDataType::StaticStruct(); }
 	virtual EStateTreeRunStatus EnterState(FStateTreeExecutionContext& Context, const FStateTreeTransitionResult& Transition) const override;
 	virtual EStateTreeRunStatus Tick(FStateTreeExecutionContext& Context, const float DeltaTime) const override;
+	virtual void GetDependencies(UE::MassBehavior::FStateTreeDependencyBuilder& Builder) const override;
 
 	TStateTreeExternalDataHandle<FMassZoneGraphLaneLocationFragment> LocationHandle;
 	TStateTreeExternalDataHandle<FMassMoveTargetFragment> MoveTargetHandle;
diff --git a/Engine/Plugins/AI/MassCrowd/Source/MassCrowd/Private/Tasks/MassCrowdClaimWaitSlotTask.cpp b/Engine/Plugins/AI/MassCrowd/Source/MassCrowd/Private/Tasks/MassCrowdClaimWaitSlotTask.cpp
--- a/Engine/Plugins/AI/MassCrowd/Source/MassCrowd/Private/Tasks/MassCrowdClaimWaitSlotTask.cpp
+++ b/Engine/Plugins/AI/MassCrowd/Source/MassCrowd/Private/Tasks/MassCrowdClaimWaitSlotTask.cpp
@@ -6,6 +6,7 @@
 #include "MassNavigationFragments.h"
 #include "MassStateTreeExecutionContext.h"
 #include "StateTreeLinker.h"
+#include "MassStateTreeDependency.h"
 
 FMassCrowdClaimWaitSlotTask::FMassCrowdClaimWaitSlotTask()
 {
@@ -22,6 +23,13 @@ bool FMassCrowdClaimWaitSlotTask::Link(FStateTreeLinker& Linker)
 	return true;
 }
 
+void FMassCrowdClaimWaitSlotTask::GetDependencies(UE::MassBehavior::FStateTreeDependencyBuilder& Builder) const
+{
+	Builder.AddReadOnly<FMassZoneGraphLaneLocationFragment>();
+	Builder.AddReadOnly<FMassMoveTargetFragment>();
+	Builder.AddReadWrite(CrowdSubsystemHandle);
+}
+
 EStateTreeRunStatus FMassCrowdClaimWaitSlotTask::EnterState(FStateTreeExecutionContext& Context, const FStateTreeTransitionResult& Transition) const
 {
 	FInstanceDataType& InstanceData = Context.GetInstanceData(*this);
diff --git a/Engine/Plugins/AI/MassCrowd/Source/MassCrowd/Private/Tasks/MassZoneGraphFindWanderTarget.cpp b/Engine/Plugins/AI/MassCrowd/Source/MassCrowd/Private/Tasks/MassZoneGraphFindWanderTarget.cpp
--- a/Engine/Plugins/AI/MassCrowd/Source/MassCrowd/Private/Tasks/MassZoneGraphFindWanderTarget.cpp
+++ b/Engine/Plugins/AI/MassCrowd/Source/MassCrowd/Private/Tasks/MassZoneGraphFindWanderTarget.cpp
@@ -11,6 +11,7 @@
 #include "MassStateTreeExecutionContext.h"
 #include "StateTreeLinker.h"
 #include "MassDebugger.h"
+#include "MassStateTreeDependency.h"
 
 
 FMassZoneGraphFindWanderTarget::FMassZoneGraphFindWanderTarget()
@@ -27,13 +28,21 @@ bool FMassZoneGraphFindWanderTarget::Link(FStateTreeLinker& Linker)
 	return true;
 }
 
+void FMassZoneGraphFindWanderTarget::GetDependencies(UE::MassBehavior::FStateTreeDependencyBuilder& Builder) const
+{
+	Builder.AddReadOnly(LocationHandle);
+	Builder.AddReadOnly(ZoneGraphSubsystemHandle);
+	Builder.AddReadOnly(ZoneGraphAnnotationSubsystemHandle);
+	Builder.AddReadOnly(MassCrowdSubsystemHandle);
+}
+
 EStateTreeRunStatus FMassZoneGraphFindWanderTarget::EnterState(FStateTreeExecutionContext& Context, const FStateTreeTransitionResult& Transition) const
 {
 	const FMassStateTreeExecutionContext& MassContext = static_cast<FMassStateTreeExecutionContext&>(Context);
 
 	const FMassZoneGraphLaneLocationFragment& LaneLocation = Context.GetExternalData(LocationHandle);
 	const UZoneGraphSubsystem& ZoneGraphSubsystem = Context.GetExternalData(ZoneGraphSubsystemHandle);
-	UZoneGraphAnnotationSubsystem& ZoneGraphAnnotationSubsystem = Context.GetExternalData(ZoneGraphAnnotationSubsystemHandle);
+	const UZoneGraphAnnotationSubsystem& ZoneGraphAnnotationSubsystem = Context.GetExternalData(ZoneGraphAnnotationSubsystemHandle);
 	const UMassCrowdSubsystem& MassCrowdSubsystem = Context.GetExternalData(MassCrowdSubsystemHandle);
 
 	FInstanceDataType& InstanceData = Context.GetInstanceData(*this);
diff --git a/Engine/Plugins/AI/MassCrowd/Source/MassCrowd/Public/Tasks/MassCrowdClaimWaitSlotTask.h b/Engine/Plugins/AI/MassCrowd/Source/MassCrowd/Public/Tasks/MassCrowdClaimWaitSlotTask.h
--- a/Engine/Plugins/AI/MassCrowd/Source/MassCrowd/Public/Tasks/MassCrowdClaimWaitSlotTask.h
+++ b/Engine/Plugins/AI/MassCrowd/Source/MassCrowd/Public/Tasks/MassCrowdClaimWaitSlotTask.h
@@ -9,6 +9,10 @@
 
 struct FStateTreeExecutionContext;
 class UMassCrowdSubsystem;
+namespace UE::MassBehavior
+{
+	struct FStateTreeDependencyBuilder;
+};
 
 /**
 * Claim wait slot and expose slot position for path follow.
@@ -42,6 +46,7 @@ protected:
 	virtual const UStruct* GetInstanceDataType() const override { return FInstanceDataType::StaticStruct(); }
 	virtual EStateTreeRunStatus EnterState(FStateTreeExecutionContext& Context, const FStateTreeTransitionResult& Transition) const override;
 	virtual void ExitState(FStateTreeExecutionContext& Context, const FStateTreeTransitionResult& Transition) const override;
+	virtual void GetDependencies(UE::MassBehavior::FStateTreeDependencyBuilder& Builder) const override;
 
 	TStateTreeExternalDataHandle<FMassZoneGraphLaneLocationFragment> LocationHandle;
 	TStateTreeExternalDataHandle<FMassMoveTargetFragment> MoveTargetHandle;
diff --git a/Engine/Plugins/AI/MassCrowd/Source/MassCrowd/Public/Tasks/MassZoneGraphFindWanderTarget.h b/Engine/Plugins/AI/MassCrowd/Source/MassCrowd/Public/Tasks/MassZoneGraphFindWanderTarget.h
--- a/Engine/Plugins/AI/MassCrowd/Source/MassCrowd/Public/Tasks/MassZoneGraphFindWanderTarget.h
+++ b/Engine/Plugins/AI/MassCrowd/Source/MassCrowd/Public/Tasks/MassZoneGraphFindWanderTarget.h
@@ -11,6 +11,10 @@ struct FMassZoneGraphLaneLocationFragment;
 class UZoneGraphSubsystem;
 class UZoneGraphAnnotationSubsystem;
 class UMassCrowdSubsystem;
+namespace UE::MassBehavior
+{
+	struct FStateTreeDependencyBuilder;
+};
 
 /**
  * Updates TargetLocation to a wander target based on the agents current location on ZoneGraph.
@@ -37,6 +41,7 @@ protected:
 	virtual bool Link(FStateTreeLinker& Linker) override;
 	virtual const UStruct* GetInstanceDataType() const override { return FInstanceDataType::StaticStruct(); }
 	virtual EStateTreeRunStatus EnterState(FStateTreeExecutionContext& Context, const FStateTreeTransitionResult& Transition) const override;
+	virtual void GetDependencies(UE::MassBehavior::FStateTreeDependencyBuilder& Builder) const override;
 
 	TStateTreeExternalDataHandle<FMassZoneGraphLaneLocationFragment> LocationHandle;
 	TStateTreeExternalDataHandle<UZoneGraphSubsystem> ZoneGraphSubsystemHandle;
