[ZoneGraph] Added prefix to parameter names

diff --git a/Engine/Plugins/AI/MassAI/Source/MassZoneGraphNavigation/Private/MassZoneGraphNavigationFragments.cpp b/Engine/Plugins/AI/MassAI/Source/MassZoneGraphNavigation/Private/MassZoneGraphNavigationFragments.cpp
--- a/Engine/Plugins/AI/MassAI/Source/MassZoneGraphNavigation/Private/MassZoneGraphNavigationFragments.cpp
+++ b/Engine/Plugins/AI/MassAI/Source/MassZoneGraphNavigation/Private/MassZoneGraphNavigationFragments.cpp
@@ -155,6 +155,7 @@ void FMassZoneGraphCachedLaneFragment::CacheLaneData(const FZoneGraphStorage& Zo
 	LaneRightSpace = FMassInt16Real(AdjacentRightWidth);
 }
 
+// Fills Points from CachedLane
 bool FMassZoneGraphShortPathFragment::RequestPath(const FMassZoneGraphCachedLaneFragment& CachedLane, const FZoneGraphShortPathRequest& Request, const float InCurrentDistanceAlongLane, const float AgentRadius)
 {
 	Reset();
diff --git a/Engine/Plugins/AI/MassAI/Source/MassZoneGraphNavigation/Private/MassZoneGraphNavigationUtils.cpp b/Engine/Plugins/AI/MassAI/Source/MassZoneGraphNavigation/Private/MassZoneGraphNavigationUtils.cpp
--- a/Engine/Plugins/AI/MassAI/Source/MassZoneGraphNavigation/Private/MassZoneGraphNavigationUtils.cpp
+++ b/Engine/Plugins/AI/MassAI/Source/MassZoneGraphNavigation/Private/MassZoneGraphNavigationUtils.cpp
@@ -20,16 +20,16 @@ namespace UE::MassNavigation
 							const FZoneGraphShortPathRequest& PathRequest,
 							const float AgentRadius,
 							const float DesiredSpeed,
-							FMassMoveTargetFragment& MoveTarget,
-							FMassZoneGraphShortPathFragment& ShortPath,
-							FMassZoneGraphCachedLaneFragment& CachedLane)
+							FMassMoveTargetFragment& InOutMoveTarget,
+							FMassZoneGraphShortPathFragment& OutShortPath,
+							FMassZoneGraphCachedLaneFragment& OutCachedLane)
 	{
-		ShortPath.Reset();
-		CachedLane.Reset();
-		MoveTarget.DistanceToGoal = 0.0f;
-		MoveTarget.DesiredSpeed.Set(0.0f);
+		OutShortPath.Reset();
+		OutCachedLane.Reset();
+		InOutMoveTarget.DistanceToGoal = 0.0f;
+		InOutMoveTarget.DesiredSpeed.Set(0.0f);
 
-		if (!ensureMsgf(MoveTarget.GetCurrentAction() == EMassMovementAction::Move, TEXT("Expecting action 'Move': Invalid action %u"), MoveTarget.GetCurrentAction()))
+		if (!ensureMsgf(InOutMoveTarget.GetCurrentAction() == EMassMovementAction::Move, TEXT("Expecting action 'Move': Invalid action %u"), InOutMoveTarget.GetCurrentAction()))
 		{
 			return false;
 		}
@@ -43,14 +43,14 @@ namespace UE::MassNavigation
 			return false;
 		}
 
-		MoveTarget.IntentAtGoal = EMassMovementAction::Stand;
-		MoveTarget.DesiredSpeed.Set(DesiredSpeed);
+		InOutMoveTarget.IntentAtGoal = EMassMovementAction::Stand;
+		InOutMoveTarget.DesiredSpeed.Set(DesiredSpeed);
 
-		CachedLane.CacheLaneData(*ZoneGraphStorage, LaneLocation.LaneHandle, LaneLocation.DistanceAlongLane, PathRequest.TargetDistance, InflateDistance);
-		if (ShortPath.RequestPath(CachedLane, PathRequest, LaneLocation.DistanceAlongLane, AgentRadius))
+		OutCachedLane.CacheLaneData(*ZoneGraphStorage, LaneLocation.LaneHandle, LaneLocation.DistanceAlongLane, PathRequest.TargetDistance, InflateDistance);
+		if (OutShortPath.RequestPath(OutCachedLane, PathRequest, LaneLocation.DistanceAlongLane, AgentRadius))
 		{
-			MoveTarget.IntentAtGoal = ShortPath.EndOfPathIntent;
-			MoveTarget.DistanceToGoal = (ShortPath.NumPoints > 0) ? ShortPath.Points[ShortPath.NumPoints - 1].DistanceAlongLane.Get() : 0.0f;
+			InOutMoveTarget.IntentAtGoal = OutShortPath.EndOfPathIntent;
+			InOutMoveTarget.DistanceToGoal = (OutShortPath.NumPoints > 0) ? OutShortPath.Points[OutShortPath.NumPoints - 1].DistanceAlongLane.Get() : 0.0f;
 #if WITH_MASSGAMEPLAY_DEBUG
 			UE_CVLOG(UE::Mass::Debug::IsDebuggingEntity(Entity),
 				Requester,
@@ -72,7 +72,7 @@ namespace UE::MassNavigation
 			return false;
 		}
 
-		UE_VLOG(Requester, LogMassNavigation, Log, TEXT("Entity [%s] successfully requested %s"), *Entity.DebugGetDescription(), *MoveTarget.ToString());
+		UE_VLOG(Requester, LogMassNavigation, Log, TEXT("Entity [%s] successfully requested %s"), *Entity.DebugGetDescription(), *InOutMoveTarget.ToString());
 		return true;
 	}
 
diff --git a/Engine/Plugins/AI/MassAI/Source/MassZoneGraphNavigation/Public/MassZoneGraphNavigationUtils.h b/Engine/Plugins/AI/MassAI/Source/MassZoneGraphNavigation/Public/MassZoneGraphNavigationUtils.h
--- a/Engine/Plugins/AI/MassAI/Source/MassZoneGraphNavigation/Public/MassZoneGraphNavigationUtils.h
+++ b/Engine/Plugins/AI/MassAI/Source/MassZoneGraphNavigation/Public/MassZoneGraphNavigationUtils.h
@@ -27,9 +27,9 @@ namespace UE::MassNavigation
 											 const FZoneGraphShortPathRequest& PathRequest,
 											 const float AgentRadius,
 											 const float DesiredSpeed,
-											 FMassMoveTargetFragment& MoveTarget,
-											 FMassZoneGraphShortPathFragment& ShortPath,
-											 FMassZoneGraphCachedLaneFragment& CachedLane);
+											 FMassMoveTargetFragment& InOutMoveTarget,
+											 FMassZoneGraphShortPathFragment& OutShortPath,
+											 FMassZoneGraphCachedLaneFragment& OutCachedLane);
 
 	MASSZONEGRAPHNAVIGATION_API bool ActivateActionStand(const UWorld& World,
 											  const UObject* Requester,
