[Mass navmesh navigation and avoidance] Add leading and trailing points at the start and end of the short path (required for environment avoidance).
UMassMovingAvoidanceProcessor: environement avoidance can now work with extruded edges (giving more freedom of movement in tight spots).
UMassMovingAvoidanceProcessor: avoidance bug fixes in environment separation and avoidance prediction handling.
UMassSteerToMoveTargetProcessor: clearer and improved debug display.

diff --git a/Engine/Plugins/AI/MassAI/Source/MassAIBehavior/Private/Tasks/MassNavMeshPathFollowTask.cpp b/Engine/Plugins/AI/MassAI/Source/MassAIBehavior/Private/Tasks/MassNavMeshPathFollowTask.cpp
--- a/Engine/Plugins/AI/MassAI/Source/MassAIBehavior/Private/Tasks/MassNavMeshPathFollowTask.cpp
+++ b/Engine/Plugins/AI/MassAI/Source/MassAIBehavior/Private/Tasks/MassNavMeshPathFollowTask.cpp
@@ -88,9 +88,9 @@ bool FMassNavMeshPathFollowTask::RequestPath(FStateTreeExecutionContext& Context
 
 		// Update short path
 		FMassNavMeshShortPathFragment& ShortPathFragment = Context.GetExternalData(ShortPathHandle);
-		ShortPathFragment.RequestShortPath(CachedPathFragment.Corridor, 0);
+		ShortPathFragment.RequestShortPath(CachedPathFragment.Corridor, /*NavCorridorStartIndex*/0, /*NumLeadingPoints*/0);
 
-		CachedPathFragment.NavPathNextStartIndex = ShortPathFragment.NumPoints-1;
+		CachedPathFragment.NavPathNextStartIndex = (uint16)FMath::Max(ShortPathFragment.NumPoints - FMassNavMeshShortPathFragment::NumPointsBeyondUpdate - FMassNavMeshCachedPathFragment::NumLeadingPoints, 0);
 		
 		// Update MoveTarget
 		FMassMoveTargetFragment& MoveTarget = Context.GetExternalData(MoveTargetHandle);
@@ -117,9 +117,9 @@ bool FMassNavMeshPathFollowTask::UpdateShortPath(FStateTreeExecutionContext& Con
 	
 	FMassNavMeshShortPathFragment& ShortPathFragment = Context.GetExternalData(ShortPathHandle);
 
-	ShortPathFragment.RequestShortPath(CachedPathFragment.Corridor, CachedPathFragment.NavPathNextStartIndex);
+	ShortPathFragment.RequestShortPath(CachedPathFragment.Corridor, CachedPathFragment.NavPathNextStartIndex, CachedPathFragment.NumLeadingPoints);
 
-	CachedPathFragment.NavPathNextStartIndex += ShortPathFragment.NumPoints-1;
+	CachedPathFragment.NavPathNextStartIndex += (uint16)FMath::Max(ShortPathFragment.NumPoints - FMassNavMeshShortPathFragment::NumPointsBeyondUpdate - FMassNavMeshCachedPathFragment::NumLeadingPoints, 0);
 
 	return true;
 }
diff --git a/Engine/Plugins/AI/MassAI/Source/MassNavMeshNavigation/Private/MassNavMeshNavigationBoundaryProcessor.cpp b/Engine/Plugins/AI/MassAI/Source/MassNavMeshNavigation/Private/MassNavMeshNavigationBoundaryProcessor.cpp
--- a/Engine/Plugins/AI/MassAI/Source/MassNavMeshNavigation/Private/MassNavMeshNavigationBoundaryProcessor.cpp
+++ b/Engine/Plugins/AI/MassAI/Source/MassNavMeshNavigation/Private/MassNavMeshNavigationBoundaryProcessor.cpp
@@ -86,7 +86,8 @@ void UMassNavMeshNavigationBoundaryProcessor::Execute(FMassEntityManager& Entity
 			}
 
 			// Make environment edges from the short path.
-			for (int32 Index = 0; Index+1 < ShortPath.NumPoints; Index++)
+			Edges.bExtrudedEdges = true;
+			for (int32 Index = 0; Index+1 < ShortPath.NumPoints && 2*Index < Edges.MaxEdgesCount; Index++)
 			{
 				const FMassNavMeshPathPoint& Portal = ShortPath.Points[Index];
 				const FMassNavMeshPathPoint& NextPortal = ShortPath.Points[Index+1];
diff --git a/Engine/Plugins/AI/MassAI/Source/MassNavMeshNavigation/Private/MassNavMeshNavigationFragments.cpp b/Engine/Plugins/AI/MassAI/Source/MassNavMeshNavigation/Private/MassNavMeshNavigationFragments.cpp
--- a/Engine/Plugins/AI/MassAI/Source/MassNavMeshNavigation/Private/MassNavMeshNavigationFragments.cpp
+++ b/Engine/Plugins/AI/MassAI/Source/MassNavMeshNavigation/Private/MassNavMeshNavigationFragments.cpp
@@ -6,7 +6,7 @@
 #include "NavCorridor.h"
 
 // Make a new short corridor
-bool FMassNavMeshShortPathFragment::RequestShortPath(const TSharedPtr<FNavCorridor>& NavCorridor, const int32 NavCorridorStartIndex)
+bool FMassNavMeshShortPathFragment::RequestShortPath(const TSharedPtr<FNavCorridor>& NavCorridor, const int32 NavCorridorStartIndex, const uint8 NumLeadingPoints)
 {
 	Reset();
 
@@ -50,6 +50,9 @@ bool FMassNavMeshShortPathFragment::RequestShortPath(const TSharedPtr<FNavCorrid
 		const FVector& Tangent = (Point.Position - PrevPoint.Position).GetSafeNormal();
 		PrevPoint.Tangent = FMassSnorm8Vector2D(Tangent);
 	}
+
+	// Update ProgressDistance to account for leading points. 
+	ProgressDistance = Points[NumLeadingPoints].Distance.Get();
 	
 	return true;
 }
diff --git a/Engine/Plugins/AI/MassAI/Source/MassNavMeshNavigation/Private/MassNavMeshNavigationProcessors.cpp b/Engine/Plugins/AI/MassAI/Source/MassNavMeshNavigation/Private/MassNavMeshNavigationProcessors.cpp
--- a/Engine/Plugins/AI/MassAI/Source/MassNavMeshNavigation/Private/MassNavMeshNavigationProcessors.cpp
+++ b/Engine/Plugins/AI/MassAI/Source/MassNavMeshNavigation/Private/MassNavMeshNavigationProcessors.cpp
@@ -10,7 +10,6 @@
 #include "MassNavigationFragments.h"
 #include "MassSignalSubsystem.h"
 #include "MassSimulationLOD.h"
-#include "NavCorridor.h"
 
 #define UNSAFE_FOR_MT 1
 
@@ -58,7 +57,6 @@ void UMassNavMeshPathFollowProcessor::Execute(FMassEntityManager& EntityManager,
 		const bool bHasLOD = (SimLODList.Num() > 0);
 
 		const TConstArrayView<FMassSimulationVariableTickFragment> SimVariableTickList = Context.GetFragmentView<FMassSimulationVariableTickFragment>();
-		const bool bHasVariableTick = (SimVariableTickList.Num() > 0);
 		const float WorldDeltaTime = Context.GetDeltaTimeSeconds();
 
 		for (int32 EntityIndex = 0; EntityIndex < NumEntities; ++EntityIndex)
@@ -66,7 +64,6 @@ void UMassNavMeshPathFollowProcessor::Execute(FMassEntityManager& EntityManager,
 			FMassNavMeshShortPathFragment& ShortPath = ShortPathList[EntityIndex];
 			FMassMoveTargetFragment& MoveTarget = MoveTargetList[EntityIndex];
 			const FMassEntityHandle Entity = Context.GetEntity(EntityIndex);
-			const float DeltaTime = bHasVariableTick ? SimVariableTickList[EntityIndex].DeltaTime : WorldDeltaTime;
 
 			bool bDisplayDebug = false;
 #if WITH_MASSGAMEPLAY_DEBUG && UNSAFE_FOR_MT // this will result in bDisplayDebug == false and disabling of all the vlogs below
@@ -86,27 +83,43 @@ void UMassNavMeshPathFollowProcessor::Execute(FMassEntityManager& EntityManager,
 				if (!bHasSteering || !MoveTarget.bSteeringFallingBehind)
 				{
 					// Update progress
-					ShortPath.ProgressDistance += MoveTarget.DesiredSpeed.Get() * DeltaTime;
+					ShortPath.ProgressDistance += MoveTarget.DesiredSpeed.Get() * WorldDeltaTime;
 				}
 				
 				if (!bWasDone)
 				{
-					const uint8 LastPointIndex = ShortPath.NumPoints - 1;
+					const uint8 LastIndex = ShortPath.NumPoints - 1;
+
+					// Point index on the short path where we should update the path.
+					const uint8 UpdatePointIndex =  ShortPath.bPartialResult ? (ShortPath.NumPoints - ShortPath.NumPointsBeyondUpdate) : LastIndex;
+
+#if WITH_MASSGAMEPLAY_DEBUG && UNSAFE_FOR_MT
+					UE_CVLOG_UELOG(bDisplayDebug, this, LogMassNavMeshNavigation, Verbose, TEXT("Entity [%s]   ProgressDistance: %f, UpdatePointIndex: %i"),
+						*Entity.DebugGetDescription(), ShortPath.ProgressDistance, UpdatePointIndex);
+
+					if (bDisplayDebug)
+					{
+						// Display update point
+						const FVector ZOffset(0,0,20);
+						UE_VLOG_CIRCLE(this, LogMassNavMeshNavigation, Display, ShortPath.Points[UpdatePointIndex].Position + ZOffset,
+							FVector(0,0,1), /*Radius*/10.f, FColor::Red, TEXT("%i"), UpdatePointIndex);
+					}
+#endif // WITH_MASSGAMEPLAY_DEBUG
 
 					if (ShortPath.ProgressDistance <= 0.0f)
 					{
-						// Requested time before the start of the path
+						// Before the start of the path
 						MoveTarget.Center = ShortPath.Points[0].Position;
 						MoveTarget.Forward = ShortPath.Points[0].Tangent.GetVector();
-						MoveTarget.DistanceToGoal = ShortPath.Points[LastPointIndex].Distance.Get();
+						MoveTarget.DistanceToGoal = ShortPath.Points[UpdatePointIndex].Distance.Get();
 
 						UE_CVLOG_UELOG(bDisplayDebug, this, LogMassNavMeshNavigation, Log, TEXT("Entity [%s]   before start of path. Distance to goal: %.1f."),
 							*Entity.DebugGetDescription(),
 							MoveTarget.DistanceToGoal);
 					}
-					else if (ShortPath.ProgressDistance <= ShortPath.Points[LastPointIndex].Distance.Get())
+					else if (ShortPath.ProgressDistance <= ShortPath.Points[UpdatePointIndex].Distance.Get())
 					{
-						// Requested time along the path, interpolate.
+						// Along the path, interpolate.
 						uint8 PointIndex = 0;
 						while (PointIndex < (ShortPath.NumPoints - 2))
 						{
@@ -117,15 +130,15 @@ void UMassNavMeshPathFollowProcessor::Execute(FMassEntityManager& EntityManager,
 							}
 							PointIndex++;
 						}
-
 						const FMassNavMeshPathPoint& CurrPoint = ShortPath.Points[PointIndex];
 						const FMassNavMeshPathPoint& NextPoint = ShortPath.Points[PointIndex + 1];
+						
 						const float T = (ShortPath.ProgressDistance - CurrPoint.Distance.Get()) / (NextPoint.Distance.Get() - CurrPoint.Distance.Get());
 
 						// Set move target location using the new progress distance.
 						MoveTarget.Center = FMath::Lerp(CurrPoint.Position, NextPoint.Position, T);
 						MoveTarget.Forward = FMath::Lerp(CurrPoint.Tangent.GetVector(), NextPoint.Tangent.GetVector(), T).GetSafeNormal();
-						MoveTarget.DistanceToGoal = ShortPath.Points[LastPointIndex].Distance.Get() - FMath::Lerp(CurrPoint.Distance.Get(), NextPoint.Distance.Get(), T);
+						MoveTarget.DistanceToGoal = ShortPath.Points[UpdatePointIndex].Distance.Get() - FMath::Lerp(CurrPoint.Distance.Get(), NextPoint.Distance.Get(), T);
 
 						UE_CVLOG_UELOG(bDisplayDebug, this, LogMassNavMeshNavigation, Verbose, TEXT("Entity [%s]   Distance to goal on short path: %.1f."),
 							*Entity.DebugGetDescription(),
@@ -133,11 +146,9 @@ void UMassNavMeshPathFollowProcessor::Execute(FMassEntityManager& EntityManager,
 					}
 					else
 					{
-						// @todo: probably need to do this earlier to make sure to avoid any discontinuity when updating the short path.
-						
-						// End of short path reached.
-						MoveTarget.Center = ShortPath.Points[LastPointIndex].Position;
-						MoveTarget.Forward = ShortPath.Points[LastPointIndex].Tangent.GetVector();
+						// Update point reached.
+						MoveTarget.Center = ShortPath.Points[UpdatePointIndex].Position;
+						MoveTarget.Forward = ShortPath.Points[UpdatePointIndex].Tangent.GetVector();
 						MoveTarget.DistanceToGoal = 0.f;
 
 						UE_CVLOG_UELOG(bDisplayDebug, this, LogMassNavMeshNavigation, Log, TEXT("Entity [%s]   Finished path follow on short path."), *Entity.DebugGetDescription());
@@ -180,6 +191,13 @@ void UMassNavMeshPathFollowProcessor::Execute(FMassEntityManager& EntityManager,
 						UE_VLOG_SEGMENT_THICK(this, LogMassNavMeshNavigation, Display, CurrPoint.Right + ZOffset, NextPoint.Right + ZOffset, DarkEntityColor, /*Thickness*/2, TEXT(""));
 					}
 
+					// Draw point indices
+					for (uint8 PointIndex = 0; PointIndex < ShortPath.NumPoints; PointIndex++)
+					{
+						const FMassNavMeshPathPoint& CurrPoint = ShortPath.Points[PointIndex];
+						UE_VLOG_CIRCLE(this, LogMassNavMeshNavigation, Display, CurrPoint.Position + ZOffset, FVector(0,0,1), /*Radius*/4.f, EntityColor, TEXT("%i"), PointIndex);
+					}
+
 					// Draw tangents
 					for (uint8 PointIndex = 0; PointIndex < ShortPath.NumPoints; PointIndex++)
 					{
diff --git a/Engine/Plugins/AI/MassAI/Source/MassNavMeshNavigation/Public/MassNavMeshNavigationFragments.h b/Engine/Plugins/AI/MassAI/Source/MassNavMeshNavigation/Public/MassNavMeshNavigationFragments.h
--- a/Engine/Plugins/AI/MassAI/Source/MassNavMeshNavigation/Public/MassNavMeshNavigationFragments.h
+++ b/Engine/Plugins/AI/MassAI/Source/MassNavMeshNavigation/Public/MassNavMeshNavigationFragments.h
@@ -40,8 +40,12 @@ struct MASSNAVMESHNAVIGATION_API FMassNavMeshShortPathFragment : public FMassFra
 	GENERATED_BODY()
 
 	FMassNavMeshShortPathFragment() = default;
-	
-	static constexpr uint8 MaxPoints = 4;
+
+	/* Maximum number of points on the short path */
+	static constexpr uint8 MaxPoints = 8;
+
+	/* Number of points beyond the update point */
+	static constexpr uint8 NumPointsBeyondUpdate = 3;
 
 	void Reset()
 	{
@@ -59,7 +63,7 @@ struct MASSNAVMESHNAVIGATION_API FMassNavMeshShortPathFragment : public FMassFra
 	}
 
 	/** Fill Points from Corridor. */
-	bool RequestShortPath(const TSharedPtr<FNavCorridor>& Corridor, const int32 NavCorridorStartIndex);
+	bool RequestShortPath(const TSharedPtr<FNavCorridor>& Corridor, const int32 NavCorridorStartIndex, const uint8 NumLeadingPoints);
 
 	/** Portal points */
 	TStaticArray<FMassNavMeshPathPoint, MaxPoints> Points;
@@ -94,6 +98,9 @@ struct MASSNAVMESHNAVIGATION_API FMassNavMeshCachedPathFragment : public FMassFr
 
 	/** Index used to keep track of progression on the navmesh path. */
 	uint16 NavPathNextStartIndex = 0;
+
+	/** Number of points before the starting point. */
+	static constexpr uint8 NumLeadingPoints = 1;
 };
 
 USTRUCT()
diff --git a/Engine/Plugins/AI/MassAI/Source/MassNavigation/Private/MassAvoidanceProcessors.cpp b/Engine/Plugins/AI/MassAI/Source/MassNavigation/Private/MassAvoidanceProcessors.cpp
--- a/Engine/Plugins/AI/MassAI/Source/MassNavigation/Private/MassAvoidanceProcessors.cpp
+++ b/Engine/Plugins/AI/MassAI/Source/MassNavigation/Private/MassAvoidanceProcessors.cpp
@@ -178,6 +178,7 @@ namespace UE::MassAvoidance
 	static const FColor ObstacleColor = FColor::Blue;
 	static const FColor ObstacleContactNormalColor = FColor::Silver;
 	static const FColor ObstacleAvoidForceColor = FColor::Magenta;
+	static const FColor ObstacleAvoidFuturHitPosColor = FColor::Magenta;
 	static constexpr FColor ObstacleSeparationForceColor = FColor(255, 66, 66);	// Bright red
 	
 	static const FVector DebugAgentHeightOffset = FVector(0., 0., 185.);
@@ -248,18 +249,19 @@ namespace UE::MassAvoidance
 		}
 	}
 
-	static void DebugDrawSphere(const FDebugContext& Context, const FVector& Center, const float Radius, const FColor& Color)
+	static void DebugDrawSphere(const FDebugContext& Context, const FVector& Center, const FVector::FReal InRadius, const FColor& Color)
 	{
 		if (!DebugIsSelected(Context.Entity))
 		{
 			return;
 		}
 
+		const float Radius = static_cast<float>(InRadius);
 		UE_VLOG_LOCATION(Context.LogOwner, Context.Category, Log, Center, Radius, Color, TEXT(""));
 
 		if (UseDrawDebugHelper() && Context.World)
 		{
-			DrawDebugSphere(Context.World, Center, Radius, /*segments = */16, Color);
+			DrawDebugSphere(Context.World, Center, Radius, /*Segments = */16, Color);
 		}
 	}
 
@@ -278,18 +280,35 @@ namespace UE::MassAvoidance
 		}
 	}
 	
-	static void DebugDrawCylinder(const FDebugContext& Context, const FVector& Bottom, const FVector& Top, const float Radius, const FColor& Color, const FString& Text = FString())
+	static void DebugDrawCylinder(const FDebugContext& Context, const FVector& Bottom, const FVector& Top, const FVector::FReal InRadius, const FColor& Color, const FString& Text = FString())
 	{
 		if (!DebugIsSelected(Context.Entity))
 		{
 			return;
 		}
 
+		const float Radius = static_cast<float>(InRadius);
 		UE_VLOG_CYLINDER(Context.LogOwner, Context.Category, Log, Bottom, Top, Radius, Color, TEXT("%s"), *Text);
 
 		if (UseDrawDebugHelper() && Context.World)
 		{
-			DrawDebugCylinder(Context.World, Bottom, Top, Radius, /*segments = */24, Color);
+			DrawDebugCylinder(Context.World, Bottom, Top, Radius, /*Segments = */24, Color);
+		}
+	}
+	
+	static void DebugDrawCircle(const FDebugContext& Context, const FVector& Bottom, const FVector::FReal InRadius, const FColor& Color, const FString& Text = FString())
+	{
+		if (!DebugIsSelected(Context.Entity))
+		{
+			return;
+		}
+
+		const float Radius = static_cast<float>(InRadius);
+		UE_VLOG_CIRCLE_THICK(Context.LogOwner, Context.Category, Log, Bottom, FVector(0,0,1), Radius, Color, /*Thickness*/2, TEXT("%s"), *Text);
+
+		if (UseDrawDebugHelper() && Context.World)
+		{
+			DrawDebugCircle(Context.World, Bottom, Radius, /*Segments = */24, Color);
 		}
 	}
 	//----------------------------------------------------------------------//
@@ -422,8 +441,9 @@ void UMassMovingAvoidanceProcessor::Execute(FMassEntityManager& EntityManager, F
 		{
 			FVector Position = FVector::ZeroVector;
 			FVector Normal = FVector::ZeroVector;
-			FVector::FReal Distance = 0.;
-			bool bRelevantForSeparation = false;
+			FVector::FReal Distance = 0.;				// Distance from agent center to contact position
+			bool bAlongTheEdge = false;
+			bool bBehindTheEdge = false;
 		};
 		TArray<FEnvironmentContact, TInlineAllocator<16>> Contacts;
 
@@ -451,7 +471,7 @@ void UMassMovingAvoidanceProcessor::Execute(FMassEntityManager& EntityManager, F
 			const FMassNavigationEdgesFragment& NavEdges = NavEdgesList[EntityIndex];
 			const FTransformFragment& Location = LocationList[EntityIndex];
 			const FMassVelocityFragment& Velocity = VelocityList[EntityIndex];
-			const FAgentRadiusFragment& Radius = RadiusList[EntityIndex];
+			const FAgentRadiusFragment& RadiusFragement = RadiusList[EntityIndex];
 			FMassForceFragment& Force = ForceList[EntityIndex];
 
 			// Smaller steering max accel makes the steering more "calm" but less opportunistic, may not find solution, or gets stuck.
@@ -462,9 +482,7 @@ void UMassMovingAvoidanceProcessor::Execute(FMassEntityManager& EntityManager, F
 			const FVector AgentLocation = Location.GetTransform().GetTranslation();
 			const FVector AgentVelocity = FVector(Velocity.Value.X, Velocity.Value.Y, 0.);
 			
-			const FVector::FReal AgentRadius = Radius.Radius;
-			const FVector::FReal SeparationAgentRadius = Radius.Radius * MovingAvoidanceParams.SeparationRadiusScale;
-			const FVector::FReal PredictiveAvoidanceAgentRadius = Radius.Radius * MovingAvoidanceParams.PredictiveAvoidanceRadiusScale;
+			const FVector::FReal AgentRadius = RadiusFragement.Radius;
 			
 			FVector SteeringForce = Force.Value;
 
@@ -495,24 +513,27 @@ void UMassMovingAvoidanceProcessor::Execute(FMassEntityManager& EntityManager, F
 			const UE::MassAvoidance::FDebugContext VelocitiesDebugContext(this, LogAvoidanceVelocities, World, Entity);
 			const UE::MassAvoidance::FDebugContext ObstacleDebugContext(this, LogAvoidanceObstacles, World, Entity);
 			const UE::MassAvoidance::FDebugContext AgentDebugContext(this, LogAvoidanceAgents, World, Entity);
-			
-			if (UE::MassAvoidance::DebugIsSelected(Entity))
+
+			FColor EntityColor = FColor::White;
+			if (UE::Mass::Debug::IsDebuggingEntity(Entity, &EntityColor))
 			{
 				// Draw agent
 				const FString Text = FString::Printf(TEXT("%i"), Entity.Index);
-				DebugDrawCylinder(BaseDebugContext, AgentLocation, AgentLocation + UE::MassAvoidance::DebugAgentHeightOffset, static_cast<float>(AgentRadius+1.),
+				DebugDrawCylinder(BaseDebugContext, AgentLocation, AgentLocation + UE::MassAvoidance::DebugAgentHeightOffset, (AgentRadius+1.),
 					UE::MassAvoidance::CurrentAgentColor, Text);
 
+				// Draw agent center
 				DebugDrawSphere(BaseDebugContext, AgentLocation, 10.f, UE::MassAvoidance::CurrentAgentColor);
 
+				// Draw circle for agent in LogMassNavigation.
+				const FVector ZOffset(0,0,25);
+				UE_VLOG_WIRECIRCLE_THICK(this, LogMassNavigation, Log, AgentLocation + ZOffset, FVector::UpVector, AgentRadius, EntityColor, /*Thickness*/4,
+					TEXT("%s"), *Entity.DebugGetDescription(), TEXT("%s"), *Entity.DebugGetDescription());
+
 				// Draw current velocity (black)
 				DebugDrawVelocity(VelocitiesDebugContext, AgentLocation + UE::MassAvoidance::DebugAgentHeightOffset,
 					AgentLocation + UE::MassAvoidance::DebugAgentHeightOffset + AgentVelocity, UE::MassAvoidance::VelocityColor);
 
-				// Draw preferred velocity (red)
-//				DebugDrawVelocity(VelocitiesDebugContext, AgentLocation + UE::MassAvoidance::DebugAgentHeightOffset,
-//					AgentLocation + UE::MassAvoidance::DebugAgentHeightOffset + PrefVelocity, UE::MassAvoidance::PrefVelocityColor);
-
 				// Draw initial steering force
 				DebugDrawArrow(BaseDebugContext, AgentLocation + UE::MassAvoidance::DebugAgentHeightOffset,
 					AgentLocation +UE::MassAvoidance:: DebugAgentHeightOffset + SteeringForce,
@@ -531,6 +552,9 @@ void UMassMovingAvoidanceProcessor::Execute(FMassEntityManager& EntityManager, F
 			
 			if (!MoveTarget.bOffBoundaries && UE::MassAvoidance::Tweakables::bEnableEnvironmentAvoidance)
 			{
+				const FVector::FReal EnvironmentSeparationAgentRadius = (RadiusFragement.Radius * MovingAvoidanceParams.SeparationRadiusScale) - (NavEdges.bExtrudedEdges ? AgentRadius : 0);
+				const FVector::FReal EnvironmentPredictiveAvoidanceAgentRadius = (RadiusFragement.Radius * MovingAvoidanceParams.PredictiveAvoidanceRadiusScale) - (NavEdges.bExtrudedEdges ? AgentRadius : 0);
+				
 				const FVector DesiredAcceleration = UE::MassNavigation::ClampVector(SteeringForce, MaxSteerAccel);
 				const FVector DesiredVelocity = UE::MassNavigation::ClampVector(AgentVelocity + DesiredAcceleration * DeltaTime, MaximumSpeed);
 
@@ -543,6 +567,8 @@ void UMassMovingAvoidanceProcessor::Execute(FMassEntityManager& EntityManager, F
 				OldSteeringForce = SteeringForce;
 				Contacts.Reset();
 
+				const FVector::FReal AgentRadiusForEnvironment = NavEdges.bExtrudedEdges ? 0 : AgentRadius;
+
 				// Collect potential contacts between agent and environment edges.
 				for (const FNavigationAvoidanceEdge& Edge : NavEdges.AvoidanceEdges)
 				{
@@ -555,97 +581,94 @@ void UMassMovingAvoidanceProcessor::Execute(FMassEntityManager& EntityManager, F
 					const FVector::FReal DistAlongEdge = FVector::DotProduct(EdgeDir, EdgeStartToAgent);
 					const FVector::FReal DistAwayFromEdge = FVector::DotProduct(Edge.LeftDir, EdgeStartToAgent);
 
-					FVector::FReal ConDist = 0.;
-					FVector ConNorm = FVector::ForwardVector;
-					FVector ConPos = FVector::ZeroVector;
-					bool bDirectlyBehindEdge = false;
+					FVector ConPos = FVector::ZeroVector; 		// contact position on edge
+					bool bBehindTheEdge = false;
 					bool bAlongTheEdge = false;
-					
-					if (DistAwayFromEdge < 0.)
+
+					if (DistAwayFromEdge < 0)
 					{
-						// Behind the edge
-						if (DistAlongEdge < 0.)
+						bBehindTheEdge = true;
+
+						if (DistAlongEdge < 0)
 						{
+							// Start corner
 							ConPos = Edge.Start;
-							ConNorm = -EdgeDir;
-							ConDist = -DistAlongEdge;
 						}
 						else if (DistAlongEdge > EdgeLength)
 						{
+							// End corner
 							ConPos = Edge.End;
-							ConNorm = EdgeDir;
-							ConDist = DistAlongEdge;
 						}
 						else
 						{
-							ConPos = Edge.Start + EdgeDir * DistAlongEdge;
-							ConNorm = Edge.LeftDir;
-							ConDist = 0.;
-							bDirectlyBehindEdge = true;
+							// Directly behind the edge
 							bAlongTheEdge = true;
+							ConPos = Edge.Start + EdgeDir * DistAlongEdge;
 						}
 					}
 					else
 					{
-						if (DistAlongEdge < 0.)
+						// In front of the edge
+						if (DistAlongEdge < 0)
 						{
-							// Start Corner
+							// Start corner
 							ConPos = Edge.Start;
 						}
 						else if (DistAlongEdge > EdgeLength)
 						{
-							// End Corner
+							// End corner
 							ConPos = Edge.End;
 						}
 						else
 						{
-							// Front
-							ConPos = Edge.Start + EdgeDir * DistAlongEdge;
+							// Directly in front of the edge
 							bAlongTheEdge = true;
+							ConPos = Edge.Start + EdgeDir * DistAlongEdge;
 						}
-
-						ConNorm = Edge.LeftDir;
-						ConDist = DistAwayFromEdge;
-
 					}
 					
-					// Check to merge contacts
-					bool bAdd = true;
-					for (int ContactIndex = 0; ContactIndex < Contacts.Num(); ContactIndex++)
+					// Add new contact
+					FEnvironmentContact Contact;
+					Contact.Position = ConPos;
+					Contact.Normal = Edge.LeftDir;
+					Contact.Distance = DistAwayFromEdge;
+					Contact.bAlongTheEdge = bAlongTheEdge;
+					Contact.bBehindTheEdge = bBehindTheEdge;
+					Contacts.Add(Contact);
+
+#if WITH_MASSGAMEPLAY_DEBUG && UNSAFE_FOR_MT
+					if (bAlongTheEdge)
 					{
-						if (FVector::DotProduct(Contacts[ContactIndex].Normal, ConNorm) > 0. &&
-							FMath::Abs(FVector::DotProduct(ConNorm, Contacts[ContactIndex].Position - ConPos)) < (10./*cm*/))
-						{
-							// Contacts are on same place, merge
-							if (ConDist < Contacts[ContactIndex].Distance)
-							{
-								// New is closer, override.
-								Contacts[ContactIndex].Position = ConPos;
-								Contacts[ContactIndex].Normal = ConNorm;
-								Contacts[ContactIndex].Distance = ConDist;
-								Contacts[ContactIndex].bRelevantForSeparation = bAlongTheEdge;
-							}
-							bAdd = false;
-							break;
-						}
-					}
+						// Draw active edges
+						FVector ZOffset = FVector(0., 0., 7.);
+						DebugDrawLine(ObstacleDebugContext, ZOffset + Edge.Start, ZOffset + Edge.End, FColor::Yellow, /*Thickness=*/3.f);
 
-					// Not found, add new contact
-					if (bAdd)
+						// Draw environment separation distance
+						const FVector SeparationOffset = MovingAvoidanceParams.EnvironmentSeparationDistance * Edge.LeftDir;
+						DebugDrawLine(ObstacleDebugContext, ZOffset + Edge.Start + SeparationOffset, ZOffset + Edge.End + SeparationOffset, FColor::Yellow, /*Thickness=*/1.f);
+					}
+					
+					if (bBehindTheEdge)
 					{
-						FEnvironmentContact Contact;
-						Contact.Position = ConPos;
-						Contact.Normal = ConNorm;
-						Contact.Distance = ConDist;
-						Contact.bRelevantForSeparation = bAlongTheEdge;
-						Contacts.Add(Contact);
+						// Draw inactive eges
+						FVector Offset = FVector(0., 0., 9.);
+						DebugDrawLine(ObstacleDebugContext, Offset + Edge.Start, Offset + Edge.End, FColor::Silver, /*Thickness=*/2.f);	
 					}
-
+#endif // WITH_MASSGAMEPLAY_DEBUG
+					
 					// Skip predictive avoidance when behind the edge.
-					if (!bDirectlyBehindEdge)
+					if (!bBehindTheEdge)
 					{
 						// Avoid edges
-						const FVector::FReal CPA = UE::MassAvoidance::ComputeClosestPointOfApproach(FVector2D(AgentLocation), FVector2D(DesiredVelocity), AgentRadius,
+						
+#if WITH_MASSGAMEPLAY_DEBUG && UNSAFE_FOR_MT
+						// Draw environment predictive avoidance distance
+						FVector ZOffset = FVector(0., 0., 7.);
+						const FVector AvoidanceOffset = MovingAvoidanceParams.PredictiveAvoidanceDistance * Edge.LeftDir;
+						DebugDrawLine(ObstacleDebugContext, ZOffset + Edge.Start + AvoidanceOffset, ZOffset + Edge.End + AvoidanceOffset, UE::MassAvoidance::ObstacleAvoidFuturHitPosColor, /*Thickness=*/1.f);
+#endif // WITH_MASSGAMEPLAY_DEBUG
+
+						const FVector::FReal CPA = UE::MassAvoidance::ComputeClosestPointOfApproach(FVector2D(AgentLocation), FVector2D(DesiredVelocity), AgentRadiusForEnvironment,
 							FVector2D(Edge.Start), FVector2D(Edge.End), MovingAvoidanceParams.PredictiveAvoidanceTime);
 						const FVector HitAgentPos = AgentLocation + DesiredVelocity * CPA;
 						const FVector::FReal EdgeT = UE::MassNavigation::ProjectPtSeg(FVector2D(HitAgentPos), FVector2D(Edge.Start), FVector2D(Edge.End));
@@ -655,10 +678,15 @@ void UMassMovingAvoidanceProcessor::Execute(FMassEntityManager& EntityManager, F
 						FVector AvoidRelPos = HitAgentPos - HitObPos;
 						AvoidRelPos.Z = 0.;	// @todo AT: ignore the z component for now until we clamp the height of obstacles
 						const FVector::FReal AvoidDist = AvoidRelPos.Size();
-						const FVector AvoidNormal = AvoidDist > 0. ? (AvoidRelPos / AvoidDist) : FVector::ForwardVector;
+						const FVector AvoidNormal = AvoidDist > 0. ? (AvoidRelPos / AvoidDist) : Edge.LeftDir;
+
+						const FVector::FReal AvoidPen = (EnvironmentPredictiveAvoidanceAgentRadius + MovingAvoidanceParams.PredictiveAvoidanceDistance) - AvoidDist;
+						FVector::FReal AvoidMag = 1.;
+						if (MovingAvoidanceParams.PredictiveAvoidanceDistance != 0.)
+						{
+							AvoidMag = FMath::Square(FMath::Clamp(AvoidPen / MovingAvoidanceParams.PredictiveAvoidanceDistance, 0., 1.));
+						}
 
-						const FVector::FReal AvoidPen = (PredictiveAvoidanceAgentRadius + MovingAvoidanceParams.PredictiveAvoidanceDistance) - AvoidDist;
-						const FVector::FReal AvoidMag = FMath::Square(FMath::Clamp(AvoidPen / MovingAvoidanceParams.PredictiveAvoidanceDistance, 0., 1.));
 						const FVector::FReal AvoidMagDist = 1. + FMath::Square(1. - CPA * InvPredictiveAvoidanceTime);
 
 						// Predictive avoidance against environment is tuned down towards the end of the path
@@ -667,17 +695,20 @@ void UMassMovingAvoidanceProcessor::Execute(FMassEntityManager& EntityManager, F
 						SteeringForce += AvoidForce;
 
 #if WITH_MASSGAMEPLAY_DEBUG && UNSAFE_FOR_MT
-						// Draw contact normal
-						UE::MassAvoidance::DebugDrawArrow(ObstacleDebugContext, ConPos, ConPos + (ConDist * ConNorm), UE::MassAvoidance::ObstacleContactNormalColor, /*HeadSize=*/ 5.f);
-						UE::MassAvoidance::DebugDrawSphere(ObstacleDebugContext, ConPos, 2.5f, UE::MassAvoidance::ObstacleContactNormalColor);
+						if (!AvoidForce.IsNearlyZero())
+						{
+							// Draw contact normal
+							DebugDrawArrow(ObstacleDebugContext, ConPos, ConPos + (Contact.Distance * Contact.Normal), UE::MassAvoidance::ObstacleContactNormalColor, /*HeadSize=*/ 5.f);
+							DebugDrawSphere(ObstacleDebugContext, ConPos, 2.5f, UE::MassAvoidance::ObstacleContactNormalColor);
 
-						// Draw hit pos with edge
-						UE::MassAvoidance::DebugDrawLine(ObstacleDebugContext, AgentLocation, HitAgentPos, UE::MassAvoidance::ObstacleAvoidForceColor);
-						UE::MassAvoidance::DebugDrawCylinder(ObstacleDebugContext, HitAgentPos, HitAgentPos + UE::MassAvoidance::DebugAgentHeightOffset,
-							static_cast<float>(AgentRadius), UE::MassAvoidance::ObstacleAvoidForceColor);
+							// Draw future hit pos with edge
+							DebugDrawSphere(ObstacleDebugContext, HitAgentPos, 1.f, UE::MassAvoidance::ObstacleAvoidFuturHitPosColor);
+							DebugDrawCircle(ObstacleDebugContext, HitAgentPos, AgentRadius, UE::MassAvoidance::ObstacleAvoidFuturHitPosColor);
+							DebugDrawLine(ObstacleDebugContext, AgentLocation, HitAgentPos, UE::MassAvoidance::ObstacleAvoidFuturHitPosColor);
 
-						// Draw avoid obstacle force
-						UE::MassAvoidance::DebugDrawForce(ObstacleDebugContext, HitObPos, HitObPos + AvoidForce, UE::MassAvoidance::ObstacleAvoidForceColor);
+							// Draw individual predictive avoidance obstacle force
+							DebugDrawForce(ObstacleDebugContext, HitObPos, HitObPos + AvoidForce, UE::MassAvoidance::ObstacleAvoidForceColor);
+						}
 #endif // WITH_MASSGAMEPLAY_DEBUG
 					}
 				} // edge loop
@@ -685,32 +716,67 @@ void UMassMovingAvoidanceProcessor::Execute(FMassEntityManager& EntityManager, F
 #if WITH_MASSGAMEPLAY_DEBUG && UNSAFE_FOR_MT
 				// Draw total steering force to avoid obstacles
 				const FVector EnvironmentAvoidSteeringForce = SteeringForce - OldSteeringForce;
-				UE::MassAvoidance::DebugDrawSummedForce(ObstacleDebugContext,
+				DebugDrawSummedForce(ObstacleDebugContext,
 					AgentLocation + UE::MassAvoidance::DebugAgentHeightOffset,
 					AgentLocation + UE::MassAvoidance::DebugAgentHeightOffset + EnvironmentAvoidSteeringForce,
 					UE::MassAvoidance::ObstacleAvoidForceColor);
-#endif // WITH_MASSGAMEPLAY_DEBUG
 
+				// Draw all contact points
+				for (const FEnvironmentContact& Contact : Contacts) 
+				{
+					DebugDrawSphere(ObstacleDebugContext, Contact.Position + FVector(0.f, 0.f, Contact.bBehindTheEdge ? 10.f : 0.f), 5.f, Contact.bBehindTheEdge ? FColor::Black: FColor::Cyan);					
+				}
+#endif // WITH_MASSGAMEPLAY_DEBUG
+				
 				// Process contacts to add edge separation force
 				const FVector SteeringForceBeforeSeparation = SteeringForce;
 				for (const FEnvironmentContact& Contact : Contacts) 
 				{
-					if (Contact.bRelevantForSeparation)
+					if (Contact.bAlongTheEdge && Contact.Distance > -AgentRadius)
 					{
-						const FVector ConNorm = Contact.Normal.GetSafeNormal();
-						const FVector::FReal ContactDist = Contact.Distance;
+						check(Contact.Normal.IsNormalized());
 
 						// Separation force (stay away from obstacles if possible)
-						const FVector::FReal SeparationPenalty = (SeparationAgentRadius + MovingAvoidanceParams.EnvironmentSeparationDistance) - ContactDist;
-						const FVector::FReal SeparationMag = UE::MassNavigation::Smooth(FMath::Clamp(SeparationPenalty / MovingAvoidanceParams.EnvironmentSeparationDistance, 0., 1.));
-						const FVector SeparationForce = ConNorm * MovingAvoidanceParams.EnvironmentSeparationStiffness * SeparationMag;
+						const FVector::FReal SeparationPenalty = (EnvironmentSeparationAgentRadius + MovingAvoidanceParams.EnvironmentSeparationDistance) - Contact.Distance;
 
+						FVector::FReal SeparationMag;
+						if (MovingAvoidanceParams.EnvironmentSeparationDistance != 0.)
+						{
+							SeparationMag = UE::MassNavigation::Smooth(FMath::Clamp(SeparationPenalty / MovingAvoidanceParams.EnvironmentSeparationDistance, 0., 1.));
+						}
+						else
+						{
+							SeparationMag = 1.;
+						}
+						const FVector SeparationForce = Contact.Normal * MovingAvoidanceParams.EnvironmentSeparationStiffness * SeparationMag;
+						
 						SteeringForce += SeparationForce;
 
 	#if WITH_MASSGAMEPLAY_DEBUG && UNSAFE_FOR_MT
-						// Draw individual contact forces
-						DebugDrawForce(ObstacleDebugContext, Contact.Position + UE::MassAvoidance::DebugAgentHeightOffset,
-						Contact.Position + SeparationForce + UE::MassAvoidance::DebugAgentHeightOffset, UE::MassAvoidance::ObstacleSeparationForceColor);
+						UE_VLOG(ObstacleDebugContext.LogOwner, ObstacleDebugContext.Category, Log, TEXT("EnvironmentSeparationAgentRadius: %f, EnvironmentSeparationDistanc: %f, Distance: %f, SeparationPenalty: %f, SeparationMag: %f, SeparationForce: %s"),
+							EnvironmentSeparationAgentRadius,
+							MovingAvoidanceParams.EnvironmentSeparationDistance,
+							Contact.Distance,
+							SeparationPenalty,
+							SeparationMag,
+							*SeparationForce.ToCompactString()
+							);
+
+						FVector ZOffset = FVector(0., 0., 7.);
+						
+						// Draw contact normal
+						float SphereRadius = 2.5f;
+						FColor ContactColor = FColor::Yellow;
+						if (!SeparationForce.IsNearlyZero())
+						{
+							SphereRadius = 5.f;
+							ContactColor = UE::MassAvoidance::ObstacleSeparationForceColor;
+							
+							// Draw individual contact forces
+							DebugDrawForce(ObstacleDebugContext, Contact.Position + ZOffset, Contact.Position + SeparationForce + ZOffset, ContactColor);
+						}
+						DebugDrawSphere(ObstacleDebugContext, Contact.Position + ZOffset, SphereRadius, ContactColor);
+						DebugDrawArrow(ObstacleDebugContext, Contact.Position + ZOffset, Contact.Position + (Contact.Distance * Contact.Normal) + ZOffset, ContactColor, /*HeadSize=*/ 5.f);
 	#endif // WITH_MASSGAMEPLAY_DEBUG
 					}
 				}
@@ -728,11 +794,11 @@ void UMassMovingAvoidanceProcessor::Execute(FMassEntityManager& EntityManager, F
 				{
 					for (const FNavigationAvoidanceEdge& Edge : NavEdges.AvoidanceEdges)
 					{
-						DebugDrawLine(ObstacleDebugContext, UE::MassAvoidance::DebugAgentHeightOffset + Edge.Start,
-							UE::MassAvoidance::DebugAgentHeightOffset + Edge.End, UE::MassAvoidance::ObstacleColor, /*Thickness=*/2.f);
-						const FVector Middle = UE::MassAvoidance::DebugAgentHeightOffset + 0.5f * (Edge.Start + Edge.End);
-						DebugDrawArrow(ObstacleDebugContext, Middle, Middle + 10. * FVector::CrossProduct((Edge.End - Edge.Start), FVector::UpVector).GetSafeNormal(),
-							UE::MassAvoidance::ObstacleColor, /*HeadSize=*/2.f);
+						FVector Offset = FVector(0., 0., 5.);
+						DebugDrawLine(ObstacleDebugContext, Offset + Edge.Start, Offset + Edge.End, UE::MassAvoidance::ObstacleColor, /*Thickness=*/2.f);
+
+						const FVector Middle = Offset + 0.5f * (Edge.Start + Edge.End);
+						DebugDrawArrow(ObstacleDebugContext, Middle, Middle + 10. * Edge.LeftDir, UE::MassAvoidance::ObstacleColor, /*HeadSize=*/2.f);
 					}
 				}
 #endif // WITH_MASSGAMEPLAY_DEBUG
@@ -740,6 +806,8 @@ void UMassMovingAvoidanceProcessor::Execute(FMassEntityManager& EntityManager, F
 
 			//////////////////////////////////////////////////////////////////////////
 			// Avoid close agents
+			const FVector::FReal SeparationAgentRadius = RadiusFragement.Radius * MovingAvoidanceParams.SeparationRadiusScale;
+			const FVector::FReal PredictiveAvoidanceAgentRadius = RadiusFragement.Radius * MovingAvoidanceParams.PredictiveAvoidanceRadiusScale;
 
 			// Update desired velocity based on avoidance so far.
 			const FVector DesAcc = UE::MassNavigation::ClampVector(SteeringForce, MaxSteerAccel);
@@ -986,7 +1054,7 @@ void UMassMovingAvoidanceProcessor::Execute(FMassEntityManager& EntityManager, F
 					UE::MassAvoidance::DebugDrawLine(AgentDebugContext, AgentLocation + UE::MassAvoidance::DebugAgentHeightOffset - LeftOffset,
 						HitPosition + UE::MassAvoidance::DebugAgentHeightOffset - LeftOffset, UE::MassAvoidance::CurrentAgentColor, 1.5f);
 					UE::MassAvoidance::DebugDrawCylinder(AgentDebugContext, HitPosition,
-						HitPosition + UE::MassAvoidance::DebugAgentHeightOffset, static_cast<float>(PredictiveAvoidanceAgentRadius), UE::MassAvoidance::CurrentAgentColor);
+						HitPosition + UE::MassAvoidance::DebugAgentHeightOffset, PredictiveAvoidanceAgentRadius, UE::MassAvoidance::CurrentAgentColor);
 
 					const FVector OtherHitPosition = Collider.Location + (Collider.Velocity * CPA);
 					const FVector OtherLeftOffset = Collider.Radius * UE::MassNavigation::GetLeftDirection(Collider.Velocity.GetSafeNormal(), FVector::UpVector);
@@ -995,9 +1063,9 @@ void UMassMovingAvoidanceProcessor::Execute(FMassEntityManager& EntityManager, F
 					UE::MassAvoidance::DebugDrawLine(AgentDebugContext, Collider.Location + Left, OtherHitPosition + Left, UE::MassAvoidance::AgentsColor, 1.5f);
 					UE::MassAvoidance::DebugDrawLine(AgentDebugContext, Collider.Location + Right, OtherHitPosition + Right, UE::MassAvoidance::AgentsColor, 1.5f);
 					UE::MassAvoidance::DebugDrawCylinder(AgentDebugContext, Collider.Location, Collider.Location + UE::MassAvoidance::DebugAgentHeightOffset,
-						static_cast<float>(AgentRadius), UE::MassAvoidance::AgentsColor);
+						AgentRadius, UE::MassAvoidance::AgentsColor);
 					UE::MassAvoidance::DebugDrawCylinder(AgentDebugContext, OtherHitPosition, OtherHitPosition + UE::MassAvoidance::DebugAgentHeightOffset,
-						static_cast<float>(AgentRadius), UE::MassAvoidance::AgentsColor);
+						AgentRadius, UE::MassAvoidance::AgentsColor);
 
 					// Draw agent avoid force
 					UE::MassAvoidance::DebugDrawForce(AgentDebugContext,
diff --git a/Engine/Plugins/AI/MassAI/Source/MassNavigation/Private/MassSteeringProcessors.cpp b/Engine/Plugins/AI/MassAI/Source/MassNavigation/Private/MassSteeringProcessors.cpp
--- a/Engine/Plugins/AI/MassAI/Source/MassNavigation/Private/MassSteeringProcessors.cpp
+++ b/Engine/Plugins/AI/MassAI/Source/MassNavigation/Private/MassSteeringProcessors.cpp
@@ -145,6 +145,18 @@ void UMassSteerToMoveTargetProcessor::Execute(FMassEntityManager& EntityManager,
 				{
 					SteerDirection *= 1. / DistanceToSteerTarget;
 				}
+
+#if WITH_MASSGAMEPLAY_DEBUG && UNSAFE_FOR_MT
+				FColor EntityColor = FColor::White;
+				const bool bDisplayDebug = UE::Mass::Debug::IsDebuggingEntity(Entity, &EntityColor);
+				if (bDisplayDebug)
+				{
+					// Display SteerDirection
+					const FVector ZOffset(0,0,25);
+					UE_VLOG_SEGMENT_THICK(this, LogMassNavigation, Log, CurrentLocation + ZOffset, CurrentLocation + ZOffset + 100.*SteerDirection,
+						FColor::Red, /*Thickness*/2, TEXT("SteerDirection"));
+				}
+#endif //WITH_MASSGAMEPLAY_DEBUG
 				
 				const FVector::FReal DirSpeedScale = UE::MassNavigation::CalcDirectionalSpeedScale(CurrentForward, SteerDirection);
 				FVector::FReal DesiredSpeed = MoveTarget.DesiredSpeed.Get() * DirSpeedScale;
@@ -263,7 +275,6 @@ void UMassSteerToMoveTargetProcessor::Execute(FMassEntityManager& EntityManager,
 					// @todo: This current completely overrides steering, we probably should have one processor that resets the steering at the beginning of the frame.
 					Steering.DesiredVelocity = SteerDirection * DesiredSpeed;
 					Force.Value = SteerK * (Steering.DesiredVelocity - Velocity.Value); // Goal force
-
 				}
 				else
 				{
@@ -298,26 +309,36 @@ void UMassSteerToMoveTargetProcessor::Execute(FMassEntityManager& EntityManager,
 				
 				const FVector MoveTargetCenter = MoveTarget.Center + ZOffset;
 
-				// MoveTarget slack boundary
-				UE_VLOG_CIRCLE_THICK(this, LogMassNavigation, Log, MoveTargetCenter, FVector::UpVector, LookAheadDistance, EntityColor, /*Thickness*/2,
-					TEXT("%s MoveTgt %s"), *Entity.DebugGetDescription(), *UEnum::GetDisplayValueAsText(MoveTarget.IntentAtGoal).ToString());
+				// Display MoveTarget location
+				UE_VLOG_CIRCLE_THICK(this, LogMassNavigation, Log, MoveTargetCenter, FVector::UpVector, 5.f, EntityColor, /*Thickness*/2,
+					TEXT("MoveTarget\ngoal: %s"), *UEnum::GetDisplayValueAsText(MoveTarget.IntentAtGoal).ToString());
 
-				// MoveTarget orientation
-				UE_VLOG_SEGMENT_THICK(this, LogMassNavigation, Log, MoveTargetCenter, MoveTargetCenter + MoveTarget.Forward * LookAheadDistance, EntityColor, /*Thickness*/2, TEXT(""));
+				// Display MoveTarget orientation
+				UE_VLOG_SEGMENT_THICK(this, LogMassNavigation, Log, MoveTargetCenter, MoveTargetCenter + MoveTarget.Forward * 100, EntityColor, /*Thickness*/1,
+					TEXT("MoveTarget\nforward"));
 
-				// MoveTarget - current location relation.
+				// Display MoveTarget - current location relation
 				if (FVector::Dist2D(CurrentLocation, MoveTarget.Center) > LookAheadDistance * 1.5f)
 				{
 					UE_VLOG_SEGMENT_THICK(this, LogMassNavigation, Log, MoveTargetCenter, CurrentLocation + ZOffset, FColor::Red, /*Thickness*/1, TEXT("LOST"));
 				}
 				else
 				{
-					UE_VLOG_SEGMENT_THICK(this, LogMassNavigation, Log, MoveTargetCenter, CurrentLocation + ZOffset, DarkEntityColor, /*Thickness*/1, TEXT(""));
+					UE_VLOG_SEGMENT_THICK(this, LogMassNavigation, Log, MoveTargetCenter, CurrentLocation + ZOffset, LightEntityColor, /*Thickness*/1, TEXT(""));
 				}
 
-				// Steering
-				UE_VLOG_SEGMENT_THICK(this, LogMassNavigation, Log, CurrentLocation + ZOffset, CurrentLocation + Steering.DesiredVelocity + ZOffset, LightEntityColor, /*Thickness*/2,
-					TEXT("%s Steer %.1f"), *Entity.DebugGetDescription(), Steering.DesiredVelocity.Length());
+				// Display Velocity
+				UE_VLOG_SEGMENT_THICK(this, LogMassNavigation, Log, CurrentLocation + ZOffset, CurrentLocation + ZOffset + Velocity.Value, FColor::Yellow, /*Thickness*/4,
+					TEXT("Velocity %.1f"), Velocity.Value.Length());
+
+				// Display DesiredVelocity
+				UE_VLOG_SEGMENT_THICK(this, LogMassNavigation, Log, CurrentLocation + ZOffset + FVector(0,0,2),
+					CurrentLocation + Steering.DesiredVelocity + ZOffset + FVector(0,0,2), FColor::Orange, /*Thickness*/4,
+					TEXT("DesiredVelocity %.1f"), Steering.DesiredVelocity.Length());
+
+				// Display Force
+				UE_VLOG_SEGMENT_THICK(this, LogMassNavigation, Log, CurrentLocation + ZOffset, CurrentLocation + Force.Value + ZOffset, FColor::Emerald, /*Thickness*/4,
+					TEXT("Force %.1f"), Force.Value.Length());
 			}
 #endif // WITH_MASSGAMEPLAY_DEBUG
 			
diff --git a/Engine/Plugins/AI/MassAI/Source/MassNavigation/Public/Avoidance/MassAvoidanceFragments.h b/Engine/Plugins/AI/MassAI/Source/MassNavigation/Public/Avoidance/MassAvoidanceFragments.h
--- a/Engine/Plugins/AI/MassAI/Source/MassNavigation/Public/Avoidance/MassAvoidanceFragments.h
+++ b/Engine/Plugins/AI/MassAI/Source/MassNavigation/Public/Avoidance/MassAvoidanceFragments.h
@@ -58,6 +58,7 @@ struct MASSNAVIGATION_API FMassMovingAvoidanceParameters : public FMassConstShar
 	UPROPERTY(EditAnywhere, Category = "Separation", meta = (ClampMin = "0", ForceUnits="x"))
 	float SeparationRadiusScale = 0.9f;
 
+	
 	/** Separation force stiffness between agents and obstacles. Range: 100..500 N/cm */
 	UPROPERTY(EditAnywhere, Category = "Separation", meta = (ClampMin = "0"))
 	float ObstacleSeparationStiffness = 250.f;
@@ -74,6 +75,7 @@ struct MASSNAVIGATION_API FMassMovingAvoidanceParameters : public FMassConstShar
 	UPROPERTY(EditAnywhere, Category = "Separation", meta = (ClampMin = "0", ForceUnits="cm"))
 	float EnvironmentSeparationDistance = 50.f;
 
+	
 	/** How far in the future the agent reacts to collisions. Range: 1..3, Indoor humans 1.4, outdoor humans 2.4 (seconds). */
 	UPROPERTY(EditAnywhere, Category = "Predictive Avoidance", meta = (ClampMin = "0.1", ForceUnits="s"))
 	float PredictiveAvoidanceTime = 2.5f;
@@ -180,6 +182,9 @@ struct MASSNAVIGATION_API FMassNavigationEdgesFragment : public FMassFragment
 {
 	GENERATED_BODY()
 
-	static const int MaxEdgesCount = 8;
+	static const int MaxEdgesCount = 16;
 	TArray<FNavigationAvoidanceEdge, TFixedAllocator<MaxEdgesCount>> AvoidanceEdges;
+
+	/* Set to true if the edges already account for the agent radius. */
+	bool bExtrudedEdges = false;
 };
diff --git a/Engine/Plugins/AI/MassAI/Source/MassNavigation/Public/MassNavigationUtils.h b/Engine/Plugins/AI/MassAI/Source/MassNavigation/Public/MassNavigationUtils.h
--- a/Engine/Plugins/AI/MassAI/Source/MassNavigation/Public/MassNavigationUtils.h
+++ b/Engine/Plugins/AI/MassAI/Source/MassNavigation/Public/MassNavigationUtils.h
@@ -5,6 +5,7 @@
 #include "CoreTypes.h"
 #include "Math/Quat.h"
 #include "Math/UnrealMathUtility.h"
+#include "MassCommonTypes.h"
 
 namespace UE::MassNavigation
 {
@@ -111,5 +112,20 @@ namespace UE::MassNavigation
 		}
 		return Mid;
 	}
-	
-} // UE::MassMovement
\ No newline at end of file
+} // UE::MassNavigation
+
+#if WITH_MASSGAMEPLAY_DEBUG
+
+namespace UE::MassNavigation::Debug
+{
+	inline FColor MixColors(const FColor ColorA, const FColor ColorB)
+	{
+		const int32 R = ((int32)ColorA.R + (int32)ColorB.R) / 2;
+		const int32 G = ((int32)ColorA.G + (int32)ColorB.G) / 2;
+		const int32 B = ((int32)ColorA.B + (int32)ColorB.B) / 2;
+		const int32 A = ((int32)ColorA.A + (int32)ColorB.A) / 2;
+		return FColor((uint8)R, (uint8)G, (uint8)B, (uint8)A);
+	}
+}
+
+#endif // WITH_MASSGAMEPLAY_DEBUG
diff --git a/Engine/Plugins/AI/MassAI/Source/MassZoneGraphNavigation/Private/MassZoneGraphNavigationProcessors.cpp b/Engine/Plugins/AI/MassAI/Source/MassZoneGraphNavigation/Private/MassZoneGraphNavigationProcessors.cpp
--- a/Engine/Plugins/AI/MassAI/Source/MassZoneGraphNavigation/Private/MassZoneGraphNavigationProcessors.cpp
+++ b/Engine/Plugins/AI/MassAI/Source/MassZoneGraphNavigation/Private/MassZoneGraphNavigationProcessors.cpp
@@ -19,22 +19,6 @@
 
 #define UNSAFE_FOR_MT 1
 
-#if WITH_MASSGAMEPLAY_DEBUG
-
-namespace UE::MassNavigation::Debug
-{
-	FColor MixColors(const FColor ColorA, const FColor ColorB)
-	{
-		const int32 R = ((int32)ColorA.R + (int32)ColorB.R) / 2;
-		const int32 G = ((int32)ColorA.G + (int32)ColorB.G) / 2;
-		const int32 B = ((int32)ColorA.B + (int32)ColorB.B) / 2;
-		const int32 A = ((int32)ColorA.A + (int32)ColorB.A) / 2;
-		return FColor((uint8)R, (uint8)G, (uint8)B, (uint8)A);
-	}
-}
-
-#endif // WITH_MASSGAMEPLAY_DEBUG
-
 //----------------------------------------------------------------------//
 //  UMassZoneGraphLocationInitializer
 //----------------------------------------------------------------------//
