[MassGameplay] Mass State Tree changes to support dynamic processors per-ST-requirements

diff --git a/Engine/Plugins/AI/MassAI/Source/MassAIBehavior/Private/MassStateTreeProcessors.cpp b/Engine/Plugins/AI/MassAI/Source/MassAIBehavior/Private/MassStateTreeProcessors.cpp
--- a/Engine/Plugins/AI/MassAI/Source/MassAIBehavior/Private/MassStateTreeProcessors.cpp
+++ b/Engine/Plugins/AI/MassAI/Source/MassAIBehavior/Private/MassStateTreeProcessors.cpp
@@ -10,6 +10,7 @@
 #include "MassSmartObjectTypes.h"
 #include "MassZoneGraphAnnotationTypes.h"
 #include "MassStateTreeSubsystem.h"
+#include "MassStateTreeFragments.h"
 #include "MassSignalSubsystem.h"
 #include "ProfilingDebugging/CsvProfiler.h"
 #include "Engine/World.h"
@@ -230,6 +231,9 @@ UMassStateTreeProcessor::UMassStateTreeProcessor(const FObjectInitializer& Objec
 	// but only on `SyncWorldToMass`
 	ExecutionOrder.ExecuteAfter.Add(UE::Mass::ProcessorGroupNames::SyncWorldToMass);
 	ExecutionOrder.ExecuteBefore.Add(UE::Mass::ProcessorGroupNames::Tasks);
+
+	bAutoRegisterWithProcessingPhases = false;
+	bAllowMultipleInstances = true;
 }
 
 void UMassStateTreeProcessor::InitializeInternal(UObject& Owner, const TSharedRef<FMassEntityManager>& EntityManager)
@@ -349,3 +353,33 @@ void UMassStateTreeProcessor::SignalEntities(FMassEntityManager& EntityManager,
 		SignalSubsystem.SignalEntities(UE::Mass::Signals::NewStateTreeTaskRequired, EntitiesToSignal);
 	}
 }
+
+void UMassStateTreeProcessor::SetExecutionRequirements(const FMassFragmentRequirements& FragmentRequirements
+	, const FMassSubsystemRequirements& SubsystemRequirements)
+{
+	if (!ensureMsgf(IsInitialized() == false, TEXT("%hs: calling after processor's initialization is not supported."), __FUNCTION__))
+	{
+		return;
+	}
+	FragmentRequirements.ExportRequirements(ExecutionRequirements);
+	SubsystemRequirements.ExportRequirements(ExecutionRequirements);
+}
+
+void UMassStateTreeProcessor::ExportRequirements(FMassExecutionRequirements& OutRequirements) const
+{
+	Super::ExportRequirements(OutRequirements);
+	OutRequirements.Append(ExecutionRequirements);
+}
+
+void UMassStateTreeProcessor::AddHandledStateTree(TNotNull<const UStateTree*> StateTree)
+{
+	HandledStateTrees.AddUnique(StateTree);
+
+	// need to clear the filter first. This is a safety-feature to ensure filters do not get accidentally overridden 
+	EntityQuery.ClearChunkFilter();
+	EntityQuery.SetChunkFilter([HandledStateTreesCopy = HandledStateTrees](const FMassExecutionContext& Context) -> bool
+	{
+		const FMassStateTreeSharedFragment& SharedStateTree = Context.GetConstSharedFragment<FMassStateTreeSharedFragment>();
+		return HandledStateTreesCopy.Find(SharedStateTree.StateTree) != INDEX_NONE;
+	});
+}
diff --git a/Engine/Plugins/AI/MassAI/Source/MassAIBehavior/Private/MassStateTreeSubsystem.cpp b/Engine/Plugins/AI/MassAI/Source/MassAIBehavior/Private/MassStateTreeSubsystem.cpp
--- a/Engine/Plugins/AI/MassAI/Source/MassAIBehavior/Private/MassStateTreeSubsystem.cpp
+++ b/Engine/Plugins/AI/MassAI/Source/MassAIBehavior/Private/MassStateTreeSubsystem.cpp
@@ -1,13 +1,31 @@
 // Copyright Epic Games, Inc. All Rights Reserved.
 
 #include "MassStateTreeSubsystem.h"
+#include "MassAIBehaviorTypes.h"
+#include "MassEntityQuery.h"
+#include "MassEntitySubsystem.h"
+#include "MassProcessorDependencySolver.h"
+#include "MassSimulationSubsystem.h"
+#include "MassStateTreeSchema.h"
 #include "StateTree.h"
 #include "Engine/Engine.h"
+#include "MassStateTreeProcessors.h"
 
 
 void UMassStateTreeSubsystem::Initialize(FSubsystemCollectionBase& Collection)
 {
 	Super::Initialize(Collection);
+
+	UMassEntitySubsystem* EntitySubsystem = Collection.InitializeDependency<UMassEntitySubsystem>();
+	check(EntitySubsystem);
+	ĘntityManager = EntitySubsystem->GetMutableEntityManager().AsShared();
+
+	if (!DynamicProcessorClass)
+	{
+		DynamicProcessorClass = UMassStateTreeProcessor::StaticClass();
+	}
+
+	SimulationSubsystem = Collection.InitializeDependency<UMassSimulationSubsystem>();
 }
 
 FMassStateTreeInstanceHandle UMassStateTreeSubsystem::AllocateInstanceData(const UStateTree* StateTree)
@@ -30,6 +48,11 @@ FMassStateTreeInstanceHandle UMassStateTreeSubsystem::AllocateInstanceData(const
 
 	FMassStateTreeInstanceDataItem& Item = InstanceDataArray[Index];
 	Item.InstanceData.Reset();
+
+	if (StateTreeToProcessor.Find(StateTree) == nullptr)
+	{
+		CreateProcessorForStateTree(StateTree);
+	}
 	
 	return FMassStateTreeInstanceHandle::Make(Index, Item.Generation);
 }
@@ -47,3 +70,88 @@ void UMassStateTreeSubsystem::FreeInstanceData(const FMassStateTreeInstanceHandl
 
 	InstanceDataFreelist.Add(Handle.GetIndex());
 }
+
+void UMassStateTreeSubsystem::CreateProcessorForStateTree(TNotNull<const UStateTree*> StateTree)
+{
+	const UMassStateTreeSchema* StateTreeSchema = Cast<UMassStateTreeSchema>(StateTree->GetSchema());
+
+	FMassSubsystemRequirements SubsystemRequirements;
+	FMassFragmentRequirements FragmentRequirements(ĘntityManager);
+
+	if (ensure(StateTreeSchema))
+	{
+		const TConstArrayView<FMassStateTreeDependency> DependenciesView = StateTreeSchema->GetDependencies();
+		if (DependenciesView.Num())
+		{
+			// Convert loosely defined dependencies (expressed with UStruct pointers)
+			// to strongly-typed Mass requirements.
+			// The functions below will complain if anything is misconfigured.
+			for (const FMassStateTreeDependency& Dependency : DependenciesView)
+			{
+				if (Dependency.Type)
+				{
+					if (const UScriptStruct* ScriptStruct = Cast<const UScriptStruct>(Dependency.Type.Get()))
+					{
+						if (UE::Mass::IsA<FMassFragment>(Dependency.Type))
+						{
+							FragmentRequirements.AddRequirement(ScriptStruct, Dependency.Access);
+						}
+						else if (UE::Mass::IsA<FMassTag>(Dependency.Type))
+						{
+							FragmentRequirements.AddTagRequirement(*ScriptStruct, EMassFragmentPresence::All);
+						}
+						else if (UE::Mass::IsA<FMassChunkFragment>(Dependency.Type))
+						{
+							FragmentRequirements.AddChunkRequirement(ScriptStruct, Dependency.Access);
+						}
+						else if (UE::Mass::IsA<FMassSharedFragment>(Dependency.Type))
+						{
+							FragmentRequirements.AddSharedRequirement(ScriptStruct, Dependency.Access);
+						}
+						else if (UE::Mass::IsA<FMassConstSharedFragment>(Dependency.Type))
+						{
+							FragmentRequirements.AddConstSharedRequirement(ScriptStruct);
+						}
+						else
+						{
+							UE_LOG(LogMassBehavior, Error, TEXT("Unhandled Mass State Tree dependency %s"), *Dependency.Type->GetName());
+						}
+					}
+					else if (Dependency.Type->IsChildOf(USubsystem::StaticClass()))
+					{
+						TSubclassOf<USubsystem> SubsystemClass = Cast<UClass>(const_cast<UStruct*>(Dependency.Type.Get()));
+						SubsystemRequirements.AddSubsystemRequirement(SubsystemClass, Dependency.Access);
+					}
+					else
+					{
+						UE_LOG(LogMassBehavior, Error, TEXT("Unhandled Mass State Tree dependency user-type %s"), *Dependency.Type->GetName());
+					}
+				}
+			}
+		}
+	}
+
+	UMassStateTreeProcessor* DynamicProcessor = nullptr;
+
+	const uint32 DependenciesHash = HashCombine(GetTypeHash(FragmentRequirements), GetTypeHash(SubsystemRequirements));
+	TObjectPtr<UMassStateTreeProcessor>& ExistingProcessor = RequirementsHashToProcessor.FindOrAddByHash(DependenciesHash, DependenciesHash);
+
+	if (ExistingProcessor == nullptr)
+	{
+		DynamicProcessor = NewObject<UMassStateTreeProcessor>(this, DynamicProcessorClass);
+		checkf(DynamicProcessor, TEXT("Failed to spawn an instance of %s"), *GetNameSafe(DynamicProcessorClass));
+		DynamicProcessor->SetExecutionRequirements(FragmentRequirements, SubsystemRequirements);
+
+		SimulationSubsystem->RegisterDynamicProcessor(*DynamicProcessor);
+
+		ExistingProcessor = DynamicProcessor;
+	}
+	else
+	{
+		DynamicProcessor = ExistingProcessor;
+	}
+
+	DynamicProcessor->AddHandledStateTree(StateTree);
+	
+	StateTreeToProcessor.Add(StateTree, DynamicProcessor);
+}
\ No newline at end of file
diff --git a/Engine/Plugins/AI/MassAI/Source/MassAIBehavior/Public/MassStateTreeProcessors.h b/Engine/Plugins/AI/MassAI/Source/MassAIBehavior/Public/MassStateTreeProcessors.h
--- a/Engine/Plugins/AI/MassAI/Source/MassAIBehavior/Public/MassStateTreeProcessors.h
+++ b/Engine/Plugins/AI/MassAI/Source/MassAIBehavior/Public/MassStateTreeProcessors.h
@@ -3,14 +3,20 @@
 #pragma once
 
 #include "MassSignalProcessorBase.h"
+#if UE_ENABLE_INCLUDE_ORDER_DEPRECATED_IN_5_6
 #include "MassStateTreeFragments.h"
+#endif // UE_ENABLE_INCLUDE_ORDER_DEPRECATED_IN_5_6
 #include "MassObserverProcessor.h"
+#include "MassProcessorDependencySolver.h"
 #if UE_ENABLE_INCLUDE_ORDER_DEPRECATED_IN_5_6
 #include "MassLODTypes.h"
 #endif // UE_ENABLE_INCLUDE_ORDER_DEPRECATED_IN_5_6
 #include "MassStateTreeProcessors.generated.h"
 
 struct FMassStateTreeExecutionContext;
+struct FMassSubsystemRequirements;
+struct FMassFragmentRequirements;
+class UStateTree;
 
 /** 
  * Processor to stop and uninitialize StateTrees on entities.
@@ -58,7 +64,9 @@ protected:
 };
 
 /** 
- * Processor for executing a StateTree
+ * The processor that the UMassStateTreeSubsystem will instantiate for every unique StateTree Mass-requirements.
+ * The user is not expected to instantiate these processors manually, but a project-specific extension can be implemented.
+ * It needs to derive from UMassStateTreeProcessor and set as the value of UMassStateTreeSubsystem.DynamicProcessorClass.
  */
 UCLASS()
 class MASSAIBEHAVIOR_API UMassStateTreeProcessor : public UMassSignalProcessorBase
@@ -68,11 +76,36 @@ class MASSAIBEHAVIOR_API UMassStateTreeProcessor : public UMassSignalProcessorBa
 public:
 	UMassStateTreeProcessor(const FObjectInitializer& ObjectInitializer = FObjectInitializer::Get());
 
+	/**
+	 * Called to configure dynamic processor's additional requirements that will ensure its located
+	 * properly within Mass's processing graph. Calling this function is allowed only until the
+	 * processor is Initialized. The function will ensure that's the case. 
+	 */
+	void SetExecutionRequirements(const FMassFragmentRequirements& FragmentRequirements, const FMassSubsystemRequirements& SubsystemRequirements);
+
+	/**
+	 * Adds StateTree to the collection of the assets this specific processor instance will handle. 
+	 */
+	void AddHandledStateTree(TNotNull<const UStateTree*> StateTree);
+
 protected:
 	virtual void InitializeInternal(UObject& Owner, const TSharedRef<FMassEntityManager>& EntityManager) override;
 	virtual void ConfigureQueries(const TSharedRef<FMassEntityManager>& EntityManager) override;
 	virtual void SignalEntities(FMassEntityManager& EntityManager, FMassExecutionContext& Context, FMassSignalNameLookup& EntitySignals) override;
 
+	virtual void ExportRequirements(FMassExecutionRequirements& OutRequirements) const override;
+
+	/**
+	 * Stores the additional requirements as configured by SetExecutionRequirements.
+	 * These requirements ensure the processor will be placed at the right location in the processing graph
+	 * to avoid data races.
+	 */
+	FMassExecutionRequirements ExecutionRequirements;
+
+	/** The assets handled by this processor - entities utilizing any of these assets will be processed by this processor */
+	UPROPERTY()
+	TArray<TObjectPtr<const UStateTree>> HandledStateTrees;
+
 	/** Configures whether parallel update for FMassArchetypeChunks should be used instead of the default single threaded update (i.e., ParallelForEachEntityChunk instead of ForEachEntityChunk). */
 	UPROPERTY(EditDefaultsOnly, Category = Processor, config)
 	bool bProcessEntitiesInParallel = false;
diff --git a/Engine/Plugins/AI/MassAI/Source/MassAIBehavior/Public/MassStateTreeSchema.h b/Engine/Plugins/AI/MassAI/Source/MassAIBehavior/Public/MassStateTreeSchema.h
--- a/Engine/Plugins/AI/MassAI/Source/MassAIBehavior/Public/MassStateTreeSchema.h
+++ b/Engine/Plugins/AI/MassAI/Source/MassAIBehavior/Public/MassStateTreeSchema.h
@@ -14,14 +14,29 @@ class MASSAIBEHAVIOR_API UMassStateTreeSchema : public UStateTreeSchema
 {
 	GENERATED_BODY()
 
-protected:
+public:
+	/** Fetches a read-only view of the Mass-relevant requirements of the associated StateTee */
+	const TArray<FMassStateTreeDependency>& GetDependencies() const;
 
+protected:
 	virtual bool IsStructAllowed(const UScriptStruct* InScriptStruct) const override;
 	virtual bool IsExternalItemAllowed(const UStruct& InStruct) const override;
 	virtual bool Link(FStateTreeLinker& Linker) override;
 
 protected:
+	/**
+	 * The Mass-relevant requirements of the associated StateTee, collected during StateTree's linking
+	 * @see UStateTree::Link
+	 * @see UMassStateTreeSchema::Link
+	 */
 	UPROPERTY(Transient)
 	TArray<FMassStateTreeDependency> Dependencies;
 };
 
+//-----------------------------------------------------------------------------
+// INLINE
+//-----------------------------------------------------------------------------
+inline const TArray<FMassStateTreeDependency>& UMassStateTreeSchema::GetDependencies() const
+{
+	return Dependencies;
+}
diff --git a/Engine/Plugins/AI/MassAI/Source/MassAIBehavior/Public/MassStateTreeSubsystem.h b/Engine/Plugins/AI/MassAI/Source/MassAIBehavior/Public/MassStateTreeSubsystem.h
--- a/Engine/Plugins/AI/MassAI/Source/MassAIBehavior/Public/MassStateTreeSubsystem.h
+++ b/Engine/Plugins/AI/MassAI/Source/MassAIBehavior/Public/MassStateTreeSubsystem.h
@@ -6,9 +6,13 @@
 #include "MassSubsystemBase.h"
 #include "StateTreeExecutionContext.h"
 #include "MassExternalSubsystemTraits.h"
+#include "UObject/ObjectKey.h"
 #include "MassStateTreeSubsystem.generated.h"
 
 class UStateTree;
+class UMassStateTreeProcessor;
+class UMassSimulationSubsystem;
+struct FMassEntityManager;
 
 USTRUCT()
 struct MASSAIBEHAVIOR_API FMassStateTreeInstanceDataItem
@@ -63,10 +67,48 @@ public:
 	
 protected:
 
+	/**
+	 * Gathers Mass-relevant processing requirements from StateTree and spawns
+	 * a dynamic processor to handle entities using this given asset
+	 */
+	void CreateProcessorForStateTree(TNotNull<const UStateTree*> StateTree);
+
 	TArray<int32> InstanceDataFreelist;
 
 	UPROPERTY(Transient)
 	TArray<FMassStateTreeInstanceDataItem> InstanceDataArray;
+
+	/**
+	 * The relevant Entity Manager. Needed to build processing requirements for dynamic processors.
+	 */
+	TSharedPtr<FMassEntityManager> ĘntityManager;
+
+	/**
+	 * The key represents a hash of mass requirements calculated from a StateTree assets.
+	 * Using the hash rather than a StateTree asset pointer since multiple assets can have identical requirements,
+	 * and we want just one dynamic processor to handle all of them.
+	 * @note that might change if in practice it turns out that runtime entity chunk filtering turns out to be
+	 *	too expensive when multiple ST assets are used.
+	 */
+	UPROPERTY()
+	TMap<uint32, TObjectPtr<UMassStateTreeProcessor>> RequirementsHashToProcessor;
+
+	// @todo ask Patrick how would this behave when ST asset gets changed/recompiled or whatnot
+	/**
+	 * Mapping StateTree assets to the dynamic processors handling them. Note that it's not 1:1, a single processor can handle multiple assets.
+	 */
+	TMap<TObjectKey<UStateTree>, TObjectPtr<UMassStateTreeProcessor>> StateTreeToProcessor;
+
+	/** Cached SimulationSubsystem for registering dynamic processors. */
+	UPROPERTY()
+	TObjectPtr<UMassSimulationSubsystem> SimulationSubsystem;
+
+	/**
+	 * Class to use when creating dynamic processors to handle given StateTree assets.
+	 * Can be overridden via DefaultMass.ini file. 
+	 */
+	UPROPERTY(Config, EditDefaultsOnly, Category=Mass)
+	TSubclassOf<UMassStateTreeProcessor> DynamicProcessorClass;
 };
 
 template<>
