StateTree: Include linked state trees when building the dependencies list.

diff --git a/Engine/Plugins/AI/MassAI/Source/MassAIBehavior/Private/MassStateTreeSchema.cpp b/Engine/Plugins/AI/MassAI/Source/MassAIBehavior/Private/MassStateTreeSchema.cpp
--- a/Engine/Plugins/AI/MassAI/Source/MassAIBehavior/Private/MassStateTreeSchema.cpp
+++ b/Engine/Plugins/AI/MassAI/Source/MassAIBehavior/Private/MassStateTreeSchema.cpp
@@ -39,16 +39,35 @@ bool UMassStateTreeSchema::Link(FStateTreeLinker& Linker)
 	Dependencies.Empty();
 	UE::MassBehavior::FStateTreeDependencyBuilder Builder(Dependencies);
 
-	const UStateTree* StateTree = CastChecked<UStateTree>(GetOuter());
-	for (FConstStructView Node : StateTree->GetNodes())
-	{
-		if (const FMassStateTreeEvaluatorBase* Evaluator = Node.GetPtr<const FMassStateTreeEvaluatorBase>())
+	auto BuildDependencies = [&Builder](const UStateTree* StateTree)
 		{
-			Evaluator->GetDependencies(Builder);
-		}
-		else if (const FMassStateTreeTaskBase* Task = Node.GetPtr<const FMassStateTreeTaskBase>())
+			for (FConstStructView Node : StateTree->GetNodes())
+			{
+				if (const FMassStateTreeEvaluatorBase* Evaluator = Node.GetPtr<const FMassStateTreeEvaluatorBase>())
+				{
+					Evaluator->GetDependencies(Builder);
+				}
+				else if (const FMassStateTreeTaskBase* Task = Node.GetPtr<const FMassStateTreeTaskBase>())
+				{
+					Task->GetDependencies(Builder);
+				}
+			}
+		};
+
+	TArray<const UStateTree*, TInlineAllocator<4>> StateTrees;
+	StateTrees.Add(CastChecked<const UStateTree>(GetOuter()));
+	for (int32 Index = 0; Index < StateTrees.Num(); ++Index)
+	{
+		const UStateTree* StateTree = StateTrees[Index];
+		BuildDependencies(StateTree);
+
+		// The StateTree::Link order is not deterministic. Build the dependencies and add them to this list.
+		for (const FCompactStateTreeState& State : StateTree->GetStates())
 		{
-			Task->GetDependencies(Builder);
+			if (State.LinkedAsset && State.Type == EStateTreeStateType::LinkedAsset)
+			{
+				StateTrees.AddUnique(State.LinkedAsset);
+			}
 		}
 	}
 
