[Mass navmesh navigation and avoidance] Fix check on environment edge normals

#rnx

diff --git a/Engine/Plugins/AI/MassAI/Source/MassNavMeshNavigation/Private/MassNavMeshNavigationBoundaryProcessor.cpp b/Engine/Plugins/AI/MassAI/Source/MassNavMeshNavigation/Private/MassNavMeshNavigationBoundaryProcessor.cpp
--- a/Engine/Plugins/AI/MassAI/Source/MassNavMeshNavigation/Private/MassNavMeshNavigationBoundaryProcessor.cpp
+++ b/Engine/Plugins/AI/MassAI/Source/MassNavMeshNavigation/Private/MassNavMeshNavigationBoundaryProcessor.cpp
@@ -92,9 +92,18 @@ void UMassNavMeshNavigationBoundaryProcessor::Execute(FMassEntityManager& Entity
 				const FMassNavMeshPathPoint& Portal = ShortPath.Points[Index];
 				const FMassNavMeshPathPoint& NextPortal = ShortPath.Points[Index+1];
 
-				// Left side: reverse start and end to keep the normal inside.
-				Edges.AvoidanceEdges.Add(FNavigationAvoidanceEdge(NextPortal.Left, Portal.Left));
-				Edges.AvoidanceEdges.Add(FNavigationAvoidanceEdge(Portal.Right, NextPortal.Right));
+				// If the points are too close, just ignore the edges.
+
+				if (!(NextPortal.Left - Portal.Left).IsNearlyZero())
+				{
+					// Left side: reverse start and end to keep the normal inside.
+					Edges.AvoidanceEdges.Add(FNavigationAvoidanceEdge(NextPortal.Left, Portal.Left));
+				}
+
+				if (!(Portal.Right - NextPortal.Right).IsNearlyZero())
+				{
+					Edges.AvoidanceEdges.Add(FNavigationAvoidanceEdge(Portal.Right, NextPortal.Right));
+				}
 			}
 		}
 	});
diff --git a/Engine/Plugins/AI/MassAI/Source/MassNavigation/Private/MassAvoidanceProcessors.cpp b/Engine/Plugins/AI/MassAI/Source/MassNavigation/Private/MassAvoidanceProcessors.cpp
--- a/Engine/Plugins/AI/MassAI/Source/MassNavigation/Private/MassAvoidanceProcessors.cpp
+++ b/Engine/Plugins/AI/MassAI/Source/MassNavigation/Private/MassAvoidanceProcessors.cpp
@@ -734,8 +734,6 @@ void UMassMovingAvoidanceProcessor::Execute(FMassEntityManager& EntityManager, F
 				{
 					if (Contact.bAlongTheEdge && Contact.Distance > -AgentRadius)
 					{
-						check(Contact.Normal.IsNormalized());
-
 						// Separation force (stay away from obstacles if possible)
 						const FVector::FReal SeparationPenalty = (EnvironmentSeparationAgentRadius + MovingAvoidanceParams.EnvironmentSeparationDistance) - Contact.Distance;
 
diff --git a/Engine/Plugins/AI/MassAI/Source/MassNavigation/Public/Avoidance/MassAvoidanceFragments.h b/Engine/Plugins/AI/MassAI/Source/MassNavigation/Public/Avoidance/MassAvoidanceFragments.h
--- a/Engine/Plugins/AI/MassAI/Source/MassNavigation/Public/Avoidance/MassAvoidanceFragments.h
+++ b/Engine/Plugins/AI/MassAI/Source/MassNavigation/Public/Avoidance/MassAvoidanceFragments.h
@@ -170,7 +170,7 @@ struct MASSNAVIGATION_API FNavigationAvoidanceEdge
 	{
 		Start = InStart;
 		End = InEnd;
-		LeftDir = FVector::CrossProduct((End - Start).GetSafeNormal(), FVector::UpVector);
+		LeftDir = FVector::CrossProduct((End - Start), FVector::UpVector).GetSafeNormal();
 	}
 	FVector Start = FVector::ZeroVector;
 	FVector End = FVector::ZeroVector;
