[MassGameplay] Added a gate CVar for the "dynamic ST processors" feature. Now the new functionality is disabled by default

diff --git a/Engine/Plugins/AI/MassAI/Source/MassAIBehavior/Private/MassStateTreeProcessors.cpp b/Engine/Plugins/AI/MassAI/Source/MassAIBehavior/Private/MassStateTreeProcessors.cpp
--- a/Engine/Plugins/AI/MassAI/Source/MassAIBehavior/Private/MassStateTreeProcessors.cpp
+++ b/Engine/Plugins/AI/MassAI/Source/MassAIBehavior/Private/MassStateTreeProcessors.cpp
@@ -17,6 +17,7 @@
 #include "MassBehaviorSettings.h"
 #include "VisualLogger/VisualLogger.h"
 #include "MassDebugger.h"
+#include "MassStateTreeSubsystem.h"
 
 
 CSV_DEFINE_CATEGORY(StateTreeProcessor, true);
@@ -232,8 +233,13 @@ UMassStateTreeProcessor::UMassStateTreeProcessor(const FObjectInitializer& Objec
 	ExecutionOrder.ExecuteAfter.Add(UE::Mass::ProcessorGroupNames::SyncWorldToMass);
 	ExecutionOrder.ExecuteBefore.Add(UE::Mass::ProcessorGroupNames::Tasks);
 
-	bAutoRegisterWithProcessingPhases = false;
-	bAllowMultipleInstances = true;
+	if (UE::Mass::StateTree::bDynamicSTProcessorsEnabled
+		// we're not setting the values for CDO to keep it "default"
+		&& HasAnyFlags(RF_ClassDefaultObject) == false)
+	{
+		bAutoRegisterWithProcessingPhases = false;
+		bAllowMultipleInstances = true;
+	}
 }
 
 void UMassStateTreeProcessor::InitializeInternal(UObject& Owner, const TSharedRef<FMassEntityManager>& EntityManager)
diff --git a/Engine/Plugins/AI/MassAI/Source/MassAIBehavior/Private/MassStateTreeSubsystem.cpp b/Engine/Plugins/AI/MassAI/Source/MassAIBehavior/Private/MassStateTreeSubsystem.cpp
--- a/Engine/Plugins/AI/MassAI/Source/MassAIBehavior/Private/MassStateTreeSubsystem.cpp
+++ b/Engine/Plugins/AI/MassAI/Source/MassAIBehavior/Private/MassStateTreeSubsystem.cpp
@@ -12,6 +12,13 @@
 #include "MassStateTreeProcessors.h"
 
 
+namespace UE::Mass::StateTree
+{
+	bool bDynamicSTProcessorsEnabled = false;
+	FAutoConsoleVariableRef CVarDynamicSTEnabled(TEXT("ai.mass.DynamicSTProcessorsEnabled"), bDynamicSTProcessorsEnabled
+		, TEXT("Whether Dynamic ST processors will be created per distinct ST requirements."), ECVF_Default);
+}
+
 void UMassStateTreeSubsystem::Initialize(FSubsystemCollectionBase& Collection)
 {
 	Super::Initialize(Collection);
@@ -49,7 +56,8 @@ FMassStateTreeInstanceHandle UMassStateTreeSubsystem::AllocateInstanceData(const
 	FMassStateTreeInstanceDataItem& Item = InstanceDataArray[Index];
 	Item.InstanceData.Reset();
 
-	if (StateTreeToProcessor.Find(StateTree) == nullptr)
+	if (UE::Mass::StateTree::bDynamicSTProcessorsEnabled
+		&& StateTreeToProcessor.Find(StateTree) == nullptr)
 	{
 		CreateProcessorForStateTree(StateTree);
 	}
diff --git a/Engine/Plugins/AI/MassAI/Source/MassAIBehavior/Public/MassStateTreeSubsystem.h b/Engine/Plugins/AI/MassAI/Source/MassAIBehavior/Public/MassStateTreeSubsystem.h
--- a/Engine/Plugins/AI/MassAI/Source/MassAIBehavior/Public/MassStateTreeSubsystem.h
+++ b/Engine/Plugins/AI/MassAI/Source/MassAIBehavior/Public/MassStateTreeSubsystem.h
@@ -14,6 +14,11 @@ class UMassStateTreeProcessor;
 class UMassSimulationSubsystem;
 struct FMassEntityManager;
 
+namespace UE::Mass::StateTree
+{
+	extern bool bDynamicSTProcessorsEnabled;
+}
+
 USTRUCT()
 struct MASSAIBEHAVIOR_API FMassStateTreeInstanceDataItem
 {
