Move UE::MassNavigation::Debug utils to it's own file (resolving an old "to do")

diff --git a/Engine/Plugins/AI/MassAI/Source/MassNavigation/Private/MassAvoidanceProcessors.cpp b/Engine/Plugins/AI/MassAI/Source/MassNavigation/Private/MassAvoidanceProcessors.cpp
--- a/Engine/Plugins/AI/MassAI/Source/MassNavigation/Private/MassAvoidanceProcessors.cpp
+++ b/Engine/Plugins/AI/MassAI/Source/MassNavigation/Private/MassAvoidanceProcessors.cpp
@@ -16,7 +16,7 @@
 #include "MassNavigationUtils.h"
 #include "Engine/World.h"
 #include "MassDebugger.h"
-
+#include "MassNavigationDebug.h"
 
 #define UNSAFE_FOR_MT 1
 
@@ -153,13 +153,9 @@ namespace UE::MassAvoidance
 		return FMath::Clamp(T, 0., TimeHoriz);
 	}
 
-	static bool UseDrawDebugHelper()
-	{
-		return Tweakables::bUseDrawDebugHelpers;
-	}
-
-#if WITH_MASSGAMEPLAY_DEBUG	
-
+#if WITH_MASSGAMEPLAY_DEBUG
+	using namespace UE::MassNavigation::Debug;
+	
 	// Colors
 	static const FColor CurrentAgentColor = FColor::Emerald;
 
@@ -183,139 +179,7 @@ namespace UE::MassAvoidance
 	
 	static const FVector DebugAgentHeightOffset = FVector(0., 0., 185.);
 	static const FVector DebugLowCylinderOffset = FVector(0., 0., 20.);
-
-	//----------------------------------------------------------------------//
-	// Begin MassDebugUtils
-	// @todo: Extract those generic debug functions to a separate location
-	//----------------------------------------------------------------------//
-	struct FDebugContext
-	{
-		FDebugContext(const UObject* InLogOwner, const FLogCategoryBase& InCategory, const UWorld* InWorld, const FMassEntityHandle InEntity)
-			: LogOwner(InLogOwner)
-			, Category(InCategory)
-			, World(InWorld)
-			, Entity(InEntity)
-		{}
-
-		const UObject* LogOwner;
-		const FLogCategoryBase& Category;
-		const UWorld* World;
-		const FMassEntityHandle Entity;
-	};
-
-	static bool DebugIsSelected(const FMassEntityHandle Entity)
-	{
-		FColor Color;
-		return UE::Mass::Debug::IsDebuggingEntity(Entity, &Color);
-	}
-
-	static void DebugDrawLine(const FDebugContext& Context, const FVector& Start, const FVector& End, const FColor& Color, const float Thickness = 0.f, const bool bPersistent = false)
-	{
-		if (!DebugIsSelected(Context.Entity))
-		{
-			return;
-		}
-
-		UE_VLOG_SEGMENT_THICK(Context.LogOwner, Context.Category, Log, Start, End, Color, (int16)Thickness, TEXT(""));
-
-		if (UseDrawDebugHelper() && Context.World)
-		{
-			DrawDebugLine(Context.World, Start, End, Color, bPersistent, /*LifeTime =*/ -1.f, /*DepthPriority =*/ 0, Thickness);
-		}
-	}
-
-	static void DebugDrawArrow(const FDebugContext& Context, const FVector& Start, const FVector& End, const FColor& Color, const float HeadSize = 8.f, const float Thickness = 1.5f)
-	{
-		if (!DebugIsSelected(Context.Entity))
-		{
-			return;
-		}
-
-		constexpr FVector::FReal Pointyness = 1.8;
-		const FVector Line = End - Start;
-		const FVector UnitV = Line.GetSafeNormal();
-		const FVector Perp = FVector::CrossProduct(UnitV, FVector::UpVector);
-		const FVector Left = Perp - (Pointyness*UnitV);
-		const FVector Right = -Perp - (Pointyness*UnitV);
-		UE_VLOG_SEGMENT_THICK(Context.LogOwner, Context.Category, Log, Start, End, Color, (int16)Thickness, TEXT(""));
-		UE_VLOG_SEGMENT_THICK(Context.LogOwner, Context.Category, Log, End, End + HeadSize * Left, Color, (int16)Thickness, TEXT(""));
-		UE_VLOG_SEGMENT_THICK(Context.LogOwner, Context.Category, Log, End, End + HeadSize * Right, Color, (int16)Thickness, TEXT(""));
-
-		if (UseDrawDebugHelper() && Context.World)
-		{
-			DrawDebugLine(Context.World, Start, End, Color, /*bPersistent=*/ false, /*LifeTime =*/ -1.f, /*DepthPriority =*/ 0, Thickness);
-			DrawDebugLine(Context.World, End, End + HeadSize * Left, Color, /*bPersistent=*/ false, /*LifeTime =*/ -1.f, /*DepthPriority =*/ 0, Thickness);
-			DrawDebugLine(Context.World, End, End + HeadSize * Right, Color, /*bPersistent=*/ false, /*LifeTime =*/ -1.f, /*DepthPriority =*/ 0, Thickness);
-		}
-	}
-
-	static void DebugDrawSphere(const FDebugContext& Context, const FVector& Center, const FVector::FReal InRadius, const FColor& Color)
-	{
-		if (!DebugIsSelected(Context.Entity))
-		{
-			return;
-		}
-
-		const float Radius = static_cast<float>(InRadius);
-		UE_VLOG_LOCATION(Context.LogOwner, Context.Category, Log, Center, Radius, Color, TEXT(""));
-
-		if (UseDrawDebugHelper() && Context.World)
-		{
-			DrawDebugSphere(Context.World, Center, Radius, /*Segments = */16, Color);
-		}
-	}
-
-	static void DebugDrawBox(const FDebugContext& Context, const FBox& Box, const FColor& Color)
-	{
-		if (!DebugIsSelected(Context.Entity))
-		{
-			return;
-		}
-
-		UE_VLOG_BOX(Context.LogOwner, Context.Category, Log, Box, Color, TEXT(""));
-		
-		if (UseDrawDebugHelper() && Context.World)
-		{
-			DrawDebugBox(Context.World, Box.GetCenter(), Box.GetExtent(), Color);
-		}
-	}
 	
-	static void DebugDrawCylinder(const FDebugContext& Context, const FVector& Bottom, const FVector& Top, const FVector::FReal InRadius, const FColor& Color, const FString& Text = FString())
-	{
-		if (!DebugIsSelected(Context.Entity))
-		{
-			return;
-		}
-
-		const float Radius = static_cast<float>(InRadius);
-		UE_VLOG_CYLINDER(Context.LogOwner, Context.Category, Log, Bottom, Top, Radius, Color, TEXT("%s"), *Text);
-
-		if (UseDrawDebugHelper() && Context.World)
-		{
-			DrawDebugCylinder(Context.World, Bottom, Top, Radius, /*Segments = */24, Color);
-		}
-	}
-	
-	static void DebugDrawCircle(const FDebugContext& Context, const FVector& Bottom, const FVector::FReal InRadius, const FColor& Color, const FString& Text = FString())
-	{
-		if (!DebugIsSelected(Context.Entity))
-		{
-			return;
-		}
-
-		const float Radius = static_cast<float>(InRadius);
-		UE_VLOG_CIRCLE_THICK(Context.LogOwner, Context.Category, Log, Bottom, FVector(0,0,1), Radius, Color, /*Thickness*/2, TEXT("%s"), *Text);
-
-		if (UseDrawDebugHelper() && Context.World)
-		{
-			DrawDebugCircle(Context.World, Bottom, Radius, /*Segments = */24, Color);
-		}
-	}
-	//----------------------------------------------------------------------//
-	// End MassDebugUtils
-	//----------------------------------------------------------------------//
-
-
 	// Local debug utils
 	static void DebugDrawVelocity(const FDebugContext& Context, const FVector& Start, const FVector& End, const FColor& Color)
 	{
@@ -531,8 +395,8 @@ void UMassMovingAvoidanceProcessor::Execute(FMassEntityManager& EntityManager, F
 					TEXT("%s"), *Entity.DebugGetDescription(), TEXT("%s"), *Entity.DebugGetDescription());
 
 				// Draw current velocity (black)
-				DebugDrawVelocity(VelocitiesDebugContext, AgentLocation + UE::MassAvoidance::DebugAgentHeightOffset,
-					AgentLocation + UE::MassAvoidance::DebugAgentHeightOffset + AgentVelocity, UE::MassAvoidance::VelocityColor);
+				UE::MassAvoidance::DebugDrawVelocity(VelocitiesDebugContext, AgentLocation + UE::MassAvoidance::DebugAgentHeightOffset,
+				                                     AgentLocation + UE::MassAvoidance::DebugAgentHeightOffset + AgentVelocity, UE::MassAvoidance::VelocityColor);
 
 				// Draw initial steering force
 				DebugDrawArrow(BaseDebugContext, AgentLocation + UE::MassAvoidance::DebugAgentHeightOffset,
@@ -707,7 +571,7 @@ void UMassMovingAvoidanceProcessor::Execute(FMassEntityManager& EntityManager, F
 							DebugDrawLine(ObstacleDebugContext, AgentLocation, HitAgentPos, UE::MassAvoidance::ObstacleAvoidFuturHitPosColor);
 
 							// Draw individual predictive avoidance obstacle force
-							DebugDrawForce(ObstacleDebugContext, HitObPos, HitObPos + AvoidForce, UE::MassAvoidance::ObstacleAvoidForceColor);
+							UE::MassAvoidance::DebugDrawForce(ObstacleDebugContext, HitObPos, HitObPos + AvoidForce, UE::MassAvoidance::ObstacleAvoidForceColor);
 						}
 #endif // WITH_MASSGAMEPLAY_DEBUG
 					}
@@ -716,7 +580,7 @@ void UMassMovingAvoidanceProcessor::Execute(FMassEntityManager& EntityManager, F
 #if WITH_MASSGAMEPLAY_DEBUG && UNSAFE_FOR_MT
 				// Draw total steering force to avoid obstacles
 				const FVector EnvironmentAvoidSteeringForce = SteeringForce - OldSteeringForce;
-				DebugDrawSummedForce(ObstacleDebugContext,
+				UE::MassAvoidance::DebugDrawSummedForce(ObstacleDebugContext,
 					AgentLocation + UE::MassAvoidance::DebugAgentHeightOffset,
 					AgentLocation + UE::MassAvoidance::DebugAgentHeightOffset + EnvironmentAvoidSteeringForce,
 					UE::MassAvoidance::ObstacleAvoidForceColor);
@@ -771,7 +635,7 @@ void UMassMovingAvoidanceProcessor::Execute(FMassEntityManager& EntityManager, F
 							ContactColor = UE::MassAvoidance::ObstacleSeparationForceColor;
 							
 							// Draw individual contact forces
-							DebugDrawForce(ObstacleDebugContext, Contact.Position + ZOffset, Contact.Position + SeparationForce + ZOffset, ContactColor);
+							UE::MassAvoidance::DebugDrawForce(ObstacleDebugContext, Contact.Position + ZOffset, Contact.Position + SeparationForce + ZOffset, ContactColor);
 						}
 						DebugDrawSphere(ObstacleDebugContext, Contact.Position + ZOffset, SphereRadius, ContactColor);
 						DebugDrawArrow(ObstacleDebugContext, Contact.Position + ZOffset, Contact.Position + (Contact.Distance * Contact.Normal) + ZOffset, ContactColor, /*HeadSize=*/ 5.f);
@@ -782,7 +646,7 @@ void UMassMovingAvoidanceProcessor::Execute(FMassEntityManager& EntityManager, F
 #if WITH_MASSGAMEPLAY_DEBUG && UNSAFE_FOR_MT
 				// Draw total steering force to separate from close edges
 				const FVector TotalSeparationForce = SteeringForce - SteeringForceBeforeSeparation;
-				DebugDrawSummedForce(ObstacleDebugContext,
+				UE::MassAvoidance::DebugDrawSummedForce(ObstacleDebugContext,
 					AgentLocation + UE::MassAvoidance::DebugAgentHeightOffset,
 					AgentLocation + UE::MassAvoidance::DebugAgentHeightOffset + TotalSeparationForce,
 					UE::MassAvoidance::ObstacleSeparationForceColor);
@@ -1108,7 +972,6 @@ void UMassMovingAvoidanceProcessor::Execute(FMassEntityManager& EntityManager, F
 //----------------------------------------------------------------------//
 UMassStandingAvoidanceProcessor::UMassStandingAvoidanceProcessor()
 	: EntityQuery(*this)
-
 {
 	bAutoRegisterWithProcessingPhases = true;
 	ExecutionFlags = (int32)EProcessorExecutionFlags::AllNetModes;
diff --git a/Engine/Plugins/AI/MassAI/Source/MassNavigation/Private/MassNavigationDebug.cpp b/Engine/Plugins/AI/MassAI/Source/MassNavigation/Private/MassNavigationDebug.cpp
new file mode 100644
--- /dev/null
+++ b/Engine/Plugins/AI/MassAI/Source/MassNavigation/Private/MassNavigationDebug.cpp
@@ -0,0 +1,142 @@
+﻿// Copyright Epic Games, Inc. All Rights Reserved.
+
+#include "MassNavigationDebug.h"
+
+#if WITH_MASSGAMEPLAY_DEBUG
+
+#include "MassDebugger.h"
+#include "VisualLogger/VisualLogger.h"
+#include "DrawDebugHelpers.h"
+
+namespace UE::MassNavigation::Debug
+{
+	namespace Tweakables
+	{
+		bool bUseDrawDebugHelpers = false;
+	} // Tweakables
+
+	FAutoConsoleVariableRef Vars[] = 
+	{
+		FAutoConsoleVariableRef(TEXT("ai.mass.debug.UseDrawDebugHelpers"), Tweakables::bUseDrawDebugHelpers, TEXT("Use debug draw helpers in addition to visual logs."), ECVF_Cheat)
+	};
+	
+	bool UseDrawDebugHelper()
+	{
+		return Tweakables::bUseDrawDebugHelpers;
+	}
+
+	bool DebugIsSelected(const FMassEntityHandle Entity)
+	{
+#if WITH_MASSENTITY_DEBUG	
+		FColor Color;
+		return UE::Mass::Debug::IsDebuggingEntity(Entity, &Color);
+#else
+		return true;
+#endif		
+	}
+
+	void DebugDrawLine(const FDebugContext& Context, const FVector& Start, const FVector& End, const FColor& Color, const float Thickness /*= 0.f*/, const bool bPersistent /*= false*/)
+	{
+		if (!DebugIsSelected(Context.Entity))
+		{
+			return;
+		}
+
+		UE_VLOG_SEGMENT_THICK(Context.LogOwner, Context.Category, Log, Start, End, Color, (int16)Thickness, TEXT(""));
+
+		if (UseDrawDebugHelper() && Context.World)
+		{
+			DrawDebugLine(Context.World, Start, End, Color, bPersistent, /*LifeTime =*/ -1.f, /*DepthPriority =*/ 0, Thickness);
+		}
+	}
+
+	void DebugDrawArrow(const FDebugContext& Context, const FVector& Start, const FVector& End, const FColor& Color, const float HeadSize /*= 8.f*/, const float Thickness /*= 1.5f*/)
+	{
+		if (!DebugIsSelected(Context.Entity))
+		{
+			return;
+		}
+
+		constexpr FVector::FReal Pointyness = 1.8;
+		const FVector Line = End - Start;
+		const FVector UnitV = Line.GetSafeNormal();
+		const FVector Perp = FVector::CrossProduct(UnitV, FVector::UpVector);
+		const FVector Left = Perp - (Pointyness*UnitV);
+		const FVector Right = -Perp - (Pointyness*UnitV);
+		UE_VLOG_SEGMENT_THICK(Context.LogOwner, Context.Category, Log, Start, End, Color, (int16)Thickness, TEXT(""));
+		UE_VLOG_SEGMENT_THICK(Context.LogOwner, Context.Category, Log, End, End + HeadSize * Left, Color, (int16)Thickness, TEXT(""));
+		UE_VLOG_SEGMENT_THICK(Context.LogOwner, Context.Category, Log, End, End + HeadSize * Right, Color, (int16)Thickness, TEXT(""));
+
+		if (UseDrawDebugHelper() && Context.World)
+		{
+			DrawDebugLine(Context.World, Start, End, Color, /*bPersistent=*/ false, /*LifeTime =*/ -1.f, /*DepthPriority =*/ 0, Thickness);
+			DrawDebugLine(Context.World, End, End + HeadSize * Left, Color, /*bPersistent=*/ false, /*LifeTime =*/ -1.f, /*DepthPriority =*/ 0, Thickness);
+			DrawDebugLine(Context.World, End, End + HeadSize * Right, Color, /*bPersistent=*/ false, /*LifeTime =*/ -1.f, /*DepthPriority =*/ 0, Thickness);
+		}
+	}
+
+	void DebugDrawSphere(const FDebugContext& Context, const FVector& Center, const FVector::FReal InRadius, const FColor& Color)
+	{
+		if (!DebugIsSelected(Context.Entity))
+		{
+			return;
+		}
+
+		const float Radius = static_cast<float>(InRadius);
+		UE_VLOG_LOCATION(Context.LogOwner, Context.Category, Log, Center, Radius, Color, TEXT(""));
+
+		if (UseDrawDebugHelper() && Context.World)
+		{
+			DrawDebugSphere(Context.World, Center, Radius, /*Segments = */16, Color);
+		}
+	}
+
+	void DebugDrawBox(const FDebugContext& Context, const FBox& Box, const FColor& Color)
+	{
+		if (!DebugIsSelected(Context.Entity))
+		{
+			return;
+		}
+
+		UE_VLOG_BOX(Context.LogOwner, Context.Category, Log, Box, Color, TEXT(""));
+		
+		if (UseDrawDebugHelper() && Context.World)
+		{
+			DrawDebugBox(Context.World, Box.GetCenter(), Box.GetExtent(), Color);
+		}
+	}
+	
+	void DebugDrawCylinder(const FDebugContext& Context, const FVector& Bottom, const FVector& Top, const FVector::FReal InRadius, const FColor& Color, const FString& Text /*= FString()*/)
+	{
+		if (!DebugIsSelected(Context.Entity))
+		{
+			return;
+		}
+
+		const float Radius = static_cast<float>(InRadius);
+		UE_VLOG_CYLINDER(Context.LogOwner, Context.Category, Log, Bottom, Top, Radius, Color, TEXT("%s"), *Text);
+
+		if (UseDrawDebugHelper() && Context.World)
+		{
+			DrawDebugCylinder(Context.World, Bottom, Top, Radius, /*Segments = */24, Color);
+		}
+	}
+	
+	void DebugDrawCircle(const FDebugContext& Context, const FVector& Bottom, const FVector::FReal InRadius, const FColor& Color, const FString& Text /*= FString()*/)
+	{
+		if (!DebugIsSelected(Context.Entity))
+		{
+			return;
+		}
+
+		const float Radius = static_cast<float>(InRadius);
+		UE_VLOG_CIRCLE_THICK(Context.LogOwner, Context.Category, Log, Bottom, FVector(0,0,1), Radius, Color, /*Thickness*/2, TEXT("%s"), *Text);
+
+		if (UseDrawDebugHelper() && Context.World)
+		{
+			DrawDebugCircle(Context.World, Bottom, Radius, /*Segments = */24, Color);
+		}
+	}
+} // UE::MassNavigation::Debug
+
+#endif //WITH_MASSGAMEPLAY_DEBUG
diff --git a/Engine/Plugins/AI/MassAI/Source/MassNavigation/Public/MassNavigationDebug.h b/Engine/Plugins/AI/MassAI/Source/MassNavigation/Public/MassNavigationDebug.h
new file mode 100644
--- /dev/null
+++ b/Engine/Plugins/AI/MassAI/Source/MassNavigation/Public/MassNavigationDebug.h
@@ -0,0 +1,46 @@
+﻿// Copyright Epic Games, Inc. All Rights Reserved.
+
+#pragma once
+
+#include "MassCommonTypes.h"
+#include "MassEntityHandle.h"
+
+#if WITH_MASSGAMEPLAY_DEBUG
+
+namespace UE::MassNavigation::Debug
+{
+	MASSNAVIGATION_API struct FDebugContext
+	{
+		FDebugContext(const UObject* InLogOwner, const FLogCategoryBase& InCategory, const UWorld* InWorld, const FMassEntityHandle InEntity)
+			: LogOwner(InLogOwner)
+			, Category(InCategory)
+			, World(InWorld)
+			, Entity(InEntity)
+		{}
+
+		const UObject* LogOwner;
+		const FLogCategoryBase& Category;
+		const UWorld* World;
+		const FMassEntityHandle Entity;
+	};
+	
+	MASSNAVIGATION_API bool UseDrawDebugHelper();
+	MASSNAVIGATION_API bool DebugIsSelected(const FMassEntityHandle Entity);
+	MASSNAVIGATION_API void DebugDrawLine(const FDebugContext& Context, const FVector& Start, const FVector& End, const FColor& Color, const float Thickness = 0.f, const bool bPersistent = false);
+	MASSNAVIGATION_API void DebugDrawArrow(const FDebugContext& Context, const FVector& Start, const FVector& End, const FColor& Color, const float HeadSize = 8.f, const float Thickness = 1.5f);
+	MASSNAVIGATION_API void DebugDrawSphere(const FDebugContext& Context, const FVector& Center, const FVector::FReal InRadius, const FColor& Color);
+	MASSNAVIGATION_API void DebugDrawBox(const FDebugContext& Context, const FBox& Box, const FColor& Color);
+	MASSNAVIGATION_API void DebugDrawCylinder(const FDebugContext& Context, const FVector& Bottom, const FVector& Top, const FVector::FReal InRadius, const FColor& Color, const FString& Text = FString());
+	MASSNAVIGATION_API void DebugDrawCircle(const FDebugContext& Context, const FVector& Bottom, const FVector::FReal InRadius, const FColor& Color, const FString& Text = FString());
+	
+	inline FColor MixColors(const FColor ColorA, const FColor ColorB)
+	{
+		const int32 R = ((int32)ColorA.R + (int32)ColorB.R) / 2;
+		const int32 G = ((int32)ColorA.G + (int32)ColorB.G) / 2;
+		const int32 B = ((int32)ColorA.B + (int32)ColorB.B) / 2;
+		const int32 A = ((int32)ColorA.A + (int32)ColorB.A) / 2;
+		return FColor((uint8)R, (uint8)G, (uint8)B, (uint8)A);
+	}
+} // UE::MassNavigation::Debug
+
+#endif // WITH_MASSGAMEPLAY_DEBUG
diff --git a/Engine/Plugins/AI/MassAI/Source/MassNavigation/Public/MassNavigationUtils.h b/Engine/Plugins/AI/MassAI/Source/MassNavigation/Public/MassNavigationUtils.h
--- a/Engine/Plugins/AI/MassAI/Source/MassNavigation/Public/MassNavigationUtils.h
+++ b/Engine/Plugins/AI/MassAI/Source/MassNavigation/Public/MassNavigationUtils.h
@@ -5,7 +5,6 @@
 #include "CoreTypes.h"
 #include "Math/Quat.h"
 #include "Math/UnrealMathUtility.h"
-#include "MassCommonTypes.h"
 
 namespace UE::MassNavigation
 {
@@ -113,19 +112,3 @@ namespace UE::MassNavigation
 		return Mid;
 	}
 } // UE::MassNavigation
-
-#if WITH_MASSGAMEPLAY_DEBUG
-
-namespace UE::MassNavigation::Debug
-{
-	inline FColor MixColors(const FColor ColorA, const FColor ColorB)
-	{
-		const int32 R = ((int32)ColorA.R + (int32)ColorB.R) / 2;
-		const int32 G = ((int32)ColorA.G + (int32)ColorB.G) / 2;
-		const int32 B = ((int32)ColorA.B + (int32)ColorB.B) / 2;
-		const int32 A = ((int32)ColorA.A + (int32)ColorB.A) / 2;
-		return FColor((uint8)R, (uint8)G, (uint8)B, (uint8)A);
-	}
-}
-
-#endif // WITH_MASSGAMEPLAY_DEBUG
diff --git a/Engine/Plugins/AI/MassAI/Source/MassZoneGraphNavigation/Private/MassZoneGraphNavigationProcessors.cpp b/Engine/Plugins/AI/MassAI/Source/MassZoneGraphNavigation/Private/MassZoneGraphNavigationProcessors.cpp
--- a/Engine/Plugins/AI/MassAI/Source/MassZoneGraphNavigation/Private/MassZoneGraphNavigationProcessors.cpp
+++ b/Engine/Plugins/AI/MassAI/Source/MassZoneGraphNavigation/Private/MassZoneGraphNavigationProcessors.cpp
@@ -15,6 +15,7 @@
 #include "MassSimulationLOD.h"
 #include "Engine/World.h"
 #include "MassDebugger.h"
+#include "MassNavigationDebug.h"
 
 #define UNSAFE_FOR_MT 1
 
