CIS fix

diff --git a/Engine/Plugins/AI/MassAI/Source/MassNavigation/Public/MassNavigationDebug.h b/Engine/Plugins/AI/MassAI/Source/MassNavigation/Public/MassNavigationDebug.h
--- a/Engine/Plugins/AI/MassAI/Source/MassNavigation/Public/MassNavigationDebug.h
+++ b/Engine/Plugins/AI/MassAI/Source/MassNavigation/Public/MassNavigationDebug.h
@@ -9,7 +9,7 @@
 
 namespace UE::MassNavigation::Debug
 {
-	MASSNAVIGATION_API struct FDebugContext
+	struct MASSNAVIGATION_API FDebugContext
 	{
 		FDebugContext(const UObject* InLogOwner, const FLogCategoryBase& InCategory, const UWorld* InWorld, const FMassEntityHandle InEntity)
 			: LogOwner(InLogOwner)
