[MassGameplay] Initializing dynamic ST processors upon creation rather than doing it lazily when the processors get included in the processing graph. This makes sure the processor registers for signals as soon as possible.
The change also fixes functionality related to UE::Mass::StateTree::bDynamicSTProcessorsEnabled and makes the expected use clear via the console variable's comment.

diff --git a/Engine/Plugins/AI/MassAI/Source/MassAIBehavior/Private/MassStateTreeProcessors.cpp b/Engine/Plugins/AI/MassAI/Source/MassAIBehavior/Private/MassStateTreeProcessors.cpp
--- a/Engine/Plugins/AI/MassAI/Source/MassAIBehavior/Private/MassStateTreeProcessors.cpp
+++ b/Engine/Plugins/AI/MassAI/Source/MassAIBehavior/Private/MassStateTreeProcessors.cpp
@@ -233,9 +233,7 @@ UMassStateTreeProcessor::UMassStateTreeProcessor(const FObjectInitializer& Objec
 	ExecutionOrder.ExecuteAfter.Add(UE::Mass::ProcessorGroupNames::SyncWorldToMass);
 	ExecutionOrder.ExecuteBefore.Add(UE::Mass::ProcessorGroupNames::Tasks);
 
-	if (UE::Mass::StateTree::bDynamicSTProcessorsEnabled
-		// we're not setting the values for CDO to keep it "default"
-		&& HasAnyFlags(RF_ClassDefaultObject) == false)
+	if (UE::Mass::StateTree::bDynamicSTProcessorsEnabled)
 	{
 		bAutoRegisterWithProcessingPhases = false;
 		bAllowMultipleInstances = true;
diff --git a/Engine/Plugins/AI/MassAI/Source/MassAIBehavior/Private/MassStateTreeSubsystem.cpp b/Engine/Plugins/AI/MassAI/Source/MassAIBehavior/Private/MassStateTreeSubsystem.cpp
--- a/Engine/Plugins/AI/MassAI/Source/MassAIBehavior/Private/MassStateTreeSubsystem.cpp
+++ b/Engine/Plugins/AI/MassAI/Source/MassAIBehavior/Private/MassStateTreeSubsystem.cpp
@@ -16,7 +16,8 @@ namespace UE::Mass::StateTree
 {
 	bool bDynamicSTProcessorsEnabled = false;
 	FAutoConsoleVariableRef CVarDynamicSTEnabled(TEXT("ai.mass.DynamicSTProcessorsEnabled"), bDynamicSTProcessorsEnabled
-		, TEXT("Whether Dynamic ST processors will be created per distinct ST requirements."), ECVF_Default);
+		, TEXT("Whether Dynamic ST processors will be created per distinct ST requirements. Can only be set via code or ini.")
+		, ECVF_ReadOnly);
 }
 
 void UMassStateTreeSubsystem::Initialize(FSubsystemCollectionBase& Collection)
@@ -150,6 +151,7 @@ void UMassStateTreeSubsystem::CreateProcessorForStateTree(TNotNull<const UStateT
 		DynamicProcessor = NewObject<UMassStateTreeProcessor>(this, DynamicProcessorClass);
 		checkf(DynamicProcessor, TEXT("Failed to spawn an instance of %s"), *GetNameSafe(DynamicProcessorClass));
 		DynamicProcessor->SetExecutionRequirements(FragmentRequirements, SubsystemRequirements);
+		DynamicProcessor->CallInitialize(this, ĘntityManager.ToSharedRef());
 
 		SimulationSubsystem->RegisterDynamicProcessor(*DynamicProcessor);
 
