[Mass] Updated most EntityQuery.ForEachEntityChunk call sites in our processors to utilize FMassExecutionContext::FEntityIterator rather than int counter. This change is required for proper functioning of processor/fragment breakpoints.

diff --git a/Engine/Plugins/AI/MassAI/Source/MassAIBehavior/Private/MassLookAtProcessors.cpp b/Engine/Plugins/AI/MassAI/Source/MassAIBehavior/Private/MassLookAtProcessors.cpp
--- a/Engine/Plugins/AI/MassAI/Source/MassAIBehavior/Private/MassLookAtProcessors.cpp
+++ b/Engine/Plugins/AI/MassAI/Source/MassAIBehavior/Private/MassLookAtProcessors.cpp
@@ -127,7 +127,6 @@ void UMassLookAtProcessor::Execute(FMassEntityManager& EntityManager, FMassExecu
 			const UMassNavigationSubsystem& MassNavSystem = Context.GetSubsystemChecked<UMassNavigationSubsystem>();
 			const UZoneGraphSubsystem& ZoneGraphSubsystem = Context.GetSubsystemChecked<UZoneGraphSubsystem>();
 
-			const int32 NumEntities = Context.GetNumEntities();
 			const TArrayView<FMassLookAtFragment> LookAtList = Context.GetMutableFragmentView<FMassLookAtFragment>();
 			const TConstArrayView<FTransformFragment> TransformList = Context.GetFragmentView<FTransformFragment>();
 			const TConstArrayView<FMassMoveTargetFragment> MoveTargetList = Context.GetFragmentView<FMassMoveTargetFragment>();
@@ -135,10 +134,10 @@ void UMassLookAtProcessor::Execute(FMassEntityManager& EntityManager, FMassExecu
 			const TConstArrayView<FMassZoneGraphShortPathFragment> ShortPathList = Context.GetFragmentView<FMassZoneGraphShortPathFragment>();
 			const TArrayView<FMassLookAtTrajectoryFragment> LookAtTrajectoryList = Context.GetMutableFragmentView<FMassLookAtTrajectoryFragment>();
 
-			for (int32 i = 0; i < NumEntities; ++i)
+			for (FMassExecutionContext::FEntityIterator EntityIt = Context.CreateEntityIterator(); EntityIt; ++EntityIt)
 			{
-				FMassLookAtFragment& LookAt = LookAtList[i];
-				const FTransformFragment& TransformFragment = TransformList[i];
+				FMassLookAtFragment& LookAt = LookAtList[EntityIt];
+				const FTransformFragment& TransformFragment = TransformList[EntityIt];
 
 				const bool bHasLookAtTrajectory =
 					MoveTargetList.Num() > 0
@@ -147,7 +146,7 @@ void UMassLookAtProcessor::Execute(FMassEntityManager& EntityManager, FMassExecu
 					&& ShortPathList.Num() > 0;
 
 				bool bDisplayDebug = false;
-				const FMassEntityHandle Entity = Context.GetEntity(i);
+				const FMassEntityHandle Entity = Context.GetEntity(EntityIt);
 	#if WITH_MASSGAMEPLAY_DEBUG
 				FColor EntityColor = FColor::White;
 				bDisplayDebug = UE::Mass::Debug::IsDebuggingEntity(Entity, &EntityColor);
@@ -176,14 +175,14 @@ void UMassLookAtProcessor::Execute(FMassEntityManager& EntityManager, FMassExecu
 				case EMassLookAtMode::LookAlongPath:
 					if (bHasLookAtTrajectory)
 					{
-						const FMassMoveTargetFragment& MoveTarget = MoveTargetList[i];
-						const FMassZoneGraphLaneLocationFragment& ZoneGraphLocation = ZoneGraphLocationList[i];
-						FMassLookAtTrajectoryFragment& LookAtTrajectory = LookAtTrajectoryList[i];
+						const FMassMoveTargetFragment& MoveTarget = MoveTargetList[EntityIt];
+						const FMassZoneGraphLaneLocationFragment& ZoneGraphLocation = ZoneGraphLocationList[EntityIt];
+						FMassLookAtTrajectoryFragment& LookAtTrajectory = LookAtTrajectoryList[EntityIt];
 
 						if (MoveTarget.GetCurrentActionID() != LookAt.LastSeenActionID)
 						{
-							const FMassZoneGraphLaneLocationFragment& LaneLocation = ZoneGraphLocationList[i];
-							const FMassZoneGraphShortPathFragment& ShortPath = ShortPathList[i];
+							const FMassZoneGraphLaneLocationFragment& LaneLocation = ZoneGraphLocationList[EntityIt];
+							const FMassZoneGraphShortPathFragment& ShortPath = ShortPathList[EntityIt];
 							
 							BuildTrajectory(ZoneGraphSubsystem, LaneLocation, ShortPath, Entity, bDisplayDebug, LookAtTrajectory);
 							LookAt.LastSeenActionID = MoveTarget.GetCurrentActionID();
diff --git a/Engine/Plugins/AI/MassAI/Source/MassAIBehavior/Private/MassStateTreeProcessors.cpp b/Engine/Plugins/AI/MassAI/Source/MassAIBehavior/Private/MassStateTreeProcessors.cpp
--- a/Engine/Plugins/AI/MassAI/Source/MassAIBehavior/Private/MassStateTreeProcessors.cpp
+++ b/Engine/Plugins/AI/MassAI/Source/MassAIBehavior/Private/MassStateTreeProcessors.cpp
@@ -40,10 +40,10 @@ void ForEachEntityInChunk(
 	
 	const UStateTree* StateTree = SharedStateTree.StateTree;
 
-	for (int32 EntityIndex = 0; EntityIndex < NumEntities; EntityIndex++)
+	for (FMassExecutionContext::FEntityIterator EntityIt = Context.CreateEntityIterator(); EntityIt; ++EntityIt)
 	{
-		const FMassEntityHandle Entity = Context.GetEntity(EntityIndex);
-		FMassStateTreeInstanceFragment& StateTreeFragment = StateTreeInstanceList[EntityIndex];
+		const FMassEntityHandle Entity = Context.GetEntity(EntityIt);
+		FMassStateTreeInstanceFragment& StateTreeFragment = StateTreeInstanceList[EntityIt];
 		FStateTreeInstanceData* InstanceData = MassStateTreeSubsystem.GetInstanceData(StateTreeFragment.InstanceHandle);
 		if (InstanceData != nullptr)
 		{
@@ -106,8 +106,6 @@ void UMassStateTreeFragmentDestructor::Execute(FMassEntityManager& EntityManager
 			UMassStateTreeSubsystem& MassStateTreeSubsystem = Context.GetMutableSubsystemChecked<UMassStateTreeSubsystem>();
 			const TArrayView<FMassStateTreeInstanceFragment> StateTreeInstanceList = Context.GetMutableFragmentView<FMassStateTreeInstanceFragment>();
 
-			const int32 NumEntities = Context.GetNumEntities();
-	
 			UE::MassBehavior::ForEachEntityInChunk(Context, MassStateTreeSubsystem,
 				[](FMassStateTreeExecutionContext& StateTreeExecutionContext, FMassStateTreeInstanceFragment&)
 				{
@@ -116,9 +114,9 @@ void UMassStateTreeFragmentDestructor::Execute(FMassEntityManager& EntityManager
 				});
 
 			// Free the StateTree instance memory
-			for (int32 EntityIndex = 0; EntityIndex < NumEntities; EntityIndex++)
+			for (FMassExecutionContext::FEntityIterator EntityIt = Context.CreateEntityIterator(); EntityIt; ++EntityIt)
 			{
-				FMassStateTreeInstanceFragment& StateTreeInstance = StateTreeInstanceList[EntityIndex];
+				FMassStateTreeInstanceFragment& StateTreeInstance = StateTreeInstanceList[EntityIt];
 				if (StateTreeInstance.InstanceHandle.IsValid())
 				{
 					MassStateTreeSubsystem.FreeInstanceData(StateTreeInstance.InstanceHandle);
@@ -150,9 +148,9 @@ void UMassStateTreeActivationProcessor::ConfigureQueries(const TSharedRef<FMassE
 	ProcessorRequirements.AddSubsystemRequirement<UMassSignalSubsystem>(EMassFragmentAccess::ReadWrite);
 }
 
-void UMassStateTreeActivationProcessor::Execute(FMassEntityManager& EntityManager, FMassExecutionContext& Context)
+void UMassStateTreeActivationProcessor::Execute(FMassEntityManager& EntityManager, FMassExecutionContext& ExecutionContext)
 {
-	UMassSignalSubsystem& SignalSubsystem = Context.GetMutableSubsystemChecked<UMassSignalSubsystem>();
+	UMassSignalSubsystem& SignalSubsystem = ExecutionContext.GetMutableSubsystemChecked<UMassSignalSubsystem>();
 
 	const UMassBehaviorSettings* BehaviorSettings = GetDefault<UMassBehaviorSettings>();
 	check(BehaviorSettings);
@@ -164,7 +162,7 @@ void UMassStateTreeActivationProcessor::Execute(FMassEntityManager& EntityManage
 	TArray<FMassEntityHandle> EntitiesToSignal;
 	int32 ActivationCounts[EMassLOD::Max] {0,0,0,0};
 	
-	EntityQuery.ForEachEntityChunk(Context,
+	EntityQuery.ForEachEntityChunk(ExecutionContext,
 		[&EntitiesToSignal, &ActivationCounts, MaxActivationsPerLOD = BehaviorSettings->MaxActivationsPerLOD, TimeInSeconds](FMassExecutionContext& Context)
 		{
 			UMassStateTreeSubsystem& MassStateTreeSubsystem = Context.GetMutableSubsystemChecked<UMassStateTreeSubsystem>();
@@ -181,9 +179,9 @@ void UMassStateTreeActivationProcessor::Execute(FMassEntityManager& EntityManage
 			const FMassStateTreeSharedFragment& SharedStateTree = Context.GetConstSharedFragment<FMassStateTreeSharedFragment>();
 
 			// Allocate and initialize the StateTree instance memory
-			for (int32 EntityIndex = 0; EntityIndex < NumEntities; EntityIndex++)
+			for (FMassExecutionContext::FEntityIterator EntityIt = Context.CreateEntityIterator(); EntityIt; ++EntityIt)
 			{
-				FMassStateTreeInstanceFragment& StateTreeInstance = StateTreeInstanceList[EntityIndex];
+				FMassStateTreeInstanceFragment& StateTreeInstance = StateTreeInstanceList[EntityIt];
 				StateTreeInstance.InstanceHandle = MassStateTreeSubsystem.AllocateInstanceData(SharedStateTree.StateTree);
 			}
 			
@@ -198,12 +196,12 @@ void UMassStateTreeActivationProcessor::Execute(FMassEntityManager& EntityManage
 
 			// Adding a tag on each entities to remember we have sent the state tree initialization signal
 			EntitiesToSignal.Reserve(EntitiesToSignal.Num() + NumEntities);
-			for (int32 EntityIndex = 0; EntityIndex < NumEntities; EntityIndex++)
+			for (FMassExecutionContext::FEntityIterator EntityIt = Context.CreateEntityIterator(); EntityIt; ++EntityIt)
 			{
-				const FMassStateTreeInstanceFragment& StateTreeInstance = StateTreeInstanceList[EntityIndex];
+				const FMassStateTreeInstanceFragment& StateTreeInstance = StateTreeInstanceList[EntityIt];
 				if (StateTreeInstance.InstanceHandle.IsValid())
 				{
-					const FMassEntityHandle Entity = Context.GetEntity(EntityIndex);
+					const FMassEntityHandle Entity = Context.GetEntity(EntityIt);
 					Context.Defer().AddTag<FMassStateTreeActivatedTag>(Entity);
 					EntitiesToSignal.Add(Entity);
 				}
diff --git a/Engine/Plugins/AI/MassAI/Source/MassAIBehavior/Private/MassZoneGraphAnnotationProcessors.cpp b/Engine/Plugins/AI/MassAI/Source/MassAIBehavior/Private/MassZoneGraphAnnotationProcessors.cpp
--- a/Engine/Plugins/AI/MassAI/Source/MassAIBehavior/Private/MassZoneGraphAnnotationProcessors.cpp
+++ b/Engine/Plugins/AI/MassAI/Source/MassAIBehavior/Private/MassZoneGraphAnnotationProcessors.cpp
@@ -35,14 +35,13 @@ void UMassZoneGraphAnnotationTagsInitializer::Execute(FMassEntityManager& Entity
 	{
 		const UZoneGraphAnnotationSubsystem& ZoneGraphAnnotationSubsystem = Context.GetSubsystemChecked<UZoneGraphAnnotationSubsystem>();
 
-		const int32 NumEntities = Context.GetNumEntities();
 		const TArrayView<FMassZoneGraphAnnotationFragment> AnnotationTagsList = Context.GetMutableFragmentView<FMassZoneGraphAnnotationFragment>();
 		const TConstArrayView<FMassZoneGraphLaneLocationFragment> LaneLocationList = Context.GetFragmentView<FMassZoneGraphLaneLocationFragment>();
 
-		for (int32 EntityIndex = 0; EntityIndex < NumEntities; ++EntityIndex)
+		for (FMassExecutionContext::FEntityIterator EntityIt = Context.CreateEntityIterator(); EntityIt; ++EntityIt)
 		{
-			FMassZoneGraphAnnotationFragment& AnnotationTags = AnnotationTagsList[EntityIndex];
-			const FMassZoneGraphLaneLocationFragment& LaneLocation = LaneLocationList[EntityIndex];
+			FMassZoneGraphAnnotationFragment& AnnotationTags = AnnotationTagsList[EntityIt];
+			const FMassZoneGraphLaneLocationFragment& LaneLocation = LaneLocationList[EntityIt];
 
 			if (!LaneLocation.LaneHandle.IsValid())
 			{
@@ -106,16 +105,15 @@ void UMassZoneGraphAnnotationTagUpdateProcessor::Execute(FMassEntityManager& Ent
 		
 		UZoneGraphAnnotationSubsystem& ZoneGraphAnnotationSubsystem = Context.GetMutableSubsystemChecked<UZoneGraphAnnotationSubsystem>();
 
-		const int32 NumEntities = Context.GetNumEntities();
 		const TArrayView<FMassZoneGraphAnnotationFragment> AnnotationTagsList = Context.GetMutableFragmentView<FMassZoneGraphAnnotationFragment>();
 		const TConstArrayView<FMassZoneGraphLaneLocationFragment> LaneLocationList = Context.GetFragmentView<FMassZoneGraphLaneLocationFragment>();
 
-		for (int32 EntityIndex = 0; EntityIndex < NumEntities; ++EntityIndex)
+		for (FMassExecutionContext::FEntityIterator EntityIt = Context.CreateEntityIterator(); EntityIt; ++EntityIt)
 		{
-			FMassZoneGraphAnnotationFragment& AnnotationTags = AnnotationTagsList[EntityIndex];
-			const FMassZoneGraphLaneLocationFragment& LaneLocation = LaneLocationList[EntityIndex];
+			FMassZoneGraphAnnotationFragment& AnnotationTags = AnnotationTagsList[EntityIt];
+			const FMassZoneGraphLaneLocationFragment& LaneLocation = LaneLocationList[EntityIt];
 
-			UpdateAnnotationTags(ZoneGraphAnnotationSubsystem, AnnotationTags, LaneLocation, Context.GetEntity(EntityIndex));
+			UpdateAnnotationTags(ZoneGraphAnnotationSubsystem, AnnotationTags, LaneLocation, Context.GetEntity(EntityIt));
 		}
 	});
 
@@ -151,16 +149,15 @@ void UMassZoneGraphAnnotationTagUpdateProcessor::SignalEntities(FMassEntityManag
 	{
 		UZoneGraphAnnotationSubsystem& ZoneGraphAnnotationSubsystem = Context.GetMutableSubsystemChecked<UZoneGraphAnnotationSubsystem>();
 
-		const int32 NumEntities = Context.GetNumEntities();
 		const TConstArrayView<FMassZoneGraphLaneLocationFragment> LaneLocationList = Context.GetFragmentView<FMassZoneGraphLaneLocationFragment>();
 		const TArrayView<FMassZoneGraphAnnotationFragment> AnnotationTagsList = Context.GetMutableFragmentView<FMassZoneGraphAnnotationFragment>();
 
-		for (int32 EntityIndex = 0; EntityIndex < NumEntities; ++EntityIndex)
+		for (FMassExecutionContext::FEntityIterator EntityIt = Context.CreateEntityIterator(); EntityIt; ++EntityIt)
 		{
-			FMassZoneGraphAnnotationFragment& AnnotationTags = AnnotationTagsList[EntityIndex];
-			const FMassZoneGraphLaneLocationFragment& LaneLocation = LaneLocationList[EntityIndex];
+			FMassZoneGraphAnnotationFragment& AnnotationTags = AnnotationTagsList[EntityIt];
+			const FMassZoneGraphLaneLocationFragment& LaneLocation = LaneLocationList[EntityIt];
 
-			UpdateAnnotationTags(ZoneGraphAnnotationSubsystem, AnnotationTags, LaneLocation, Context.GetEntity(EntityIndex));
+			UpdateAnnotationTags(ZoneGraphAnnotationSubsystem, AnnotationTags, LaneLocation, Context.GetEntity(EntityIt));
 		}
 	});
 }
diff --git a/Engine/Plugins/AI/MassAI/Source/MassAIDebug/Private/MassDebugStateTreeProcessor.cpp b/Engine/Plugins/AI/MassAI/Source/MassAIDebug/Private/MassDebugStateTreeProcessor.cpp
--- a/Engine/Plugins/AI/MassAI/Source/MassAIDebug/Private/MassDebugStateTreeProcessor.cpp
+++ b/Engine/Plugins/AI/MassAI/Source/MassAIDebug/Private/MassDebugStateTreeProcessor.cpp
@@ -58,7 +58,6 @@ void UMassDebugStateTreeProcessor::Execute(FMassEntityManager& EntityManager, FM
 	EntityQuery.ForEachEntityChunk(Context, [this, Debugger, MassStateTreeSubsystem](FMassExecutionContext& Context)
 	{
 		const FMassEntityHandle SelectedEntity = Debugger->GetSelectedEntity();
-		const int32 NumEntities = Context.GetNumEntities();
 		const TConstArrayView<FMassStateTreeInstanceFragment> StateTreeInstanceList = Context.GetFragmentView<FMassStateTreeInstanceFragment>();
 		const FMassStateTreeSharedFragment& SharedStateTree = Context.GetConstSharedFragment<FMassStateTreeSharedFragment>();
 		const TConstArrayView<FTransformFragment> TransformList = Context.GetFragmentView<FTransformFragment>();
@@ -71,16 +70,16 @@ void UMassDebugStateTreeProcessor::Execute(FMassEntityManager& EntityManager, FM
 			return;
 		}
 	
-		for (int32 i = 0; i < NumEntities; ++i)
+		for (FMassExecutionContext::FEntityIterator EntityIt = Context.CreateEntityIterator(); EntityIt; ++EntityIt)
 		{
-			const FMassEntityHandle Entity = Context.GetEntity(i);
+			const FMassEntityHandle Entity = Context.GetEntity(EntityIt);
 
 			if (Entity != SelectedEntity && !UE::Mass::Debug::IsDebuggingEntity(Entity))
 			{
 				continue;
 			}
 			
-			const FMassStateTreeInstanceFragment& StateTreeInstance = StateTreeInstanceList[i];
+			const FMassStateTreeInstanceFragment& StateTreeInstance = StateTreeInstanceList[EntityIt];
 
 			FStateTreeInstanceData* InstanceData = MassStateTreeSubsystem->GetInstanceData(StateTreeInstance.InstanceHandle);
 			if (InstanceData == nullptr)
@@ -102,7 +101,7 @@ void UMassDebugStateTreeProcessor::Execute(FMassEntityManager& EntityManager, FM
 			const bool bDisplayDebug = UE::Mass::Debug::IsDebuggingEntity(Entity, &EntityColor);
 			if (bDisplayDebug)
 			{
-				const FTransformFragment& Transform = TransformList[i];
+				const FTransformFragment& Transform = TransformList[EntityIt];
 				
 				const FVector ZOffset(0,0,50);
 				const FVector Position = Transform.GetTransform().GetLocation() + ZOffset;
diff --git a/Engine/Plugins/AI/MassAI/Source/MassNavMeshNavigation/Private/MassNavMeshNavigationBoundaryProcessor.cpp b/Engine/Plugins/AI/MassAI/Source/MassNavMeshNavigation/Private/MassNavMeshNavigationBoundaryProcessor.cpp
--- a/Engine/Plugins/AI/MassAI/Source/MassNavMeshNavigation/Private/MassNavMeshNavigationBoundaryProcessor.cpp
+++ b/Engine/Plugins/AI/MassAI/Source/MassNavMeshNavigation/Private/MassNavMeshNavigationBoundaryProcessor.cpp
@@ -35,22 +35,20 @@ void UMassNavMeshNavigationBoundaryProcessor::Execute(FMassEntityManager& Entity
 
 	EntityQuery.ForEachEntityChunk(Context, [this](FMassExecutionContext& Context)
 	{
-		const int32 NumEntities = Context.GetNumEntities();
-
 		TConstArrayView<FMassNavMeshShortPathFragment> ShortPathList = Context.GetFragmentView<FMassNavMeshShortPathFragment>();
 		TConstArrayView<FMassMoveTargetFragment> MovementTargetList = Context.GetFragmentView<FMassMoveTargetFragment>();
 
 		TArrayView<FMassNavMeshBoundaryFragment> NavmeshBoundaryList = Context.GetMutableFragmentView<FMassNavMeshBoundaryFragment>();
 		TArrayView<FMassNavigationEdgesFragment> EdgesList = Context.GetMutableFragmentView<FMassNavigationEdgesFragment>();
 
-		for (int32 EntityIndex = 0; EntityIndex < NumEntities; ++EntityIndex)
+		for (FMassExecutionContext::FEntityIterator EntityIt = Context.CreateEntityIterator(); EntityIt; ++EntityIt)
 		{
-			const FMassNavMeshShortPathFragment& ShortPath = ShortPathList[EntityIndex];
-			const FMassMoveTargetFragment& MovementTarget = MovementTargetList[EntityIndex];
-			FMassNavMeshBoundaryFragment& NavmeshBoundary = NavmeshBoundaryList[EntityIndex];
-			FMassNavigationEdgesFragment& Edges = EdgesList[EntityIndex];
+			const FMassNavMeshShortPathFragment& ShortPath = ShortPathList[EntityIt];
+			const FMassMoveTargetFragment& MovementTarget = MovementTargetList[EntityIt];
+			FMassNavMeshBoundaryFragment& NavmeshBoundary = NavmeshBoundaryList[EntityIt];
+			FMassNavigationEdgesFragment& Edges = EdgesList[EntityIt];
 
-			const FMassEntityHandle Entity = Context.GetEntity(EntityIndex);
+			const FMassEntityHandle Entity = Context.GetEntity(EntityIt);
 
 			// First check if we moved enough for an update.
 			const FVector::FReal DeltaDistSquared = FVector::DistSquared(MovementTarget.Center, NavmeshBoundary.LastUpdatePosition);
diff --git a/Engine/Plugins/AI/MassAI/Source/MassNavMeshNavigation/Private/MassNavMeshNavigationProcessors.cpp b/Engine/Plugins/AI/MassAI/Source/MassNavMeshNavigation/Private/MassNavMeshNavigationProcessors.cpp
--- a/Engine/Plugins/AI/MassAI/Source/MassNavMeshNavigation/Private/MassNavMeshNavigationProcessors.cpp
+++ b/Engine/Plugins/AI/MassAI/Source/MassNavMeshNavigation/Private/MassNavMeshNavigationProcessors.cpp
@@ -48,7 +48,6 @@ void UMassNavMeshPathFollowProcessor::Execute(FMassEntityManager& EntityManager,
 
 	EntityQuery_Conditional.ForEachEntityChunk(Context, [this, &EntitiesToSignalPathDone](FMassExecutionContext& Context)
 	{
-		const int32 NumEntities = Context.GetNumEntities();
 		const TArrayView<FMassNavMeshShortPathFragment> ShortPathList = Context.GetMutableFragmentView<FMassNavMeshShortPathFragment>();
 		const TArrayView<FMassMoveTargetFragment> MoveTargetList = Context.GetMutableFragmentView<FMassMoveTargetFragment>();
 
@@ -59,11 +58,11 @@ void UMassNavMeshPathFollowProcessor::Execute(FMassEntityManager& EntityManager,
 		const TConstArrayView<FMassSimulationVariableTickFragment> SimVariableTickList = Context.GetFragmentView<FMassSimulationVariableTickFragment>();
 		const float WorldDeltaTime = Context.GetDeltaTimeSeconds();
 
-		for (int32 EntityIndex = 0; EntityIndex < NumEntities; ++EntityIndex)
+		for (FMassExecutionContext::FEntityIterator EntityIt = Context.CreateEntityIterator(); EntityIt; ++EntityIt)
 		{
-			FMassNavMeshShortPathFragment& ShortPath = ShortPathList[EntityIndex];
-			FMassMoveTargetFragment& MoveTarget = MoveTargetList[EntityIndex];
-			const FMassEntityHandle Entity = Context.GetEntity(EntityIndex);
+			FMassNavMeshShortPathFragment& ShortPath = ShortPathList[EntityIt];
+			FMassMoveTargetFragment& MoveTarget = MoveTargetList[EntityIt];
+			const FMassEntityHandle Entity = Context.GetEntity(EntityIt);
 
 			bool bDisplayDebug = false;
 #if WITH_MASSGAMEPLAY_DEBUG && UNSAFE_FOR_MT // this will result in bDisplayDebug == false and disabling of all the vlogs below
@@ -78,7 +77,7 @@ void UMassNavMeshPathFollowProcessor::Execute(FMassEntityManager& EntityManager,
 				const bool bWasDone = ShortPath.IsDone();
 
 				// Note: this should be in sync with the logic in apply velocity.
-				const bool bHasSteering = (bHasLOD == false) || (SimLODList[EntityIndex].LOD != EMassLOD::Off);
+				const bool bHasSteering = (bHasLOD == false) || (SimLODList[EntityIt].LOD != EMassLOD::Off);
 
 				if (!bHasSteering || !MoveTarget.bSteeringFallingBehind)
 				{
diff --git a/Engine/Plugins/AI/MassAI/Source/MassNavigation/Private/MassAvoidanceProcessors.cpp b/Engine/Plugins/AI/MassAI/Source/MassNavigation/Private/MassAvoidanceProcessors.cpp
--- a/Engine/Plugins/AI/MassAI/Source/MassNavigation/Private/MassAvoidanceProcessors.cpp
+++ b/Engine/Plugins/AI/MassAI/Source/MassNavigation/Private/MassAvoidanceProcessors.cpp
@@ -274,7 +274,6 @@ void UMassMovingAvoidanceProcessor::Execute(FMassEntityManager& EntityManager, F
 	{
 		const float DeltaTime = Context.GetDeltaTimeSeconds();
 		const double CurrentTime = World->GetTimeSeconds();
-		const int32 NumEntities = Context.GetNumEntities();
 		
 		const TArrayView<FMassForceFragment> ForceList = Context.GetMutableFragmentView<FMassForceFragment>();
 		const TConstArrayView<FMassNavigationEdgesFragment> NavEdgesList = Context.GetFragmentView<FMassNavigationEdgesFragment>();
@@ -322,21 +321,21 @@ void UMassMovingAvoidanceProcessor::Execute(FMassEntityManager& EntityManager, F
 		};
 		TArray<FCollider, TInlineAllocator<16>> Colliders;
 
-		for (int32 EntityIndex = 0; EntityIndex < NumEntities; ++EntityIndex)
+		for (FMassExecutionContext::FEntityIterator EntityIt = Context.CreateEntityIterator(); EntityIt; ++EntityIt)
 		{
 			// @todo: this check should eventually be part of the query (i.e. only handle moving agents).
-			const FMassMoveTargetFragment& MoveTarget = MoveTargetList[EntityIndex];
+			const FMassMoveTargetFragment& MoveTarget = MoveTargetList[EntityIt];
 			if (MoveTarget.GetCurrentAction() == EMassMovementAction::Animate || MoveTarget.GetCurrentAction() == EMassMovementAction::Stand)
 			{
 				continue;
 			}
 
-			FMassEntityHandle Entity = Context.GetEntity(EntityIndex);
-			const FMassNavigationEdgesFragment& NavEdges = NavEdgesList[EntityIndex];
-			const FTransformFragment& Location = LocationList[EntityIndex];
-			const FMassVelocityFragment& Velocity = VelocityList[EntityIndex];
-			const FAgentRadiusFragment& RadiusFragement = RadiusList[EntityIndex];
-			FMassForceFragment& Force = ForceList[EntityIndex];
+			FMassEntityHandle Entity = Context.GetEntity(EntityIt);
+			const FMassNavigationEdgesFragment& NavEdges = NavEdgesList[EntityIt];
+			const FTransformFragment& Location = LocationList[EntityIt];
+			const FMassVelocityFragment& Velocity = VelocityList[EntityIt];
+			const FAgentRadiusFragment& RadiusFragment = RadiusList[EntityIt];
+			FMassForceFragment& Force = ForceList[EntityIt];
 
 			// Smaller steering max accel makes the steering more "calm" but less opportunistic, may not find solution, or gets stuck.
 			// Max contact accel should be quite a big bigger than steering so that collision response is firm. 
@@ -346,7 +345,7 @@ void UMassMovingAvoidanceProcessor::Execute(FMassEntityManager& EntityManager, F
 			const FVector AgentLocation = Location.GetTransform().GetTranslation();
 			const FVector AgentVelocity = FVector(Velocity.Value.X, Velocity.Value.Y, 0.);
 			
-			const FVector::FReal AgentRadius = RadiusFragement.Radius;
+			const FVector::FReal AgentRadius = RadiusFragment.Radius;
 			
 			FVector SteeringForce = Force.Value;
 
@@ -416,8 +415,8 @@ void UMassMovingAvoidanceProcessor::Execute(FMassEntityManager& EntityManager, F
 			
 			if (!MoveTarget.bOffBoundaries && UE::MassAvoidance::Tweakables::bEnableEnvironmentAvoidance)
 			{
-				const FVector::FReal EnvironmentSeparationAgentRadius = (RadiusFragement.Radius * MovingAvoidanceParams.SeparationRadiusScale) - (NavEdges.bExtrudedEdges ? AgentRadius : 0);
-				const FVector::FReal EnvironmentPredictiveAvoidanceAgentRadius = (RadiusFragement.Radius * MovingAvoidanceParams.PredictiveAvoidanceRadiusScale) - (NavEdges.bExtrudedEdges ? AgentRadius : 0);
+				const FVector::FReal EnvironmentSeparationAgentRadius = (RadiusFragment.Radius * MovingAvoidanceParams.SeparationRadiusScale) - (NavEdges.bExtrudedEdges ? AgentRadius : 0);
+				const FVector::FReal EnvironmentPredictiveAvoidanceAgentRadius = (RadiusFragment.Radius * MovingAvoidanceParams.PredictiveAvoidanceRadiusScale) - (NavEdges.bExtrudedEdges ? AgentRadius : 0);
 				
 				const FVector DesiredAcceleration = UE::MassNavigation::ClampVector(SteeringForce, MaxSteerAccel);
 				const FVector DesiredVelocity = UE::MassNavigation::ClampVector(AgentVelocity + DesiredAcceleration * DeltaTime, MaximumSpeed);
@@ -668,8 +667,8 @@ void UMassMovingAvoidanceProcessor::Execute(FMassEntityManager& EntityManager, F
 
 			//////////////////////////////////////////////////////////////////////////
 			// Avoid close agents
-			const FVector::FReal SeparationAgentRadius = RadiusFragement.Radius * MovingAvoidanceParams.SeparationRadiusScale;
-			const FVector::FReal PredictiveAvoidanceAgentRadius = RadiusFragement.Radius * MovingAvoidanceParams.PredictiveAvoidanceRadiusScale;
+			const FVector::FReal SeparationAgentRadius = RadiusFragment.Radius * MovingAvoidanceParams.SeparationRadiusScale;
+			const FVector::FReal PredictiveAvoidanceAgentRadius = RadiusFragment.Radius * MovingAvoidanceParams.PredictiveAvoidanceRadiusScale;
 
 			// Update desired velocity based on avoidance so far.
 			const FVector DesAcc = UE::MassNavigation::ClampVector(SteeringForce, MaxSteerAccel);
@@ -1012,7 +1011,6 @@ void UMassStandingAvoidanceProcessor::Execute(FMassEntityManager& EntityManager,
 	// Avoidance while standing
 	EntityQuery.ForEachEntityChunk(Context, [this, &EntityManager](FMassExecutionContext& Context)
 	{
-		const int32 NumEntities = Context.GetNumEntities();
 		const float DeltaTime = Context.GetDeltaTimeSeconds();
 
 		const TArrayView<FMassGhostLocationFragment> GhostList = Context.GetMutableFragmentView<FMassGhostLocationFragment>();
@@ -1047,26 +1045,26 @@ void UMassStandingAvoidanceProcessor::Execute(FMassEntityManager& EntityManager,
 		};
 		TArray<FSortedObstacle, TFixedAllocator<UE::MassAvoidance::MaxObstacleResults>> ClosestObstacles;
 
-		for (int32 EntityIndex = 0; EntityIndex < NumEntities; ++EntityIndex)
+		for (FMassExecutionContext::FEntityIterator EntityIt = Context.CreateEntityIterator(); EntityIt; ++EntityIt)
 		{
 			// @todo: this check should eventually be part of the query.
-			const FMassMoveTargetFragment& MoveTarget = MoveTargetList[EntityIndex];
+			const FMassMoveTargetFragment& MoveTarget = MoveTargetList[EntityIt];
 			if (MoveTarget.GetCurrentAction() != EMassMovementAction::Stand)
 			{
 				continue;
 			}
 			
-			FMassGhostLocationFragment& Ghost = GhostList[EntityIndex];
+			FMassGhostLocationFragment& Ghost = GhostList[EntityIt];
 			// Skip if the ghost is not valid for this movement action yet.
 			if (Ghost.IsValid(MoveTarget.GetCurrentActionID()) == false)
 			{
 				continue;
 			}
 
-			const FTransformFragment& Location = LocationList[EntityIndex];
-			const FAgentRadiusFragment& Radius = RadiusList[EntityIndex];
+			const FTransformFragment& Location = LocationList[EntityIt];
+			const FAgentRadiusFragment& Radius = RadiusList[EntityIt];
 
-			FMassEntityHandle Entity = Context.GetEntity(EntityIndex);
+			FMassEntityHandle Entity = Context.GetEntity(EntityIt);
 			const FVector AgentLocation = Location.GetTransform().GetTranslation();
 			const FVector::FReal AgentRadius = Radius.Radius;
 
diff --git a/Engine/Plugins/AI/MassAI/Source/MassNavigation/Private/MassNavigationProcessors.cpp b/Engine/Plugins/AI/MassAI/Source/MassNavigation/Private/MassNavigationProcessors.cpp
--- a/Engine/Plugins/AI/MassAI/Source/MassNavigation/Private/MassNavigationProcessors.cpp
+++ b/Engine/Plugins/AI/MassAI/Source/MassNavigation/Private/MassNavigationProcessors.cpp
@@ -49,15 +49,14 @@ void UMassOffLODNavigationProcessor::Execute(FMassEntityManager& EntityManager,
 			return;
 		}
 #endif // WITH_MASSGAMEPLAY_DEBUG
-		const int32 NumEntities = Context.GetNumEntities();
 
 		const TArrayView<FTransformFragment> LocationList = Context.GetMutableFragmentView<FTransformFragment>();
 		const TConstArrayView<FMassMoveTargetFragment> MoveTargetList = Context.GetFragmentView<FMassMoveTargetFragment>();
 
-		for (int32 EntityIndex = 0; EntityIndex < NumEntities; ++EntityIndex)
+		for (FMassExecutionContext::FEntityIterator EntityIt = Context.CreateEntityIterator(); EntityIt; ++EntityIt)
 		{
-			FTransform& CurrentTransform = LocationList[EntityIndex].GetMutableTransform();
-			const FMassMoveTargetFragment& MoveTarget = MoveTargetList[EntityIndex];
+			FTransform& CurrentTransform = LocationList[EntityIt].GetMutableTransform();
+			const FMassMoveTargetFragment& MoveTarget = MoveTargetList[EntityIt];
 
 			// Snap position to move target directly
 			CurrentTransform.SetLocation(MoveTarget.Center);
@@ -96,17 +95,16 @@ void UMassNavigationSmoothHeightProcessor::Execute(FMassEntityManager& EntityMan
 			return;
 		}
 #endif // WITH_MASSGAMEPLAY_DEBUG
-		const int32 NumEntities = Context.GetNumEntities();
 		const float DeltaTime = Context.GetDeltaTimeSeconds();
 
 		const FMassMovementParameters& MovementParams = Context.GetConstSharedFragment<FMassMovementParameters>();
 		const TArrayView<FTransformFragment> LocationList = Context.GetMutableFragmentView<FTransformFragment>();
 		const TConstArrayView<FMassMoveTargetFragment> MoveTargetList = Context.GetFragmentView<FMassMoveTargetFragment>();
 
-		for (int32 EntityIndex = 0; EntityIndex < NumEntities; ++EntityIndex)
+		for (FMassExecutionContext::FEntityIterator EntityIt = Context.CreateEntityIterator(); EntityIt; ++EntityIt)
 		{
-			FTransform& CurrentTransform = LocationList[EntityIndex].GetMutableTransform();
-			const FMassMoveTargetFragment& MoveTarget = MoveTargetList[EntityIndex];
+			FTransform& CurrentTransform = LocationList[EntityIt].GetMutableTransform();
+			const FMassMoveTargetFragment& MoveTarget = MoveTargetList[EntityIt];
 
 			if (MoveTarget.GetCurrentAction() == EMassMovementAction::Move || MoveTarget.GetCurrentAction() == EMassMovementAction::Stand)
 			{
@@ -140,14 +138,13 @@ void UMassMoveTargetFragmentInitializer::Execute(FMassEntityManager& EntityManag
 {
 	InitializerQuery.ForEachEntityChunk(Context, [](FMassExecutionContext& Context)
 	{
-		const int32 NumEntities = Context.GetNumEntities();
 		const TArrayView<FMassMoveTargetFragment> MoveTargetList = Context.GetMutableFragmentView<FMassMoveTargetFragment>();
 		const TConstArrayView<FTransformFragment> LocationList = Context.GetFragmentView<FTransformFragment>();
 
-		for (int32 EntityIndex = 0; EntityIndex < NumEntities; ++EntityIndex)
+		for (FMassExecutionContext::FEntityIterator EntityIt = Context.CreateEntityIterator(); EntityIt; ++EntityIt)
 		{
-			FMassMoveTargetFragment& MoveTarget = MoveTargetList[EntityIndex];
-			const FTransformFragment& Location = LocationList[EntityIndex];
+			FMassMoveTargetFragment& MoveTarget = MoveTargetList[EntityIt];
+			const FTransformFragment& Location = LocationList[EntityIt];
 
 			MoveTarget.Center = Location.GetTransform().GetLocation();
 			MoveTarget.Forward = Location.GetTransform().GetRotation().Vector();
@@ -199,25 +196,24 @@ void UMassNavigationObstacleGridProcessor::Execute(FMassEntityManager& EntityMan
 	AddToGridEntityQuery.ForEachEntityChunk(Context, [this, &EntityManager](FMassExecutionContext& Context)
 	{
 		FNavigationObstacleHashGrid2D& HashGrid = Context.GetMutableSubsystemChecked<UMassNavigationSubsystem>().GetObstacleGridMutable();
-		const int32 NumEntities = Context.GetNumEntities();
 
 		TConstArrayView<FTransformFragment> LocationList = Context.GetFragmentView<FTransformFragment>();
 		TConstArrayView<FAgentRadiusFragment> RadiiList = Context.GetFragmentView<FAgentRadiusFragment>();
 		TArrayView<FMassNavigationObstacleGridCellLocationFragment> NavigationObstacleCellLocationList = Context.GetMutableFragmentView<FMassNavigationObstacleGridCellLocationFragment>();
 		const bool bHasColliderData = Context.GetFragmentView<FMassAvoidanceColliderFragment>().Num() > 0;
 
-		for (int32 EntityIndex = 0; EntityIndex < NumEntities; ++EntityIndex)
+		for (FMassExecutionContext::FEntityIterator EntityIt = Context.CreateEntityIterator(); EntityIt; ++EntityIt)
 		{
 			// Add to the grid
-			const FVector NewPos = LocationList[EntityIndex].GetTransform().GetLocation();
-			const float Radius = RadiiList[EntityIndex].Radius;
+			const FVector NewPos = LocationList[EntityIt].GetTransform().GetLocation();
+			const float Radius = RadiiList[EntityIt].Radius;
 
 			FMassNavigationObstacleItem ObstacleItem;
-			ObstacleItem.Entity = Context.GetEntity(EntityIndex);
+			ObstacleItem.Entity = Context.GetEntity(EntityIt);
 			ObstacleItem.ItemFlags |= bHasColliderData ? EMassNavigationObstacleFlags::HasColliderData : EMassNavigationObstacleFlags::None;
 			
 			const FBox NewBounds(NewPos - FVector(Radius, Radius, 0.f), NewPos + FVector(Radius, Radius, 0.f));
-			NavigationObstacleCellLocationList[EntityIndex].CellLoc = HashGrid.Add(ObstacleItem, NewBounds);
+			NavigationObstacleCellLocationList[EntityIt].CellLoc = HashGrid.Add(ObstacleItem, NewBounds);
 
 			Context.Defer().AddTag<FMassInNavigationObstacleGridTag>(ObstacleItem.Entity);
 		}
@@ -226,30 +222,29 @@ void UMassNavigationObstacleGridProcessor::Execute(FMassEntityManager& EntityMan
 	UpdateGridEntityQuery.ForEachEntityChunk(Context, [this, &EntityManager](FMassExecutionContext& Context)
 	{
 		FNavigationObstacleHashGrid2D& HashGrid = Context.GetMutableSubsystemChecked<UMassNavigationSubsystem>().GetObstacleGridMutable();
-		const int32 NumEntities = Context.GetNumEntities();
 
 		TConstArrayView<FTransformFragment> LocationList = Context.GetFragmentView<FTransformFragment>();
 		TConstArrayView<FAgentRadiusFragment> RadiiList = Context.GetFragmentView<FAgentRadiusFragment>();
 		TArrayView<FMassNavigationObstacleGridCellLocationFragment> NavigationObstacleCellLocationList = Context.GetMutableFragmentView<FMassNavigationObstacleGridCellLocationFragment>();
 		const bool bHasColliderData = Context.GetFragmentView<FMassAvoidanceColliderFragment>().Num() > 0;
 
-		for (int32 EntityIndex = 0; EntityIndex < NumEntities; ++EntityIndex)
+		for (FMassExecutionContext::FEntityIterator EntityIt = Context.CreateEntityIterator(); EntityIt; ++EntityIt)
 		{
 			// Update position in grid
-			const FVector NewPos = LocationList[EntityIndex].GetTransform().GetLocation();
-			const float Radius = RadiiList[EntityIndex].Radius;
+			const FVector NewPos = LocationList[EntityIt].GetTransform().GetLocation();
+			const float Radius = RadiiList[EntityIt].Radius;
 			FMassNavigationObstacleItem ObstacleItem;
-			ObstacleItem.Entity = Context.GetEntity(EntityIndex);
+			ObstacleItem.Entity = Context.GetEntity(EntityIt);
 			ObstacleItem.ItemFlags |= bHasColliderData ? EMassNavigationObstacleFlags::HasColliderData : EMassNavigationObstacleFlags::None;
 
 			const FBox NewBounds(NewPos - FVector(Radius, Radius, 0.f), NewPos + FVector(Radius, Radius, 0.f));
-			NavigationObstacleCellLocationList[EntityIndex].CellLoc = HashGrid.Move(ObstacleItem, NavigationObstacleCellLocationList[EntityIndex].CellLoc, NewBounds);
+			NavigationObstacleCellLocationList[EntityIt].CellLoc = HashGrid.Move(ObstacleItem, NavigationObstacleCellLocationList[EntityIt].CellLoc, NewBounds);
 
 #if WITH_MASSGAMEPLAY_DEBUG && 0
 			const FDebugContext BaseDebugContext(this, LogAvoidance, nullptr, ObstacleItem.Entity);
 			if (DebugIsSelected(ObstacleItem.Entity))
 			{
-				FBox Box = MovementSubsystem->GetGridMutable().CalcCellBounds(AvoidanceObstacleCellLocationList[EntityIndex].CellLoc);
+				FBox Box = MovementSubsystem->GetGridMutable().CalcCellBounds(AvoidanceObstacleCellLocationList[EntityIt].CellLoc);
 				Box.Max.Z += 200.f;
 				DebugDrawBox(BaseDebugContext, Box, FColor::Yellow);
 			}
@@ -260,15 +255,14 @@ void UMassNavigationObstacleGridProcessor::Execute(FMassEntityManager& EntityMan
 	RemoveFromGridEntityQuery.ForEachEntityChunk(Context, [this, &EntityManager](FMassExecutionContext& Context)
 	{
 		FNavigationObstacleHashGrid2D& HashGrid = Context.GetMutableSubsystemChecked<UMassNavigationSubsystem>().GetObstacleGridMutable();
-		const int32 NumEntities = Context.GetNumEntities();
 
 		TArrayView<FMassNavigationObstacleGridCellLocationFragment> AvoidanceObstacleCellLocationList = Context.GetMutableFragmentView<FMassNavigationObstacleGridCellLocationFragment>();
-		for (int32 EntityIndex = 0; EntityIndex < NumEntities; ++EntityIndex)
+		for (FMassExecutionContext::FEntityIterator EntityIt = Context.CreateEntityIterator(); EntityIt; ++EntityIt)
 		{
 			FMassNavigationObstacleItem ObstacleItem;
-			ObstacleItem.Entity = Context.GetEntity(EntityIndex);
-			HashGrid.Remove(ObstacleItem, AvoidanceObstacleCellLocationList[EntityIndex].CellLoc);
-			AvoidanceObstacleCellLocationList[EntityIndex].CellLoc = FNavigationObstacleHashGrid2D::FCellLocation();
+			ObstacleItem.Entity = Context.GetEntity(EntityIt);
+			HashGrid.Remove(ObstacleItem, AvoidanceObstacleCellLocationList[EntityIt].CellLoc);
+			AvoidanceObstacleCellLocationList[EntityIt].CellLoc = FNavigationObstacleHashGrid2D::FCellLocation();
 
 			Context.Defer().RemoveTag<FMassInNavigationObstacleGridTag>(ObstacleItem.Entity);
 		}
@@ -297,14 +291,13 @@ void UMassNavigationObstacleRemoverProcessor::Execute(FMassEntityManager& Entity
 	EntityQuery.ForEachEntityChunk(Context, [this](FMassExecutionContext& Context)
 	{
 		FNavigationObstacleHashGrid2D& HashGrid = Context.GetMutableSubsystemChecked<UMassNavigationSubsystem>().GetObstacleGridMutable();
-		const int32 NumEntities = Context.GetNumEntities();
 		const TArrayView<FMassNavigationObstacleGridCellLocationFragment> AvoidanceObstacleCellLocationList = Context.GetMutableFragmentView<FMassNavigationObstacleGridCellLocationFragment>();
 
-		for (int32 i = 0; i < NumEntities; ++i)
+		for (FMassExecutionContext::FEntityIterator EntityIt = Context.CreateEntityIterator(); EntityIt; ++EntityIt)
 		{
 			FMassNavigationObstacleItem ObstacleItem;
-			ObstacleItem.Entity = Context.GetEntity(i);
-			HashGrid.Remove(ObstacleItem, AvoidanceObstacleCellLocationList[i].CellLoc);
+			ObstacleItem.Entity = Context.GetEntity(EntityIt);
+			HashGrid.Remove(ObstacleItem, AvoidanceObstacleCellLocationList[EntityIt].CellLoc);
 		}
 	});
 }
diff --git a/Engine/Plugins/AI/MassAI/Source/MassNavigation/Private/MassSmoothOrientationProcessors.cpp b/Engine/Plugins/AI/MassAI/Source/MassNavigation/Private/MassSmoothOrientationProcessors.cpp
--- a/Engine/Plugins/AI/MassAI/Source/MassNavigation/Private/MassSmoothOrientationProcessors.cpp
+++ b/Engine/Plugins/AI/MassAI/Source/MassNavigation/Private/MassSmoothOrientationProcessors.cpp
@@ -51,17 +51,15 @@ void UMassSmoothOrientationProcessor::Execute(FMassEntityManager& EntityManager,
 
 		HighResEntityQuery.ForEachEntityChunk(Context, [this, DeltaTime](FMassExecutionContext& Context)
 		{
-			const int32 NumEntities = Context.GetNumEntities();
-
 			const FMassSmoothOrientationParameters& OrientationParams = Context.GetConstSharedFragment<FMassSmoothOrientationParameters>();
 
 			const TConstArrayView<FMassMoveTargetFragment> MoveTargetList = Context.GetFragmentView<FMassMoveTargetFragment>();
 			const TArrayView<FTransformFragment> LocationList = Context.GetMutableFragmentView<FTransformFragment>();
 			const TArrayView<FMassVelocityFragment> VelocityList = Context.GetMutableFragmentView<FMassVelocityFragment>();
 
-			for (int32 EntityIndex = 0; EntityIndex < NumEntities; ++EntityIndex)
+			for (FMassExecutionContext::FEntityIterator EntityIt = Context.CreateEntityIterator(); EntityIt; ++EntityIt)
 			{
-				const FMassMoveTargetFragment& MoveTarget = MoveTargetList[EntityIndex];
+				const FMassMoveTargetFragment& MoveTarget = MoveTargetList[EntityIt];
 
 				// Do not touch transform at all when animating
 				if (MoveTarget.GetCurrentAction() == EMassMovementAction::Animate)
@@ -69,8 +67,8 @@ void UMassSmoothOrientationProcessor::Execute(FMassEntityManager& EntityManager,
 					continue;
 				}
 
-				const FMassVelocityFragment& CurrentVelocity = VelocityList[EntityIndex];
-				FTransform& CurrentTransform = LocationList[EntityIndex].GetMutableTransform();
+				const FMassVelocityFragment& CurrentVelocity = VelocityList[EntityIt];
+				FTransform& CurrentTransform = LocationList[EntityIt].GetMutableTransform();
 				const FVector CurrentForward = CurrentTransform.GetRotation().GetForwardVector();
 				const FVector::FReal CurrentHeading = UE::MassNavigation::GetYawFromDirection(CurrentForward);
 
@@ -120,15 +118,13 @@ void UMassSmoothOrientationProcessor::Execute(FMassEntityManager& EntityManager,
 
 		LowResEntityQuery_Conditional.ForEachEntityChunk(Context, [this](FMassExecutionContext& Context)
 		{
-			const int32 NumEntities = Context.GetNumEntities();
-
 			const TArrayView<FTransformFragment> LocationList = Context.GetMutableFragmentView<FTransformFragment>();
 			const TConstArrayView<FMassMoveTargetFragment> MoveTargetList = Context.GetFragmentView<FMassMoveTargetFragment>();
 
-			for (int32 EntityIndex = 0; EntityIndex < NumEntities; ++EntityIndex)
+			for (FMassExecutionContext::FEntityIterator EntityIt = Context.CreateEntityIterator(); EntityIt; ++EntityIt)
 			{
-				FTransform& CurrentTransform = LocationList[EntityIndex].GetMutableTransform();
-				const FMassMoveTargetFragment& MoveTarget = MoveTargetList[EntityIndex];
+				FTransform& CurrentTransform = LocationList[EntityIt].GetMutableTransform();
+				const FMassMoveTargetFragment& MoveTarget = MoveTargetList[EntityIt];
 
 				// Snap position to move target directly
 				CurrentTransform.SetRotation(FQuat::FindBetweenNormals(FVector::ForwardVector, MoveTarget.Forward));
diff --git a/Engine/Plugins/AI/MassAI/Source/MassNavigation/Private/MassSteeringProcessors.cpp b/Engine/Plugins/AI/MassAI/Source/MassNavigation/Private/MassSteeringProcessors.cpp
--- a/Engine/Plugins/AI/MassAI/Source/MassNavigation/Private/MassSteeringProcessors.cpp
+++ b/Engine/Plugins/AI/MassAI/Source/MassNavigation/Private/MassSteeringProcessors.cpp
@@ -80,7 +80,6 @@ void UMassSteerToMoveTargetProcessor::Execute(FMassEntityManager& EntityManager,
 {
 	EntityQuery.ForEachEntityChunk(Context, [this](FMassExecutionContext& Context)
 	{
-		const int32 NumEntities = Context.GetNumEntities();
 		const TConstArrayView<FTransformFragment> TransformList = Context.GetFragmentView<FTransformFragment>();
 		const TArrayView<FMassMoveTargetFragment> MoveTargetList = Context.GetMutableFragmentView<FMassMoveTargetFragment>();
 		const TArrayView<FMassVelocityFragment> VelocityList = Context.GetMutableFragmentView<FMassVelocityFragment>();
@@ -95,16 +94,16 @@ void UMassSteerToMoveTargetProcessor::Execute(FMassEntityManager& EntityManager,
 		const FVector::FReal SteerK = 1. / MovingSteeringParams.ReactionTime;
 		const float DeltaTime = Context.GetDeltaTimeSeconds();
 
-		for (int32 EntityIndex = 0; EntityIndex < NumEntities; ++EntityIndex)
+		for (FMassExecutionContext::FEntityIterator EntityIt = Context.CreateEntityIterator(); EntityIt; ++EntityIt)
 		{
-			const FTransformFragment& TransformFragment = TransformList[EntityIndex];
-			FMassSteeringFragment& Steering = SteeringList[EntityIndex];
-			FMassStandingSteeringFragment& StandingSteering = StandingSteeringList[EntityIndex];
-			FMassGhostLocationFragment& Ghost = GhostList[EntityIndex];
-			FMassMoveTargetFragment& MoveTarget = MoveTargetList[EntityIndex];
-			FMassForceFragment& Force = ForceList[EntityIndex];
-			FMassVelocityFragment& Velocity = VelocityList[EntityIndex];
-			const FMassEntityHandle Entity = Context.GetEntity(EntityIndex);
+			const FTransformFragment& TransformFragment = TransformList[EntityIt];
+			FMassSteeringFragment& Steering = SteeringList[EntityIt];
+			FMassStandingSteeringFragment& StandingSteering = StandingSteeringList[EntityIt];
+			FMassGhostLocationFragment& Ghost = GhostList[EntityIt];
+			FMassMoveTargetFragment& MoveTarget = MoveTargetList[EntityIt];
+			FMassForceFragment& Force = ForceList[EntityIt];
+			FMassVelocityFragment& Velocity = VelocityList[EntityIt];
+			const FMassEntityHandle Entity = Context.GetEntity(EntityIt);
 
 			const FTransform& Transform = TransformFragment.GetTransform();;
 
diff --git a/Engine/Plugins/AI/MassAI/Source/MassZoneGraphNavigation/Private/MassZoneGraphNavigationProcessors.cpp b/Engine/Plugins/AI/MassAI/Source/MassZoneGraphNavigation/Private/MassZoneGraphNavigationProcessors.cpp
--- a/Engine/Plugins/AI/MassAI/Source/MassZoneGraphNavigation/Private/MassZoneGraphNavigationProcessors.cpp
+++ b/Engine/Plugins/AI/MassAI/Source/MassZoneGraphNavigation/Private/MassZoneGraphNavigationProcessors.cpp
@@ -45,18 +45,17 @@ void UMassZoneGraphLocationInitializer::Execute(FMassEntityManager& EntityManage
 	{
 		const UZoneGraphSubsystem& ZoneGraphSubsystem = Context.GetSubsystemChecked<UZoneGraphSubsystem>();
 
-		const int32 NumEntities = Context.GetNumEntities();
 		const TArrayView<FMassZoneGraphLaneLocationFragment> LaneLocationList = Context.GetMutableFragmentView<FMassZoneGraphLaneLocationFragment>();
 		const TArrayView<FMassMoveTargetFragment> MoveTargetList = Context.GetMutableFragmentView<FMassMoveTargetFragment>();
 		const TConstArrayView<FTransformFragment> TransformList = Context.GetFragmentView<FTransformFragment>();
 		const FMassZoneGraphNavigationParameters& NavigationParams = Context.GetConstSharedFragment<FMassZoneGraphNavigationParameters>();
 
-		for (int32 EntityIndex = 0; EntityIndex < NumEntities; ++EntityIndex)
+		for (FMassExecutionContext::FEntityIterator EntityIt = Context.CreateEntityIterator(); EntityIt; ++EntityIt)
 		{
-			const FTransformFragment& Transform = TransformList[EntityIndex];
+			const FTransformFragment& Transform = TransformList[EntityIt];
 			const FVector& AgentLocation = Transform.GetTransform().GetLocation();
-			FMassMoveTargetFragment& MoveTarget = MoveTargetList[EntityIndex];
-			FMassZoneGraphLaneLocationFragment& LaneLocation = LaneLocationList[EntityIndex];
+			FMassMoveTargetFragment& MoveTarget = MoveTargetList[EntityIt];
+			FMassZoneGraphLaneLocationFragment& LaneLocation = LaneLocationList[EntityIt];
 
 			const FVector QuerySize(NavigationParams.QueryRadius);
 			const FBox QueryBounds(AgentLocation - QuerySize, AgentLocation + QuerySize);
@@ -139,7 +138,6 @@ void UMassZoneGraphPathFollowProcessor::Execute(FMassEntityManager& EntityManage
 	{
 		const UZoneGraphSubsystem& ZoneGraphSubsystem = Context.GetSubsystemChecked<UZoneGraphSubsystem>();
 
-		const int32 NumEntities = Context.GetNumEntities();
 		const TArrayView<FMassZoneGraphShortPathFragment> ShortPathList = Context.GetMutableFragmentView<FMassZoneGraphShortPathFragment>();
 		const TArrayView<FMassZoneGraphLaneLocationFragment> LaneLocationList = Context.GetMutableFragmentView<FMassZoneGraphLaneLocationFragment>();
 		const TArrayView<FMassMoveTargetFragment> MoveTargetList = Context.GetMutableFragmentView<FMassMoveTargetFragment>();
@@ -149,13 +147,13 @@ void UMassZoneGraphPathFollowProcessor::Execute(FMassEntityManager& EntityManage
 		const bool bHasVariableTick = (SimVariableTickList.Num() > 0);
 		const float WorldDeltaTime = Context.GetDeltaTimeSeconds();
 
-		for (int32 EntityIndex = 0; EntityIndex < NumEntities; ++EntityIndex)
+		for (FMassExecutionContext::FEntityIterator EntityIt = Context.CreateEntityIterator(); EntityIt; ++EntityIt)
 		{
-			FMassZoneGraphShortPathFragment& ShortPath = ShortPathList[EntityIndex];
-			FMassZoneGraphLaneLocationFragment& LaneLocation = LaneLocationList[EntityIndex];
-			FMassMoveTargetFragment& MoveTarget = MoveTargetList[EntityIndex];
-			const FMassEntityHandle Entity = Context.GetEntity(EntityIndex);
-			const float DeltaTime = bHasVariableTick ? SimVariableTickList[EntityIndex].DeltaTime : WorldDeltaTime;
+			FMassZoneGraphShortPathFragment& ShortPath = ShortPathList[EntityIt];
+			FMassZoneGraphLaneLocationFragment& LaneLocation = LaneLocationList[EntityIt];
+			FMassMoveTargetFragment& MoveTarget = MoveTargetList[EntityIt];
+			const FMassEntityHandle Entity = Context.GetEntity(EntityIt);
+			const float DeltaTime = bHasVariableTick ? SimVariableTickList[EntityIt].DeltaTime : WorldDeltaTime;
 
 			bool bDisplayDebug = false;
 #if WITH_MASSGAMEPLAY_DEBUG && UNSAFE_FOR_MT // this will result in bDisplayDebug == false and disabling of all the vlogs below
@@ -172,7 +170,7 @@ void UMassZoneGraphPathFollowProcessor::Execute(FMassEntityManager& EntityManage
 				const bool bWasDone = ShortPath.IsDone();
 
 				// Note: this should be in sync with the logic in apply velocity.
-				const bool bHasSteering = (bHasLOD == false) || (SimLODList[EntityIndex].LOD != EMassLOD::Off);
+				const bool bHasSteering = (bHasLOD == false) || (SimLODList[EntityIt].LOD != EMassLOD::Off);
 
 				if (!bHasSteering || !MoveTarget.bSteeringFallingBehind)
 				{
@@ -436,8 +434,6 @@ void UMassZoneGraphLaneCacheBoundaryProcessor::Execute(FMassEntityManager& Entit
 
 	EntityQuery.ForEachEntityChunk(Context, [this, World](FMassExecutionContext& Context)
 	{
-		const int32 NumEntities = Context.GetNumEntities();
-
 		TConstArrayView<FMassZoneGraphCachedLaneFragment> CachedLaneList = Context.GetFragmentView<FMassZoneGraphCachedLaneFragment>();
 		TConstArrayView<FMassZoneGraphLaneLocationFragment> LaneLocationList = Context.GetFragmentView<FMassZoneGraphLaneLocationFragment>();
 		TConstArrayView<FMassMoveTargetFragment> MovementTargetList = Context.GetFragmentView<FMassMoveTargetFragment>();
@@ -447,14 +443,14 @@ void UMassZoneGraphLaneCacheBoundaryProcessor::Execute(FMassEntityManager& Entit
 		TArray<FZoneGraphLinkedLane> LinkedLanes;
 		LinkedLanes.Reserve(4);	
 	
-		for (int32 EntityIndex = 0; EntityIndex < NumEntities; ++EntityIndex)
+		for (FMassExecutionContext::FEntityIterator EntityIt = Context.CreateEntityIterator(); EntityIt; ++EntityIt)
 		{
-			const FMassZoneGraphCachedLaneFragment& CachedLane = CachedLaneList[EntityIndex];
-			const FMassZoneGraphLaneLocationFragment& LaneLocation = LaneLocationList[EntityIndex];
-			const FMassMoveTargetFragment& MovementTarget = MovementTargetList[EntityIndex];
-			FMassNavigationEdgesFragment& Edges = EdgesList[EntityIndex];
-			FMassLaneCacheBoundaryFragment& LaneCacheBoundary = LaneCacheBoundaryList[EntityIndex];
-			const FMassEntityHandle Entity = Context.GetEntity(EntityIndex);
+			const FMassZoneGraphCachedLaneFragment& CachedLane = CachedLaneList[EntityIt];
+			const FMassZoneGraphLaneLocationFragment& LaneLocation = LaneLocationList[EntityIt];
+			const FMassMoveTargetFragment& MovementTarget = MovementTargetList[EntityIt];
+			FMassNavigationEdgesFragment& Edges = EdgesList[EntityIt];
+			FMassLaneCacheBoundaryFragment& LaneCacheBoundary = LaneCacheBoundaryList[EntityIt];
+			const FMassEntityHandle Entity = Context.GetEntity(EntityIt);
 
 			// First check if we moved enough for an update
 			const FVector::FReal DeltaDistSquared = FVector::DistSquared(MovementTarget.Center, LaneCacheBoundary.LastUpdatePosition);
diff --git a/Engine/Plugins/AI/MassCrowd/Source/MassCrowd/Private/MassCrowdNavigationProcessor.cpp b/Engine/Plugins/AI/MassCrowd/Source/MassCrowd/Private/MassCrowdNavigationProcessor.cpp
--- a/Engine/Plugins/AI/MassCrowd/Source/MassCrowd/Private/MassCrowdNavigationProcessor.cpp
+++ b/Engine/Plugins/AI/MassCrowd/Source/MassCrowd/Private/MassCrowdNavigationProcessor.cpp
@@ -44,17 +44,16 @@ void UMassCrowdLaneTrackingSignalProcessor::SignalEntities(FMassEntityManager& E
 	EntityQuery.ForEachEntityChunk(Context, [this](FMassExecutionContext& Context)
 	{
 		UMassCrowdSubsystem& MassCrowdSubsystem = Context.GetMutableSubsystemChecked<UMassCrowdSubsystem>();
-		const int32 NumEntities = Context.GetNumEntities();
 		const TConstArrayView<FMassZoneGraphLaneLocationFragment> LaneLocationList = Context.GetFragmentView<FMassZoneGraphLaneLocationFragment>();
 		const TArrayView<FMassCrowdLaneTrackingFragment> LaneTrackingList = Context.GetMutableFragmentView<FMassCrowdLaneTrackingFragment>();
 
-		for (int32 EntityIndex = 0; EntityIndex < NumEntities; ++EntityIndex)
+		for (FMassExecutionContext::FEntityIterator EntityIt = Context.CreateEntityIterator(); EntityIt; ++EntityIt)
 		{
-			const FMassZoneGraphLaneLocationFragment& LaneLocation = LaneLocationList[EntityIndex];
-			FMassCrowdLaneTrackingFragment& LaneTracking = LaneTrackingList[EntityIndex];
+			const FMassZoneGraphLaneLocationFragment& LaneLocation = LaneLocationList[EntityIt];
+			FMassCrowdLaneTrackingFragment& LaneTracking = LaneTrackingList[EntityIt];
 			if (LaneTracking.TrackedLaneHandle != LaneLocation.LaneHandle)
 			{
-				MassCrowdSubsystem.OnEntityLaneChanged(Context.GetEntity(EntityIndex), LaneTracking.TrackedLaneHandle, LaneLocation.LaneHandle);
+				MassCrowdSubsystem.OnEntityLaneChanged(Context.GetEntity(EntityIt), LaneTracking.TrackedLaneHandle, LaneLocation.LaneHandle);
 				LaneTracking.TrackedLaneHandle = LaneLocation.LaneHandle;
 			}
 		}
@@ -84,15 +83,14 @@ void UMassCrowdLaneTrackingDestructor::Execute(FMassEntityManager& EntityManager
 	EntityQuery.ForEachEntityChunk(Context, [this](FMassExecutionContext& Context)
 	{
 		UMassCrowdSubsystem& MassCrowdSubsystem = Context.GetMutableSubsystemChecked<UMassCrowdSubsystem>();
-		const int32 NumEntities = Context.GetNumEntities();
 		const TConstArrayView<FMassCrowdLaneTrackingFragment> LaneTrackingList = Context.GetFragmentView<FMassCrowdLaneTrackingFragment>();
 
-		for (int32 EntityIndex = 0; EntityIndex < NumEntities; ++EntityIndex)
+		for (FMassExecutionContext::FEntityIterator EntityIt = Context.CreateEntityIterator(); EntityIt; ++EntityIt)
 		{
-			const FMassCrowdLaneTrackingFragment& LaneTracking = LaneTrackingList[EntityIndex];
+			const FMassCrowdLaneTrackingFragment& LaneTracking = LaneTrackingList[EntityIt];
 			if (LaneTracking.TrackedLaneHandle.IsValid())
 			{
-				MassCrowdSubsystem.OnEntityLaneChanged(Context.GetEntity(EntityIndex), LaneTracking.TrackedLaneHandle, FZoneGraphLaneHandle());
+				MassCrowdSubsystem.OnEntityLaneChanged(Context.GetEntity(EntityIt), LaneTracking.TrackedLaneHandle, FZoneGraphLaneHandle());
 			}
 		}
 	});
@@ -137,8 +135,6 @@ void UMassCrowdDynamicObstacleProcessor::Execute(FMassEntityManager& EntityManag
 
 	EntityQuery_Conditional.ForEachEntityChunk(Context, [this, World, CrowdSettings](FMassExecutionContext& Context)
 	{
-		const int32 NumEntities = Context.GetNumEntities();
-
 		const TConstArrayView<FTransformFragment> LocationList = Context.GetFragmentView<FTransformFragment>();
 		const TConstArrayView<FMassVelocityFragment> VelocityList = Context.GetFragmentView<FMassVelocityFragment>();
 		const TConstArrayView<FAgentRadiusFragment> RadiusList = Context.GetFragmentView<FAgentRadiusFragment>();
@@ -154,20 +150,20 @@ void UMassCrowdDynamicObstacleProcessor::Execute(FMassEntityManager& EntityManag
 		const float ObstacleTimeToStop = CrowdSettings->ObstacleTimeToStop;
 		const float ObstacleEffectRadius = CrowdSettings->ObstacleEffectRadius;
 		
-		for (int32 EntityIndex = 0; EntityIndex < NumEntities; ++EntityIndex)
+		for (FMassExecutionContext::FEntityIterator EntityIt = Context.CreateEntityIterator(); EntityIt; ++EntityIt)
 		{
 			// @todo: limit update frequency, this does not need to occur every frame
-			const FVector Position = LocationList[EntityIndex].GetTransform().GetLocation();
-			const float Radius = RadiusList[EntityIndex].Radius;
-			FMassCrowdObstacleFragment& Obstacle = ObstacleDataList[EntityIndex];
-			const float DeltaTime = FMath::Max(KINDA_SMALL_NUMBER, bHasVariableTick ? SimVariableTickList[EntityIndex].DeltaTime : WorldDeltaTime);
+			const FVector Position = LocationList[EntityIt].GetTransform().GetLocation();
+			const float Radius = RadiusList[EntityIt].Radius;
+			FMassCrowdObstacleFragment& Obstacle = ObstacleDataList[EntityIt];
+			const float DeltaTime = FMath::Max(KINDA_SMALL_NUMBER, bHasVariableTick ? SimVariableTickList[EntityIt].DeltaTime : WorldDeltaTime);
 
 			UE_VLOG_LOCATION(this, LogMassNavigationObstacle, Display, Position, Radius, Obstacle.bIsMoving ? FColor::Green : FColor::Red, TEXT(""));
 
 			if (Obstacle.bIsMoving)
 			{
 				// Calculate current speed based on velocity or last known position.
-				const FVector::FReal CurrentSpeed = bHasVelocity ? VelocityList[EntityIndex].Value.Length() : (FVector::Dist(Position, Obstacle.LastPosition) / DeltaTime);
+				const FVector::FReal CurrentSpeed = bHasVelocity ? VelocityList[EntityIt].Value.Length() : (FVector::Dist(Position, Obstacle.LastPosition) / DeltaTime);
 
 				// Update position while moving, the stop logic will use the last position while check if the obstacles moves again.
 				Obstacle.LastPosition = Position;
@@ -246,14 +242,13 @@ void UMassCrowdDynamicObstacleInitializer::Execute(FMassEntityManager& EntityMan
 	
 	EntityQuery.ForEachEntityChunk(Context, [World](FMassExecutionContext& Context)
 	{
-		const int32 NumEntities = Context.GetNumEntities();
 		const TConstArrayView<FTransformFragment> LocationList = Context.GetFragmentView<FTransformFragment>();
 		const TArrayView<FMassCrowdObstacleFragment> ObstacleDataList = Context.GetMutableFragmentView<FMassCrowdObstacleFragment>();
 
-		for (int32 EntityIndex = 0; EntityIndex < NumEntities; ++EntityIndex)
+		for (FMassExecutionContext::FEntityIterator EntityIt = Context.CreateEntityIterator(); EntityIt; ++EntityIt)
 		{
-			const FVector Position = LocationList[EntityIndex].GetTransform().GetLocation();
-			FMassCrowdObstacleFragment& Obstacle = ObstacleDataList[EntityIndex];
+			const FVector Position = LocationList[EntityIt].GetTransform().GetLocation();
+			FMassCrowdObstacleFragment& Obstacle = ObstacleDataList[EntityIt];
 
 			Obstacle.LastPosition = Position;
 			Obstacle.TimeSinceStopped = 0.0f;
@@ -293,12 +288,11 @@ void UMassCrowdDynamicObstacleDeinitializer::Execute(FMassEntityManager& EntityM
 {
 	EntityQuery.ForEachEntityChunk(Context, [this](FMassExecutionContext& Context)
 	{
-		const int32 NumEntities = Context.GetNumEntities();
 		const TArrayView<FMassCrowdObstacleFragment> ObstacleDataList = Context.GetMutableFragmentView<FMassCrowdObstacleFragment>();
 
-		for (int32 EntityIndex = 0; EntityIndex < NumEntities; ++EntityIndex)
+		for (FMassExecutionContext::FEntityIterator EntityIt = Context.CreateEntityIterator(); EntityIt; ++EntityIt)
 		{
-			FMassCrowdObstacleFragment& Obstacle = ObstacleDataList[EntityIndex];
+			FMassCrowdObstacleFragment& Obstacle = ObstacleDataList[EntityIt];
 
 			if (Obstacle.LaneObstacleID.IsValid())
 			{
diff --git a/Engine/Plugins/AI/MassCrowd/Source/MassCrowd/Private/MassCrowdVisualizationLODProcessor.cpp b/Engine/Plugins/AI/MassCrowd/Source/MassCrowd/Private/MassCrowdVisualizationLODProcessor.cpp
--- a/Engine/Plugins/AI/MassCrowd/Source/MassCrowd/Private/MassCrowdVisualizationLODProcessor.cpp
+++ b/Engine/Plugins/AI/MassCrowd/Source/MassCrowd/Private/MassCrowdVisualizationLODProcessor.cpp
@@ -79,20 +79,19 @@ void UMassCrowdVisualizationLODProcessor::Execute(FMassEntityManager& EntityMana
 	{
 		TRACE_CPUPROFILER_EVENT_SCOPE(ShowISMUnderSpecifiedRange)
 
-		DebugEntityQuery.ForEachEntityChunk(Context, [World](const FMassExecutionContext& Context)
+		DebugEntityQuery.ForEachEntityChunk(Context, [World](FMassExecutionContext& Context)
 		{
 			const TConstArrayView<FTransformFragment> LocationList = Context.GetFragmentView<FTransformFragment>();
 			const TConstArrayView<FMassRepresentationFragment> RepresentationFragmentList = Context.GetFragmentView<FMassRepresentationFragment>();
 			const TConstArrayView<FMassViewerInfoFragment> LODInfoFragmentList = Context.GetFragmentView<FMassViewerInfoFragment>();
-			const int32 NumEntities = Context.GetNumEntities();
 			const float SpecifiedRangeSquaredCentimeters = static_cast<float>(FMath::Square(UE::MassCrowd::bDebugShowISMUnderSpecifiedRange * 100));
-			for (int EntityIdx = 0; EntityIdx < NumEntities; EntityIdx++)
+			for (FMassExecutionContext::FEntityIterator EntityIt = Context.CreateEntityIterator(); EntityIt; ++EntityIt)
 			{
-				const FMassRepresentationFragment& RepresentationFragment = RepresentationFragmentList[EntityIdx];
-				const FMassViewerInfoFragment& LODInfoFragment = LODInfoFragmentList[EntityIdx];
+				const FMassRepresentationFragment& RepresentationFragment = RepresentationFragmentList[EntityIt];
+				const FMassViewerInfoFragment& LODInfoFragment = LODInfoFragmentList[EntityIt];
 				if (RepresentationFragment.CurrentRepresentation == EMassRepresentationType::StaticMeshInstance && SpecifiedRangeSquaredCentimeters > LODInfoFragment.ClosestViewerDistanceSq)
 				{
-					const FTransformFragment& EntityLocation = LocationList[EntityIdx];
+					const FTransformFragment& EntityLocation = LocationList[EntityIt];
 					DrawDebugSolidBox(World, EntityLocation.GetTransform().GetLocation() + FVector(0.0f, 0.0f, 150.0f), FVector(50.0f), FColor::Red);
 				}
 			}
diff --git a/Engine/Plugins/AI/MassCrowd/Source/MassCrowd/Private/MassCrowdVisualizationProcessor.cpp b/Engine/Plugins/AI/MassCrowd/Source/MassCrowd/Private/MassCrowdVisualizationProcessor.cpp
--- a/Engine/Plugins/AI/MassCrowd/Source/MassCrowd/Private/MassCrowdVisualizationProcessor.cpp
+++ b/Engine/Plugins/AI/MassCrowd/Source/MassCrowd/Private/MassCrowdVisualizationProcessor.cpp
@@ -80,18 +80,17 @@ void UMassDebugCrowdVisualizationProcessor::Execute(FMassEntityManager& EntityMa
 
 		TRACE_CPUPROFILER_EVENT_SCOPE(DebugDisplayVisualType)
 
-		EntityQuery.ForEachEntityChunk(Context, [this](const FMassExecutionContext& Context)
+		EntityQuery.ForEachEntityChunk(Context, [this](FMassExecutionContext& Context)
 		{
 			const TConstArrayView<FMassRepresentationFragment> VisualizationList = Context.GetFragmentView<FMassRepresentationFragment>();
 			const TConstArrayView<FMassActorFragment> ActorList = Context.GetFragmentView<FMassActorFragment>();
 			const TConstArrayView<FTransformFragment> EntityLocationList = Context.GetFragmentView<FTransformFragment>();
 
-			const int32 NumEntities = Context.GetNumEntities();
-			for (int EntityIdx = 0; EntityIdx < NumEntities; EntityIdx++)
+			for (FMassExecutionContext::FEntityIterator EntityIt = Context.CreateEntityIterator(); EntityIt; ++EntityIt)
 			{
-				const FTransformFragment& EntityLocation = EntityLocationList[EntityIdx];
-				const FMassRepresentationFragment& Visualization = VisualizationList[EntityIdx];
-				const FMassActorFragment& ActorInfo = ActorList[EntityIdx];
+				const FTransformFragment& EntityLocation = EntityLocationList[EntityIt];
+				const FMassRepresentationFragment& Visualization = VisualizationList[EntityIt];
+				const FMassActorFragment& ActorInfo = ActorList[EntityIt];
 				const int32 RepresentationTypeIdx = (int32)Visualization.CurrentRepresentation;
 				// Show replicated actors
 				if (ActorInfo.IsValid() && !ActorInfo.IsOwnedByMass())
diff --git a/Engine/Plugins/Runtime/InstancedActors/Source/InstancedActors/Private/InstancedActorsData.cpp b/Engine/Plugins/Runtime/InstancedActors/Source/InstancedActors/Private/InstancedActorsData.cpp
--- a/Engine/Plugins/Runtime/InstancedActors/Source/InstancedActors/Private/InstancedActorsData.cpp
+++ b/Engine/Plugins/Runtime/InstancedActors/Source/InstancedActors/Private/InstancedActorsData.cpp
@@ -467,16 +467,15 @@ void UInstancedActorsData::DespawnEntities()
 	{
 		InstancedActorLocationQuery.ForEachEntityChunk(Collection, ExecutionContext, [this, &WorldToLocalTranslation, &ManagerTransform](FMassExecutionContext& Context)
 			{
-				int32 NumEntities = Context.GetNumEntities();
 				TConstArrayView<FInstancedActorsFragment> InstancedActorFragments = Context.GetFragmentView<FInstancedActorsFragment>();
 				TConstArrayView<FTransformFragment> TransformsList = Context.GetFragmentView<FTransformFragment>();
 				check(TransformsList.GetTypeSize() == sizeof(FTransform));
 
 				// Re-build InstanceTransforms, being careful to put transforms back into the right index it was created from
-				for (int32 Index = 0; Index < NumEntities; ++Index)
+				for (FMassExecutionContext::FEntityIterator EntityIt = Context.CreateEntityIterator(); EntityIt; ++EntityIt)
 				{
-					FInstancedActorsInstanceIndex InstanceIndex = InstancedActorFragments[Index].InstanceIndex;
-					InstanceTransforms[InstanceIndex.GetIndex()] = TransformsList[Index].GetTransform();
+					FInstancedActorsInstanceIndex InstanceIndex = InstancedActorFragments[EntityIt].InstanceIndex;
+					InstanceTransforms[InstanceIndex.GetIndex()] = TransformsList[EntityIt].GetTransform();
 
 					checkf(UE::InstancedActors::Helpers::IsValidInstanceTransform(InstanceTransforms[InstanceIndex.GetIndex()]), TEXT("Found Mass entity with unexpected Scale 0 'invalid' transform"));
 
diff --git a/Engine/Plugins/Runtime/InstancedActors/Source/InstancedActors/Private/InstancedActorsDebugProcessor.cpp b/Engine/Plugins/Runtime/InstancedActors/Source/InstancedActors/Private/InstancedActorsDebugProcessor.cpp
--- a/Engine/Plugins/Runtime/InstancedActors/Source/InstancedActors/Private/InstancedActorsDebugProcessor.cpp
+++ b/Engine/Plugins/Runtime/InstancedActors/Source/InstancedActors/Private/InstancedActorsDebugProcessor.cpp
@@ -125,22 +125,20 @@ namespace UE::InstancedActors::Debug
 					FMassExecutionContext Context = EntityManager.CreateExecutionContext(0);
 					DetailedLODEntityQuery.ForEachEntityChunk(Context, [World = InWorld](FMassExecutionContext& Context)
 						{
-							//const FMassEntityHandle SelectedEntity = Debugger->GetSelectedEntity();
-							const int32 NumEntities = Context.GetNumEntities();
 							const TConstArrayView<FMassActorFragment> ActorFragmentsList = Context.GetFragmentView<FMassActorFragment>();
 							const TArrayView<FMassRepresentationFragment> RepresentationsList = Context.GetMutableFragmentView<FMassRepresentationFragment>();
 							const UMassRepresentationSubsystem* RepresentationSubsystem = Context.GetSharedFragment<FMassRepresentationSubsystemSharedFragment>().RepresentationSubsystem;
 
 							if (ensure(RepresentationSubsystem))
 							{
-								for (int32 EntityIndex = 0; EntityIndex < NumEntities; ++EntityIndex)
+								for (FMassExecutionContext::FEntityIterator EntityIt = Context.CreateEntityIterator(); EntityIt; ++EntityIt)
 								{
-									if (ActorFragmentsList[EntityIndex].IsValid())
+									if (ActorFragmentsList[EntityIt].IsValid())
 									{
 										continue;
 									}
 
-									FMassRepresentationFragment& Representation = RepresentationsList[EntityIndex];
+									FMassRepresentationFragment& Representation = RepresentationsList[EntityIt];
 
 									if (Representation.ActorSpawnRequestHandle.IsValid()
 										&& ensure(RepresentationSubsystem->GetActorSpawnerSubsystem()))
@@ -224,8 +222,6 @@ void UInstancedActorsDebugProcessor::Execute(FMassEntityManager& EntityManager,
 	
 	auto ExecFunction = [this, Debugger, World = World, ReferenceLocation](FMassExecutionContext& Context)
 		{
-			//const FMassEntityHandle SelectedEntity = Debugger->GetSelectedEntity();
-			const int32 NumEntities = Context.GetNumEntities();
 			const TConstArrayView<FTransformFragment> TransformsList = Context.GetFragmentView<FTransformFragment>();
 			const TConstArrayView<FMassActorFragment> ActorFragmentsList = Context.GetFragmentView<FMassActorFragment>();
 			const TConstArrayView<FMassRepresentationFragment> RepresentationsList = Context.GetFragmentView<FMassRepresentationFragment>();
@@ -233,18 +229,18 @@ void UInstancedActorsDebugProcessor::Execute(FMassEntityManager& EntityManager,
 			const bool bFullyRepresented = !ActorFragmentsList.IsEmpty() && !RepresentationsList.IsEmpty() && RepresentationSubsystem;
 			const bool bDetailedLOD = (UE::InstancedActors::Debug::bDebugDrawAllEntities == false) || Context.DoesArchetypeHaveTag<FMassDistanceLODProcessorTag>();
 
-			for (int32 EntityIndex = 0; EntityIndex < NumEntities; ++EntityIndex)
+			for (FMassExecutionContext::FEntityIterator EntityIt = Context.CreateEntityIterator(); EntityIt; ++EntityIt)
 			{
-				const FMassEntityHandle EntityHandle = Context.GetEntity(EntityIndex);
-				const FTransformFragment& Transform = TransformsList[EntityIndex];
-				const FMassRepresentationFragment& Representation = RepresentationsList[EntityIndex];
+				const FMassEntityHandle EntityHandle = Context.GetEntity(EntityIt);
+				const FTransformFragment& Transform = TransformsList[EntityIt];
+				const FMassRepresentationFragment& Representation = RepresentationsList[EntityIt];
 
 				FVector Offset(0, 0, 150);
 				if (bFullyRepresented)
 				{
 					if (UE::InstancedActors::Debug::bDebugDrawMissingActors)
 					{
-						if (const AActor* RepresentationActor = ActorFragmentsList[EntityIndex].Get())
+						if (const AActor* RepresentationActor = ActorFragmentsList[EntityIt].Get())
 						{
 							DrawDebugSolidBox(World, RepresentationActor->GetActorLocation() + Offset, FVector(BoxSize), IsValid(RepresentationActor) ? FColor::Green : FColor::Red);
 						}
diff --git a/Engine/Plugins/Runtime/InstancedActors/Source/InstancedActors/Private/InstancedActorsInitializerProcessor.cpp b/Engine/Plugins/Runtime/InstancedActors/Source/InstancedActors/Private/InstancedActorsInitializerProcessor.cpp
--- a/Engine/Plugins/Runtime/InstancedActors/Source/InstancedActors/Private/InstancedActorsInitializerProcessor.cpp
+++ b/Engine/Plugins/Runtime/InstancedActors/Source/InstancedActors/Private/InstancedActorsInitializerProcessor.cpp
@@ -39,7 +39,7 @@ void InitInstanceFragments(UInstancedActorsData* InstanceData, int32& NextInstan
 	, const FTransform& ManagerTransform, FMassExecutionContext& Context, TArray<FMassEntityHandle>& InOutEntitiesToSignal)
 {
 	const int32 NumEntities = Context.GetNumEntities();
-	InOutEntitiesToSignal.Reserve(InOutEntitiesToSignal.Num() + Context.GetNumEntities());
+	InOutEntitiesToSignal.Reserve(InOutEntitiesToSignal.Num() + NumEntities);
 
 	TArrayView<FMassActorInstanceFragment> MassActorInstanceFragments = Context.GetMutableFragmentView<FMassActorInstanceFragment>();
 	TArrayView<FInstancedActorsFragment> InstancedActorFragments = Context.GetMutableFragmentView<FInstancedActorsFragment>();
@@ -55,12 +55,12 @@ void InitInstanceFragments(UInstancedActorsData* InstanceData, int32& NextInstan
 	}
 
 	AInstancedActorsManager& Manager = InstanceData->GetManagerChecked();
-	for (int32 Index = 0; Index < NumEntities; ++Index)
+	for (FMassExecutionContext::FEntityIterator EntityIt = Context.CreateEntityIterator(); EntityIt; ++EntityIt)
 	{
-		FMassActorInstanceFragment& MassActorInstanceFragment = MassActorInstanceFragments[Index];
-		FInstancedActorsFragment& InstancedActorFragment = InstancedActorFragments[Index];
-		FTransformFragment& TransformFragment = TransformFragments[Index];
-		const FMassEntityHandle EntityHandle(Context.GetEntity(Index));
+		FMassActorInstanceFragment& MassActorInstanceFragment = MassActorInstanceFragments[EntityIt];
+		FInstancedActorsFragment& InstancedActorFragment = InstancedActorFragments[EntityIt];
+		FTransformFragment& TransformFragment = TransformFragments[EntityIt];
+		const FMassEntityHandle EntityHandle(Context.GetEntity(EntityIt));
 		InOutEntitiesToSignal.Add(EntityHandle);
 
 		if constexpr (bFilterInstanceTransforms)
@@ -87,7 +87,7 @@ void InitInstanceFragments(UInstancedActorsData* InstanceData, int32& NextInstan
 		// @todo make GuidFragments required if we decide to go with deterministic entity naming/guid-ing
 		if (GuidFragments.Num())
 		{
-			GuidFragments[Index].Guid.D = NextInstanceIndex;
+			GuidFragments[EntityIt].Guid.D = NextInstanceIndex;
 		}
 
 		// Convert to world space
diff --git a/Engine/Plugins/Runtime/InstancedActors/Source/InstancedActors/Private/InstancedActorsManager.cpp b/Engine/Plugins/Runtime/InstancedActors/Source/InstancedActors/Private/InstancedActorsManager.cpp
--- a/Engine/Plugins/Runtime/InstancedActors/Source/InstancedActors/Private/InstancedActorsManager.cpp
+++ b/Engine/Plugins/Runtime/InstancedActors/Source/InstancedActors/Private/InstancedActorsManager.cpp
@@ -949,12 +949,11 @@ bool AInstancedActorsManager::ForEachInstance(FInstanceOperationFunc Operation,
 
 					TConstArrayView<FInstancedActorsFragment> InstancedActorFragments = Context.GetFragmentView<FInstancedActorsFragment>();
 					TConstArrayView<FTransformFragment> TransformsFragments = Context.GetFragmentView<FTransformFragment>();
-					const int32 NumEntities = Context.GetNumEntities();
-					for (int32 EntityIndex = 0; EntityIndex < NumEntities; ++EntityIndex)
+					for (FMassExecutionContext::FEntityIterator EntityIt = Context.CreateEntityIterator(); EntityIt; ++EntityIt)
 					{
-						const FTransformFragment& TransformFragment = TransformsFragments[EntityIndex];
+						const FTransformFragment& TransformFragment = TransformsFragments[EntityIt];
 
-						const FInstancedActorsFragment& InstancedActorFragment = InstancedActorFragments[EntityIndex];
+						const FInstancedActorsFragment& InstancedActorFragment = InstancedActorFragments[EntityIt];
 						InstanceHandle.Index = InstancedActorFragment.InstanceIndex;
 
 						// Execute operation
diff --git a/Engine/Plugins/Runtime/InstancedActors/Source/InstancedActors/Private/InstancedActorsVisualizationSwitcherProcessor.cpp b/Engine/Plugins/Runtime/InstancedActors/Source/InstancedActors/Private/InstancedActorsVisualizationSwitcherProcessor.cpp
--- a/Engine/Plugins/Runtime/InstancedActors/Source/InstancedActors/Private/InstancedActorsVisualizationSwitcherProcessor.cpp
+++ b/Engine/Plugins/Runtime/InstancedActors/Source/InstancedActors/Private/InstancedActorsVisualizationSwitcherProcessor.cpp
@@ -38,16 +38,15 @@ void UInstancedActorsVisualizationSwitcherProcessor::Execute(FMassEntityManager&
 		check(RepresentationSubsystem);
 		checkSlow(RepresentationSubsystem->IsA<UInstancedActorsRepresentationSubsystem>());
 		FMassInstancedStaticMeshInfoArrayView ISMInfosView = RepresentationSubsystem->GetMutableInstancedStaticMeshInfos();
-	
-		const int32 NumEntities = Context.GetNumEntities();
+
 		TConstArrayView<FInstancedActorsMeshSwitchFragment> MeshSwitchFragments = Context.GetFragmentView<FInstancedActorsMeshSwitchFragment>();
 		TArrayView<FMassRepresentationFragment> RepresentationFragments = Context.GetMutableFragmentView<FMassRepresentationFragment>();
 
-		for (int32 EntityIndex = 0; EntityIndex < NumEntities; ++EntityIndex)
+		for (FMassExecutionContext::FEntityIterator EntityIt = Context.CreateEntityIterator(); EntityIt; ++EntityIt)
 		{
-			const FMassEntityHandle EntityHandle = Context.GetEntity(EntityIndex);
-			const FInstancedActorsMeshSwitchFragment& MeshSwitchFragment = MeshSwitchFragments[EntityIndex];
-			FMassRepresentationFragment& RepresentationFragment = RepresentationFragments[EntityIndex];
+			const FMassEntityHandle EntityHandle = Context.GetEntity(EntityIt);
+			const FInstancedActorsMeshSwitchFragment& MeshSwitchFragment = MeshSwitchFragments[EntityIt];
+			FMassRepresentationFragment& RepresentationFragment = RepresentationFragments[EntityIt];
 
 			SwitchEntityMeshDesc(ISMInfosView, RepresentationFragment, EntityHandle, MeshSwitchFragment.NewStaticMeshDescHandle);
 
diff --git a/Engine/Plugins/Runtime/MassGameplay/Source/MassActors/Private/Translators/MassCharacterMovementTranslators.cpp b/Engine/Plugins/Runtime/MassGameplay/Source/MassActors/Private/Translators/MassCharacterMovementTranslators.cpp
--- a/Engine/Plugins/Runtime/MassGameplay/Source/MassActors/Private/Translators/MassCharacterMovementTranslators.cpp
+++ b/Engine/Plugins/Runtime/MassGameplay/Source/MassActors/Private/Translators/MassCharacterMovementTranslators.cpp
@@ -38,15 +38,13 @@ void UMassCharacterMovementToMassTranslator::Execute(FMassEntityManager& EntityM
 		const TArrayView<FTransformFragment> LocationList = Context.GetMutableFragmentView<FTransformFragment>();
 		const TArrayView<FMassVelocityFragment> VelocityList = Context.GetMutableFragmentView<FMassVelocityFragment>();
 
-		const int32 NumEntities = Context.GetNumEntities();
-		
-		for (int32 i = 0; i < NumEntities; ++i)
+		for (FMassExecutionContext::FEntityIterator EntityIt = Context.CreateEntityIterator(); EntityIt; ++EntityIt)
 		{
-			if (const UCharacterMovementComponent* AsMovementComponent = ComponentList[i].Component.Get())
+			if (const UCharacterMovementComponent* AsMovementComponent = ComponentList[EntityIt].Component.Get())
 			{
-				LocationList[i].GetMutableTransform().SetLocation(AsMovementComponent->GetActorNavLocation());
+				LocationList[EntityIt].GetMutableTransform().SetLocation(AsMovementComponent->GetActorNavLocation());
 				
-				VelocityList[i].Value = AsMovementComponent->Velocity;
+				VelocityList[EntityIt].Value = AsMovementComponent->Velocity;
 			}
 		}
 	});
@@ -78,14 +76,12 @@ void UMassCharacterMovementToActorTranslator::Execute(FMassEntityManager& Entity
 	{
 		const TArrayView<FCharacterMovementComponentWrapperFragment> ComponentList = Context.GetMutableFragmentView<FCharacterMovementComponentWrapperFragment>();
 		const TConstArrayView<FMassVelocityFragment> VelocityList = Context.GetFragmentView<FMassVelocityFragment>();
-		
-		const int32 NumEntities = Context.GetNumEntities();
 
-		for (int32 i = 0; i < NumEntities; ++i)
+		for (FMassExecutionContext::FEntityIterator EntityIt = Context.CreateEntityIterator(); EntityIt; ++EntityIt)
 		{
-			if (UCharacterMovementComponent* AsMovementComponent = ComponentList[i].Component.Get())
+			if (UCharacterMovementComponent* AsMovementComponent = ComponentList[EntityIt].Component.Get())
 			{
-				AsMovementComponent->RequestDirectMove(VelocityList[i].Value, /*bForceMaxSpeed=*/false);
+				AsMovementComponent->RequestDirectMove(VelocityList[EntityIt].Value, /*bForceMaxSpeed=*/false);
 			}
 		}
 	});
@@ -117,15 +113,13 @@ void UMassCharacterOrientationToMassTranslator::Execute(FMassEntityManager& Enti
 		const TConstArrayView<FCharacterMovementComponentWrapperFragment> ComponentList = Context.GetFragmentView<FCharacterMovementComponentWrapperFragment>();
 		const TArrayView<FTransformFragment> LocationList = Context.GetMutableFragmentView<FTransformFragment>();
 
-		const int32 NumEntities = Context.GetNumEntities();
-		
-		for (int32 i = 0; i < NumEntities; ++i)
+		for (FMassExecutionContext::FEntityIterator EntityIt = Context.CreateEntityIterator(); EntityIt; ++EntityIt)
 		{
-			if (const UCharacterMovementComponent* AsMovementComponent = ComponentList[i].Component.Get())
+			if (const UCharacterMovementComponent* AsMovementComponent = ComponentList[EntityIt].Component.Get())
 			{
 				if (AsMovementComponent->UpdatedComponent != nullptr)
 				{
-					LocationList[i].GetMutableTransform().SetRotation(AsMovementComponent->UpdatedComponent->GetComponentTransform().GetRotation());
+					LocationList[EntityIt].GetMutableTransform().SetRotation(AsMovementComponent->UpdatedComponent->GetComponentTransform().GetRotation());
 				}
 			}
 		}
@@ -159,16 +153,14 @@ void UMassCharacterOrientationToActorTranslator::Execute(FMassEntityManager& Ent
 	{
 		const TArrayView<FCharacterMovementComponentWrapperFragment> ComponentList = Context.GetMutableFragmentView<FCharacterMovementComponentWrapperFragment>();
 		const TConstArrayView<FTransformFragment> TransformList = Context.GetFragmentView<FTransformFragment>();
-		
-		const int32 NumEntities = Context.GetNumEntities();
 
-		for (int32 i = 0; i < NumEntities; ++i)
+		for (FMassExecutionContext::FEntityIterator EntityIt = Context.CreateEntityIterator(); EntityIt; ++EntityIt)
 		{
-			if (UCharacterMovementComponent* AsMovementComponent = ComponentList[i].Component.Get())
+			if (UCharacterMovementComponent* AsMovementComponent = ComponentList[EntityIt].Component.Get())
 			{
 				if (AsMovementComponent->UpdatedComponent != nullptr)
 				{
-					const FTransformFragment& Transform = TransformList[i];
+					const FTransformFragment& Transform = TransformList[EntityIt];
 					AsMovementComponent->bOrientRotationToMovement = false;
 					AsMovementComponent->UpdatedComponent->SetWorldRotation(Transform.GetTransform().GetRotation());
 				}
diff --git a/Engine/Plugins/Runtime/MassGameplay/Source/MassActors/Private/Translators/MassSceneComponentLocationTranslator.cpp b/Engine/Plugins/Runtime/MassGameplay/Source/MassActors/Private/Translators/MassSceneComponentLocationTranslator.cpp
--- a/Engine/Plugins/Runtime/MassGameplay/Source/MassActors/Private/Translators/MassSceneComponentLocationTranslator.cpp
+++ b/Engine/Plugins/Runtime/MassGameplay/Source/MassActors/Private/Translators/MassSceneComponentLocationTranslator.cpp
@@ -30,12 +30,11 @@ void UMassSceneComponentLocationToMassTranslator::Execute(FMassEntityManager& En
 		const TConstArrayView<FMassSceneComponentWrapperFragment> ComponentList = Context.GetFragmentView<FMassSceneComponentWrapperFragment>();
 		const TArrayView<FTransformFragment> LocationList = Context.GetMutableFragmentView<FTransformFragment>();
 
-		const int32 NumEntities = Context.GetNumEntities();
-		for (int32 i = 0; i < NumEntities; ++i)
+		for (FMassExecutionContext::FEntityIterator EntityIt = Context.CreateEntityIterator(); EntityIt; ++EntityIt)
 		{
-			if (const USceneComponent* AsComponent = ComponentList[i].Component.Get())
+			if (const USceneComponent* AsComponent = ComponentList[EntityIt].Component.Get())
 			{
-				LocationList[i].GetMutableTransform().SetLocation(AsComponent->GetComponentTransform().GetLocation() - FVector(0.f, 0.f, AsComponent->Bounds.BoxExtent.Z));
+				LocationList[EntityIt].GetMutableTransform().SetLocation(AsComponent->GetComponentTransform().GetLocation() - FVector(0.f, 0.f, AsComponent->Bounds.BoxExtent.Z));
 			}
 		}
 	});
@@ -69,12 +68,11 @@ void UMassSceneComponentLocationToActorTranslator::Execute(FMassEntityManager& E
 			const TConstArrayView<FMassSceneComponentWrapperFragment> ComponentList = Context.GetFragmentView<FMassSceneComponentWrapperFragment>();
 			const TArrayView<FTransformFragment> LocationList = Context.GetMutableFragmentView<FTransformFragment>();
 
-			const int32 NumEntities = Context.GetNumEntities();
-			for (int32 i = 0; i < NumEntities; ++i)
+			for (FMassExecutionContext::FEntityIterator EntityIt = Context.CreateEntityIterator(); EntityIt; ++EntityIt)
 			{
-				if (USceneComponent* AsComponent = ComponentList[i].Component.Get())
+				if (USceneComponent* AsComponent = ComponentList[EntityIt].Component.Get())
 				{
-					AsComponent->SetWorldLocation(LocationList[i].GetTransform().GetLocation() + FVector(0.f, 0.f, AsComponent->Bounds.BoxExtent.Z));
+					AsComponent->SetWorldLocation(LocationList[EntityIt].GetTransform().GetLocation() + FVector(0.f, 0.f, AsComponent->Bounds.BoxExtent.Z));
 				}
 			}
 		});
diff --git a/Engine/Plugins/Runtime/MassGameplay/Source/MassEQS/Private/Processors/MassEnvQueryGeneratorProcessor_MassEntityHandles.cpp b/Engine/Plugins/Runtime/MassGameplay/Source/MassEQS/Private/Processors/MassEnvQueryGeneratorProcessor_MassEntityHandles.cpp
--- a/Engine/Plugins/Runtime/MassGameplay/Source/MassEQS/Private/Processors/MassEnvQueryGeneratorProcessor_MassEntityHandles.cpp
+++ b/Engine/Plugins/Runtime/MassGameplay/Source/MassEQS/Private/Processors/MassEnvQueryGeneratorProcessor_MassEntityHandles.cpp
@@ -46,11 +46,10 @@ void UMassEnvQueryGeneratorProcessor_MassEntityHandles::Execute(FMassEntityManag
 		EntityQuery.ForEachEntityChunk(Context, [&GeneratorData, &Items, &SearchRadius, &ContextPositions](FMassExecutionContext& ChunkContext)
 		{
 			const TConstArrayView<FTransformFragment> TransformFragmentList = ChunkContext.GetFragmentView<FTransformFragment>();
-			const int32 NumEntities = ChunkContext.GetNumEntities();
-			for (int32 EntityIndex = 0; EntityIndex < NumEntities; ++EntityIndex)
+			for (FMassExecutionContext::FEntityIterator EntityIt = ChunkContext.CreateEntityIterator(); EntityIt; ++EntityIt)
 			{
-				FMassEntityHandle EntityHandle = ChunkContext.GetEntity(EntityIndex);
-				Items.Emplace(EntityHandle.Index, EntityHandle.SerialNumber, TransformFragmentList[EntityIndex].GetTransform());
+				FMassEntityHandle EntityHandle = ChunkContext.GetEntity(EntityIt);
+				Items.Emplace(EntityHandle.Index, EntityHandle.SerialNumber, TransformFragmentList[EntityIt].GetTransform());
 			}
 		});
 	}
@@ -60,10 +59,9 @@ void UMassEnvQueryGeneratorProcessor_MassEntityHandles::Execute(FMassEntityManag
 		EntityQuery.ForEachEntityChunk(Context, [&GeneratorData, &Items, &SearchRadiusSqr, &ContextPositions](FMassExecutionContext& ChunkContext)
 		{
 			const TConstArrayView<FTransformFragment> TransformFragmentList = ChunkContext.GetFragmentView<FTransformFragment>();
-			const int32 NumEntities = ChunkContext.GetNumEntities();
-			for (int32 EntityIndex = 0; EntityIndex < NumEntities; ++EntityIndex)
+			for (FMassExecutionContext::FEntityIterator EntityIt = ChunkContext.CreateEntityIterator(); EntityIt; ++EntityIt)
 			{
-				const FTransformFragment& TransformFragment = TransformFragmentList[EntityIndex];
+				const FTransformFragment& TransformFragment = TransformFragmentList[EntityIt];
 				const FVector& EntityPosition = TransformFragment.GetTransform().GetTranslation();
 				for (const FVector& ContextPosition : ContextPositions)
 				{
@@ -71,8 +69,8 @@ void UMassEnvQueryGeneratorProcessor_MassEntityHandles::Execute(FMassEntityManag
 
 					if (ContextDistanceToEntitySqr <= SearchRadiusSqr)
 					{
-						FMassEntityHandle EntityHandle = ChunkContext.GetEntity(EntityIndex);
-						Items.Emplace(EntityHandle.Index, EntityHandle.SerialNumber, TransformFragmentList[EntityIndex].GetTransform());
+						FMassEntityHandle EntityHandle = ChunkContext.GetEntity(EntityIt);
+						Items.Emplace(EntityHandle.Index, EntityHandle.SerialNumber, TransformFragmentList[EntityIt].GetTransform());
 						break;
 					}
 				}
diff --git a/Engine/Plugins/Runtime/MassGameplay/Source/MassEQS/Private/Processors/MassEnvQueryTestProcessor_MassEntityTags.cpp b/Engine/Plugins/Runtime/MassGameplay/Source/MassEQS/Private/Processors/MassEnvQueryTestProcessor_MassEntityTags.cpp
--- a/Engine/Plugins/Runtime/MassGameplay/Source/MassEQS/Private/Processors/MassEnvQueryTestProcessor_MassEntityTags.cpp
+++ b/Engine/Plugins/Runtime/MassGameplay/Source/MassEQS/Private/Processors/MassEnvQueryTestProcessor_MassEntityTags.cpp
@@ -90,10 +90,9 @@ void UMassEnvQueryTestProcessor_MassEntityTags::Execute(FMassEntityManager& Enti
 		{
 			bool ChunkSuccess = ScoringFunction(Tags, Context);
 
-			const int32 NumEntities = Context.GetNumEntities();
-			for (int32 EntityIndex = 0; EntityIndex < NumEntities; ++EntityIndex)
+			for (FMassExecutionContext::FEntityIterator EntityIt = Context.CreateEntityIterator(); EntityIt; ++EntityIt)
 			{
-				const FMassEntityHandle EntityHandle = Context.GetEntity(EntityIndex);
+				const FMassEntityHandle EntityHandle = Context.GetEntity(EntityIt);
 				ScoreMap.Add(EntityHandle, ChunkSuccess);
 			}
 		});
diff --git a/Engine/Plugins/Runtime/MassGameplay/Source/MassGameplayDebug/Private/DebugVisLocationProcessor.cpp b/Engine/Plugins/Runtime/MassGameplay/Source/MassGameplayDebug/Private/DebugVisLocationProcessor.cpp
--- a/Engine/Plugins/Runtime/MassGameplay/Source/MassGameplayDebug/Private/DebugVisLocationProcessor.cpp
+++ b/Engine/Plugins/Runtime/MassGameplay/Source/MassGameplayDebug/Private/DebugVisLocationProcessor.cpp
@@ -40,16 +40,15 @@ void UDebugVisLocationProcessor::Execute(FMassEntityManager& EntityManager, FMas
 		TArrayView<UHierarchicalInstancedStaticMeshComponent* const> VisualDataISMCs = Visualizer->GetVisualDataISMCs();
 		if (VisualDataISMCs.Num() > 0)
 		{
-			const int32 NumEntities = Context.GetNumEntities();
 			const TConstArrayView<FTransformFragment> LocationList = Context.GetFragmentView<FTransformFragment>();
 			const TConstArrayView<FSimDebugVisFragment> DebugVisList = Context.GetFragmentView<FSimDebugVisFragment>();
 
-			for (int32 i = 0; i < NumEntities; ++i)
+			for (FMassExecutionContext::FEntityIterator EntityIt = Context.CreateEntityIterator(); EntityIt; ++EntityIt)
 			{
-				const FSimDebugVisFragment& VisualComp = DebugVisList[i];
+				const FSimDebugVisFragment& VisualComp = DebugVisList[EntityIt];
 
 				// @todo: remove this code once the asset is exported with correct alignment SM_Mannequin.uasset
-				FTransform SMTransform = LocationList[i].GetTransform();
+				FTransform SMTransform = LocationList[EntityIt].GetTransform();
 				FQuat FromEngineToSM(FVector::UpVector, -HALF_PI);
 				SMTransform.SetRotation(FromEngineToSM * SMTransform.GetRotation());
 
@@ -99,14 +98,13 @@ void UMassProcessor_UpdateDebugVis::Execute(FMassEntityManager& EntityManager, F
 
 	EntityQuery.ForEachEntityChunk(Context, [&Debugger](FMassExecutionContext& Context)
 		{
-			const int32 NumEntities = Context.GetNumEntities();
 			const TConstArrayView<FTransformFragment> LocationList = Context.GetFragmentView<FTransformFragment>();
 			const TArrayView<FDataFragment_DebugVis> DebugVisList = Context.GetMutableFragmentView<FDataFragment_DebugVis>();
 			const TArrayView<FAgentRadiusFragment> RadiiList = Context.GetMutableFragmentView<FAgentRadiusFragment>();
 
-			for (int32 i = 0; i < NumEntities; ++i)
+			for (FMassExecutionContext::FEntityIterator EntityIt = Context.CreateEntityIterator(); EntityIt; ++EntityIt)
 			{
-				Debugger.AddShape(DebugVisList[i].Shape, LocationList[i].GetTransform().GetLocation(), RadiiList[i].Radius);
+				Debugger.AddShape(DebugVisList[EntityIt].Shape, LocationList[EntityIt].GetTransform().GetLocation(), RadiiList[EntityIt].Radius);
 			}
 		});
 }
diff --git a/Engine/Plugins/Runtime/MassGameplay/Source/MassLOD/Private/MassSimulationLOD.cpp b/Engine/Plugins/Runtime/MassGameplay/Source/MassLOD/Private/MassSimulationLOD.cpp
--- a/Engine/Plugins/Runtime/MassGameplay/Source/MassLOD/Private/MassSimulationLOD.cpp
+++ b/Engine/Plugins/Runtime/MassGameplay/Source/MassLOD/Private/MassSimulationLOD.cpp
@@ -178,13 +178,13 @@ void UMassSimulationLODProcessor::Execute(FMassEntityManager& EntityManager, FMa
 		EntityQuerySetLODTag.ForEachEntityChunk(Context, [](FMassExecutionContext& Context)
 		{
 			TConstArrayView<FMassSimulationLODFragment> SimulationLODFragments = Context.GetFragmentView<FMassSimulationLODFragment>();
-			const int32 NumEntities = Context.GetNumEntities();
-			for (int32 Index = 0; Index < NumEntities; ++Index)
+
+			for (FMassExecutionContext::FEntityIterator EntityIt = Context.CreateEntityIterator(); EntityIt; ++EntityIt)
 			{
-				const FMassSimulationLODFragment& EntityLOD = SimulationLODFragments[Index];
+				const FMassSimulationLODFragment& EntityLOD = SimulationLODFragments[EntityIt];
 				if (EntityLOD.PrevLOD != EntityLOD.LOD)
 				{
-					const FMassEntityHandle Entity = Context.GetEntity(Index);
+					const FMassEntityHandle Entity = Context.GetEntity(EntityIt);
 					UE::MassLOD::PushSwapTagsCommand(Context.Defer(), Entity, EntityLOD.PrevLOD, EntityLOD.LOD);
 				}
 			}
diff --git a/Engine/Plugins/Runtime/MassGameplay/Source/MassLOD/Public/MassLODCalculator.h b/Engine/Plugins/Runtime/MassGameplay/Source/MassLOD/Public/MassLODCalculator.h
--- a/Engine/Plugins/Runtime/MassGameplay/Source/MassLOD/Public/MassLODCalculator.h
+++ b/Engine/Plugins/Runtime/MassGameplay/Source/MassLOD/Public/MassLODCalculator.h
@@ -439,12 +439,11 @@ void TMassLODCalculator<FLODLogic>::CalculateLOD(FMassExecutionContext& Context,
 	}
 #endif // WITH_MASSGAMEPLAY_DEBUG
 
-	const int32 NumEntities = Context.GetNumEntities();
-	for (int EntityIdx = 0; EntityIdx < NumEntities; EntityIdx++)
+	for (FMassExecutionContext::FEntityIterator EntityIt = Context.CreateEntityIterator(); EntityIt; ++EntityIt)
 	{
 		// Calculate the LOD purely upon distances
-		const TViewerInfoFragment& EntityViewersInfo = ViewersInfoList[EntityIdx];
-		TLODFragment& EntityLOD = LODList[EntityIdx];
+		const TViewerInfoFragment& EntityViewersInfo = ViewersInfoList[EntityIt];
+		TLODFragment& EntityLOD = LODList[EntityIt];
 		const float ClosestDistanceToFrustum = GetClosestDistanceToFrustum<bCalculateVisibility>(EntityViewersInfo, FLT_MAX);
 		const EMassVisibility PrevVisibility = GetVisibility<bCalculateVisibility>(EntityLOD, EMassVisibility::Max);
 		const bool bIsVisibleByAViewer = CalculateVisibility(PrevVisibility == EMassVisibility::CanBeSeen, ClosestDistanceToFrustum);
@@ -465,7 +464,7 @@ void TMassLODCalculator<FLODLogic>::CalculateLOD(FMassExecutionContext& Context,
 		// Do per viewer logic if asked for
 		if (bCalculateLODPerViewer || bCalculateVisibilityPerViewer)
 		{
-			const TPerViewerInfoFragment& EntityPerViewerInfo = PerViewerInfoList[EntityIdx];
+			const TPerViewerInfoFragment& EntityPerViewerInfo = PerViewerInfoList[EntityIt];
 
 			SetLODPerViewerNum<bCalculateLODPerViewer>(EntityLOD, Viewers.Num());
 			SetPrevLODPerViewerNum<bCalculateLODPerViewer>(EntityLOD, Viewers.Num());
@@ -676,16 +675,15 @@ void TMassLODCalculator<FLODLogic>::AdjustLODFromCount(FMassExecutionContext& Co
 	static_assert(!bCalculateVisibilityPerViewer || (FLODLogic::bDoVisibilityLogic && FLODLogic::bStoreInfoPerViewer), "FLODLogic must have bDoVisibilityLogic and bStoreInfoPerViewer enabled to calculate visibility per viewer.");
 	static_assert(!bMaximizeCountPerViewer || FLODLogic::bMaximizeCountPerViewer, "FLODLogic must have bMaximizeCountPerViewer enabled to maximize count per viewer.");
 
-	const int32 NumEntities = Context.GetNumEntities();
 	// Adjust LOD for each viewer and remember the new highest
-	for (int EntityIdx = 0; EntityIdx < NumEntities; EntityIdx++)
+	for (FMassExecutionContext::FEntityIterator EntityIt = Context.CreateEntityIterator(); EntityIt; ++EntityIt)
 	{
-		const TViewerInfoFragment& EntityViewersInfo = ViewersInfoList[EntityIdx];
-		TLODFragment& EntityLOD = LODList[EntityIdx];
+		const TViewerInfoFragment& EntityViewersInfo = ViewersInfoList[EntityIt];
+		TLODFragment& EntityLOD = LODList[EntityIt];
 		EMassLOD::Type HighestViewerLOD = EMassLOD::Off;
 		if (bMaximizeCountPerViewer)
 		{
-			const TPerViewerInfoFragment& EntityPerViewerInfo = PerViewerInfoList[EntityIdx];
+			const TPerViewerInfoFragment& EntityPerViewerInfo = PerViewerInfoList[EntityIt];
 
 			for (int ViewerIdx = 0; ViewerIdx < Viewers.Num(); ++ViewerIdx)
 			{
@@ -759,10 +757,9 @@ template <typename FLODLogic>
 template <typename TLODFragment>
 void TMassLODCalculator<FLODLogic>::ForceOffLOD(FMassExecutionContext& Context, TArrayView<TLODFragment> LODList)
 {
-	const int32 NumEntities = Context.GetNumEntities();
-	for (int EntityIdx = 0; EntityIdx < NumEntities; EntityIdx++)
+	for (FMassExecutionContext::FEntityIterator EntityIt = Context.CreateEntityIterator(); EntityIt; ++EntityIt)
 	{
-		TLODFragment& EntityLOD = LODList[EntityIdx];
+		TLODFragment& EntityLOD = LODList[EntityIt];
 
 		// Force off LOD
 		EntityLOD.PrevLOD = EntityLOD.LOD;
@@ -797,11 +794,10 @@ template <typename FLODLogic>
 template <typename TLODFragment, typename TTransformFragment>
 void TMassLODCalculator<FLODLogic>::DebugDisplayLOD(FMassExecutionContext& Context, TConstArrayView<TLODFragment> LODList, TConstArrayView<TTransformFragment> LocationList, UWorld* World)
 {
-	const int32 NumEntities = Context.GetNumEntities();
-	for (int EntityIdx = 0; EntityIdx < NumEntities; EntityIdx++)
+	for (FMassExecutionContext::FEntityIterator EntityIt = Context.CreateEntityIterator(); EntityIt; ++EntityIt)
 	{
-		const TTransformFragment& EntityLocation = LocationList[EntityIdx];
-		const TLODFragment& EntityLOD = LODList[EntityIdx];
+		const TTransformFragment& EntityLocation = LocationList[EntityIt];
+		const TLODFragment& EntityLOD = LODList[EntityIt];
 		int32 LODIdx = (int32)EntityLOD.LOD;
 		DrawDebugSolidBox(World, EntityLocation.GetTransform().GetLocation() + FVector(0.0f, 0.0f, 120.0f), FVector(25.0f), UE::MassLOD::LODColors[LODIdx]);
 	}
@@ -811,13 +807,12 @@ template <typename FLODLogic>
 template <typename TLODFragment, typename TTransformFragment>
 void TMassLODCalculator<FLODLogic>::DebugDisplaySignificantLOD(FMassExecutionContext& Context, TConstArrayView<TLODFragment> LODList, TConstArrayView<TTransformFragment> LocationList, UWorld* World, const float MaxLODSignificance)
 {
-	const int32 NumEntities = Context.GetNumEntities();
-	for (int EntityIdx = 0; EntityIdx < NumEntities; EntityIdx++)
+	for (FMassExecutionContext::FEntityIterator EntityIt = Context.CreateEntityIterator(); EntityIt; ++EntityIt)
 	{
-		const TLODFragment& EntityLOD = LODList[EntityIdx];
+		const TLODFragment& EntityLOD = LODList[EntityIt];
 		if (EntityLOD.LODSignificance <= MaxLODSignificance)
 		{
-			const TTransformFragment& EntityLocation = LocationList[EntityIdx];
+			const TTransformFragment& EntityLocation = LocationList[EntityIt];
 			int32 LODIdx = (int32)EntityLOD.LOD;
 			DrawDebugSolidBox(World, EntityLocation.GetTransform().GetLocation() + FVector(0.0f, 0.0f, 120.0f), FVector(25.0f), UE::MassLOD::LODColors[LODIdx]);
 		}
@@ -829,13 +824,12 @@ template <typename TLODFragment, typename TTransformFragment>
 void TMassLODCalculator<FLODLogic>::VisLogLOD(FMassExecutionContext& Context, TConstArrayView<TLODFragment> LODList, TConstArrayView<TTransformFragment> LocationList, UObject* LogOwner)
 {
 #if ENABLE_VISUAL_LOG
-	const int32 NumEntities = Context.GetNumEntities();
-	for (int EntityIdx = 0; EntityIdx < NumEntities; EntityIdx++)
+	for (FMassExecutionContext::FEntityIterator EntityIt = Context.CreateEntityIterator(); EntityIt; ++EntityIt)
 	{
-		const TTransformFragment& EntityLocation = LocationList[EntityIdx];
-		const TLODFragment& EntityLOD = LODList[EntityIdx];
+		const TTransformFragment& EntityLocation = LocationList[EntityIt];
+		const TLODFragment& EntityLOD = LODList[EntityIt];
 		int32 LODIdx = (int32)EntityLOD.LOD;
-		UE_VLOG_LOCATION(LogOwner, LogMassLOD, Verbose, EntityLocation.GetTransform().GetLocation(), 20.0f, UE::MassLOD::LODColors[LODIdx], TEXT("%s %d"), *Context.GetEntity(EntityIdx).DebugGetDescription(), LODIdx);
+		UE_VLOG_LOCATION(LogOwner, LogMassLOD, Verbose, EntityLocation.GetTransform().GetLocation(), 20.0f, UE::MassLOD::LODColors[LODIdx], TEXT("%s %d"), *Context.GetEntity(EntityIt).DebugGetDescription(), LODIdx);
 	}
 #endif
 }
@@ -845,15 +839,14 @@ template <typename TLODFragment, typename TTransformFragment>
 void TMassLODCalculator<FLODLogic>::VisLogSignificantLOD(FMassExecutionContext& Context, TConstArrayView<TLODFragment> LODList, TConstArrayView<TTransformFragment> LocationList, UObject* LogOwner, const float MaxLODSignificance)
 {
 #if ENABLE_VISUAL_LOG
-	const int32 NumEntities = Context.GetNumEntities();
-	for (int EntityIdx = 0; EntityIdx < NumEntities; EntityIdx++)
+	for (FMassExecutionContext::FEntityIterator EntityIt = Context.CreateEntityIterator(); EntityIt; ++EntityIt)
 	{
-		const TLODFragment& EntityLOD = LODList[EntityIdx];
+		const TLODFragment& EntityLOD = LODList[EntityIt];
 		if (EntityLOD.LODSignificance <= MaxLODSignificance)
 		{
-			const TTransformFragment& EntityLocation = LocationList[EntityIdx];
+			const TTransformFragment& EntityLocation = LocationList[EntityIt];
 			int32 LODIdx = (int32)EntityLOD.LOD;
-			UE_VLOG_LOCATION(LogOwner, LogMassLOD, Verbose, EntityLocation.GetTransform().GetLocation(), 20, UE::MassLOD::LODColors[LODIdx], TEXT("%s %d"), *Context.GetEntity(EntityIdx).DebugGetDescription(), LODIdx);
+			UE_VLOG_LOCATION(LogOwner, LogMassLOD, Verbose, EntityLocation.GetTransform().GetLocation(), 20, UE::MassLOD::LODColors[LODIdx], TEXT("%s %d"), *Context.GetEntity(EntityIt).DebugGetDescription(), LODIdx);
 		}
 	}
 #endif
diff --git a/Engine/Plugins/Runtime/MassGameplay/Source/MassLOD/Public/MassLODCollector.h b/Engine/Plugins/Runtime/MassGameplay/Source/MassLOD/Public/MassLODCollector.h
--- a/Engine/Plugins/Runtime/MassGameplay/Source/MassLOD/Public/MassLODCollector.h
+++ b/Engine/Plugins/Runtime/MassGameplay/Source/MassLOD/Public/MassLODCollector.h
@@ -92,14 +92,13 @@ void TMassLODCollector<FLODLogic>::CollectLODInfo(FMassExecutionContext& Context
 												  TArrayView<TPerViewerInfoFragment> PerViewerInfoList)
 {
 	static TPerViewerInfoFragment DummyFragment;
-	const int32 NumEntities = Context.GetNumEntities();
-	for (int EntityIdx = 0; EntityIdx < NumEntities; EntityIdx++)
+	for (FMassExecutionContext::FEntityIterator EntityIt = Context.CreateEntityIterator(); EntityIt; ++EntityIt)
 	{
 		float ClosestViewerDistanceSq = FLT_MAX;
 		float ClosestDistanceToFrustum = FLT_MAX;
-		const TTransformFragment& EntityTransform = TransformList[EntityIdx];
-		TViewerInfoFragment& EntityViewerInfo = ViewersInfoList[EntityIdx];
-		TPerViewerInfoFragment& EntityInfoPerViewer = FLODLogic::bStoreInfoPerViewer ? PerViewerInfoList[EntityIdx] : DummyFragment;
+		const TTransformFragment& EntityTransform = TransformList[EntityIt];
+		TViewerInfoFragment& EntityViewerInfo = ViewersInfoList[EntityIt];
+		TPerViewerInfoFragment& EntityInfoPerViewer = FLODLogic::bStoreInfoPerViewer ? PerViewerInfoList[EntityIt] : DummyFragment;
 
 		SetDistanceToViewerSqNum<bCollectDistancePerViewer>(EntityInfoPerViewer, Viewers.Num());
 		SetDistanceToFrustumNum<bCollectDistanceToFrustumPerViewer>(EntityInfoPerViewer, Viewers.Num());
diff --git a/Engine/Plugins/Runtime/MassGameplay/Source/MassLOD/Public/MassLODTickRateController.h b/Engine/Plugins/Runtime/MassGameplay/Source/MassLOD/Public/MassLODTickRateController.h
--- a/Engine/Plugins/Runtime/MassGameplay/Source/MassLOD/Public/MassLODTickRateController.h
+++ b/Engine/Plugins/Runtime/MassGameplay/Source/MassLOD/Public/MassLODTickRateController.h
@@ -156,16 +156,15 @@ bool TMassLODTickRateController<TVariableTickChunkFragment, FLODLogic>::UpdateTi
 
 	if (bWasChunkTicked)
 	{
-		const int32 NumEntities = Context.GetNumEntities();
-		for (int32 Index = 0; Index < NumEntities; ++Index)
+		for (FMassExecutionContext::FEntityIterator EntityIt = Context.CreateEntityIterator(); EntityIt; ++EntityIt)
 		{
-			const TLODFragment& EntityLOD = LODList[Index];
-			TVariableTickRateFragment& TickRate = TickRateList[Index];
+			const TLODFragment& EntityLOD = LODList[EntityIt];
+			TVariableTickRateFragment& TickRate = TickRateList[EntityIt];
 			TickRate.DeltaTime = TickRate.LastTickedTime != 0.0 ? static_cast<float>(Time - TickRate.LastTickedTime) : DeltaTime;
 			TickRate.LastTickedTime = Time;
 			if (EntityLOD.LOD != ChunkLOD)
 			{
-				const FMassEntityHandle Entity = Context.GetEntity(Index);
+				const FMassEntityHandle Entity = Context.GetEntity(EntityIt);
 				UE::MassLOD::PushSwapTagsCommand(Context.Defer(), Entity, ChunkLOD, EntityLOD.LOD);
 			}
 		}
diff --git a/Engine/Plugins/Runtime/MassGameplay/Source/MassMovement/Private/MassMovementProcessors.cpp b/Engine/Plugins/Runtime/MassGameplay/Source/MassMovement/Private/MassMovementProcessors.cpp
--- a/Engine/Plugins/Runtime/MassGameplay/Source/MassMovement/Private/MassMovementProcessors.cpp
+++ b/Engine/Plugins/Runtime/MassGameplay/Source/MassMovement/Private/MassMovementProcessors.cpp
@@ -43,19 +43,17 @@ void UMassApplyMovementProcessor::Execute(FMassEntityManager& EntityManager,
 
 	EntityQuery.ForEachEntityChunk(Context, [this, DeltaTime](FMassExecutionContext& Context)
 	{
-		const int32 NumEntities = Context.GetNumEntities();
-
 		const FMassMovementParameters& MovementParams = Context.GetConstSharedFragment<FMassMovementParameters>();
 
 		const TArrayView<FTransformFragment> LocationList = Context.GetMutableFragmentView<FTransformFragment>();
 		const TArrayView<FMassForceFragment> ForceList = Context.GetMutableFragmentView<FMassForceFragment>();
 		const TArrayView<FMassVelocityFragment> VelocityList = Context.GetMutableFragmentView<FMassVelocityFragment>();
 
-		for (int32 EntityIndex = 0; EntityIndex < NumEntities; ++EntityIndex)
+		for (FMassExecutionContext::FEntityIterator EntityIt = Context.CreateEntityIterator(); EntityIt; ++EntityIt)
 		{
-			FMassForceFragment& Force = ForceList[EntityIndex];
-			FMassVelocityFragment& Velocity = VelocityList[EntityIndex];
-			FTransform& CurrentTransform = LocationList[EntityIndex].GetMutableTransform();
+			FMassForceFragment& Force = ForceList[EntityIt];
+			FMassVelocityFragment& Velocity = VelocityList[EntityIt];
+			FTransform& CurrentTransform = LocationList[EntityIt].GetMutableTransform();
 
 			// Update velocity from steering forces.
 			Velocity.Value += Force.Value * DeltaTime;
diff --git a/Engine/Plugins/Runtime/MassGameplay/Source/MassMovement/Private/MassSimpleMovementTrait.cpp b/Engine/Plugins/Runtime/MassGameplay/Source/MassMovement/Private/MassSimpleMovementTrait.cpp
--- a/Engine/Plugins/Runtime/MassGameplay/Source/MassMovement/Private/MassSimpleMovementTrait.cpp
+++ b/Engine/Plugins/Runtime/MassGameplay/Source/MassMovement/Private/MassSimpleMovementTrait.cpp
@@ -41,9 +41,9 @@ void UMassSimpleMovementProcessor::ConfigureQueries(const TSharedRef<FMassEntity
 	EntityQuery.SetChunkFilter(&FMassSimulationVariableTickChunkFragment::ShouldTickChunkThisFrame);
 }
 
-void UMassSimpleMovementProcessor::Execute(FMassEntityManager& EntityManager, FMassExecutionContext& Context)
+void UMassSimpleMovementProcessor::Execute(FMassEntityManager& EntityManager, FMassExecutionContext& ExecutionContext)
 {
-	EntityQuery.ForEachEntityChunk(Context, ([this](FMassExecutionContext& Context)
+	EntityQuery.ForEachEntityChunk(ExecutionContext, ([this](FMassExecutionContext& Context)
 		{
 			const TConstArrayView<FMassVelocityFragment> VelocitiesList = Context.GetFragmentView<FMassVelocityFragment>();
 			const TArrayView<FTransformFragment> TransformsList = Context.GetMutableFragmentView<FTransformFragment>();
@@ -51,11 +51,11 @@ void UMassSimpleMovementProcessor::Execute(FMassEntityManager& EntityManager, FM
 			const bool bHasVariableTick = (SimVariableTickList.Num() > 0);
 			const float WorldDeltaTime = Context.GetDeltaTimeSeconds();
 		
-			for (int32 EntityIndex = 0; EntityIndex < Context.GetNumEntities(); ++EntityIndex)
+			for (FMassExecutionContext::FEntityIterator EntityIt = Context.CreateEntityIterator(); EntityIt; ++EntityIt)
 			{
-				const FMassVelocityFragment& Velocity = VelocitiesList[EntityIndex];
-				FTransform& Transform = TransformsList[EntityIndex].GetMutableTransform();
-				const float DeltaTime = bHasVariableTick ? SimVariableTickList[EntityIndex].DeltaTime : WorldDeltaTime;
+				const FMassVelocityFragment& Velocity = VelocitiesList[EntityIt];
+				FTransform& Transform = TransformsList[EntityIt].GetMutableTransform();
+				const float DeltaTime = bHasVariableTick ? SimVariableTickList[EntityIt].DeltaTime : WorldDeltaTime;
 				Transform.SetTranslation(Transform.GetTranslation() + Velocity.Value * DeltaTime);
 			}
 		}));
diff --git a/Engine/Plugins/Runtime/MassGameplay/Source/MassReplication/Private/MassReplicationFragments.cpp b/Engine/Plugins/Runtime/MassGameplay/Source/MassReplication/Private/MassReplicationFragments.cpp
--- a/Engine/Plugins/Runtime/MassGameplay/Source/MassReplication/Private/MassReplicationFragments.cpp
+++ b/Engine/Plugins/Runtime/MassGameplay/Source/MassReplication/Private/MassReplicationFragments.cpp
@@ -38,11 +38,9 @@ void UMassNetworkIDFragmentInitializer::Execute(FMassEntityManager& EntityManage
 				UMassReplicationSubsystem& ReplicationSubsystem = Context.GetMutableSubsystemChecked<UMassReplicationSubsystem>();
 
 				const TArrayView<FMassNetworkIDFragment> NetworkIDList = Context.GetMutableFragmentView<FMassNetworkIDFragment>();
-				const int32 NumEntities = Context.GetNumEntities();
-
-				for (int32 Idx = 0; Idx < NumEntities; ++Idx)
+				for (FMassExecutionContext::FEntityIterator EntityIt = Context.CreateEntityIterator(); EntityIt; ++EntityIt)
 				{
-					NetworkIDList[Idx].NetID = ReplicationSubsystem.GetNextAvailableMassNetID();
+					NetworkIDList[EntityIt].NetID = ReplicationSubsystem.GetNextAvailableMassNetID();
 				}
 			});
 #endif //UE_REPLICATION_COMPILE_SERVER_CODE
diff --git a/Engine/Plugins/Runtime/MassGameplay/Source/MassReplication/Private/MassReplicationGridProcessor.cpp b/Engine/Plugins/Runtime/MassGameplay/Source/MassReplication/Private/MassReplicationGridProcessor.cpp
--- a/Engine/Plugins/Runtime/MassGameplay/Source/MassReplication/Private/MassReplicationGridProcessor.cpp
+++ b/Engine/Plugins/Runtime/MassGameplay/Source/MassReplication/Private/MassReplicationGridProcessor.cpp
@@ -57,22 +57,21 @@ void UMassReplicationGridProcessor::Execute(FMassEntityManager& EntityManager, F
 	{
 		UMassReplicationSubsystem& ReplicationSubsystem = Context.GetMutableSubsystemChecked<UMassReplicationSubsystem>();
 		FReplicationHashGrid2D& ReplicationGrid = ReplicationSubsystem.GetGridMutable();
-		const int32 NumEntities = Context.GetNumEntities();
 
 		TConstArrayView<FTransformFragment> LocationList = Context.GetFragmentView<FTransformFragment>();
 		TConstArrayView<FAgentRadiusFragment> RadiusList = Context.GetFragmentView<FAgentRadiusFragment>();
 		TArrayView<FMassReplicationGridCellLocationFragment> ReplicationCellLocationList = Context.GetMutableFragmentView<FMassReplicationGridCellLocationFragment>();
 
-		for (int32 EntityIndex = 0; EntityIndex < NumEntities; ++EntityIndex)
+		for (FMassExecutionContext::FEntityIterator EntityIt = Context.CreateEntityIterator(); EntityIt; ++EntityIt)
 		{
 			// Add to the grid
-			const FVector NewPos = LocationList[EntityIndex].GetTransform().GetLocation();
+			const FVector NewPos = LocationList[EntityIt].GetTransform().GetLocation();
 			// note that 0-radius is fine, the underlying THierarchicalHashGrid2D supports that just fine
-			const float Radius = RadiusList.IsEmpty() ? 0.f : RadiusList[EntityIndex].Radius;
+			const float Radius = RadiusList.IsEmpty() ? 0.f : RadiusList[EntityIt].Radius;
 
-			const FMassEntityHandle EntityHandle = Context.GetEntity(EntityIndex);
+			const FMassEntityHandle EntityHandle = Context.GetEntity(EntityIt);
 			const FBox NewBounds(NewPos - FVector(Radius, Radius, 0.f), NewPos + FVector(Radius, Radius, 0.f));
-			ReplicationCellLocationList[EntityIndex].CellLoc = ReplicationGrid.Add(EntityHandle, NewBounds);
+			ReplicationCellLocationList[EntityIt].CellLoc = ReplicationGrid.Add(EntityHandle, NewBounds);
 
 			Context.Defer().AddTag<FMassInReplicationGridTag>(EntityHandle);
 		}
@@ -82,27 +81,25 @@ void UMassReplicationGridProcessor::Execute(FMassEntityManager& EntityManager, F
 	{
 		UMassReplicationSubsystem& ReplicationSubsystem = Context.GetMutableSubsystemChecked<UMassReplicationSubsystem>();
 		FReplicationHashGrid2D& ReplicationGrid = ReplicationSubsystem.GetGridMutable();
-		const int32 NumEntities = Context.GetNumEntities();
-
 		TConstArrayView<FTransformFragment> LocationList = Context.GetFragmentView<FTransformFragment>();
 		TConstArrayView<FAgentRadiusFragment> RadiusList = Context.GetFragmentView<FAgentRadiusFragment>();
 		TArrayView<FMassReplicationGridCellLocationFragment> ReplicationCellLocationList = Context.GetMutableFragmentView<FMassReplicationGridCellLocationFragment>();
 
-		for (int32 EntityIndex = 0; EntityIndex < NumEntities; ++EntityIndex)
+		for (FMassExecutionContext::FEntityIterator EntityIt = Context.CreateEntityIterator(); EntityIt; ++EntityIt)
 		{
 			// Update position in grid
-			const FVector NewPos = LocationList[EntityIndex].GetTransform().GetLocation();
+			const FVector NewPos = LocationList[EntityIt].GetTransform().GetLocation();
 			// note that 0-radius is fine, the underlying THierarchicalHashGrid2D supports that just fine
-			const float Radius = RadiusList.IsEmpty() ? 0.f : RadiusList[EntityIndex].Radius;
-			const FMassEntityHandle EntityHandle = Context.GetEntity(EntityIndex);
+			const float Radius = RadiusList.IsEmpty() ? 0.f : RadiusList[EntityIt].Radius;
+			const FMassEntityHandle EntityHandle = Context.GetEntity(EntityIt);
 			const FBox NewBounds(NewPos - FVector(Radius, Radius, 0.f), NewPos + FVector(Radius, Radius, 0.f));
-			ReplicationCellLocationList[EntityIndex].CellLoc = ReplicationGrid.Move(EntityHandle, ReplicationCellLocationList[EntityIndex].CellLoc, NewBounds);
+			ReplicationCellLocationList[EntityIt].CellLoc = ReplicationGrid.Move(EntityHandle, ReplicationCellLocationList[EntityIt].CellLoc, NewBounds);
 
 #if WITH_MASSGAMEPLAY_DEBUG && 0
 			const FDebugContext BaseDebugContext(this, LogMassReplication, nullptr, EntityHandle);
 			if (DebugIsSelected(EntityHandle))
 			{
-				FBox Box = ReplicationGrid.CalcCellBounds(ReplicationCellLocationList[EntityIndex].CellLoc);
+				FBox Box = ReplicationGrid.CalcCellBounds(ReplicationCellLocationList[EntityIt].CellLoc);
 				Box.Max.Z += 200.f;
 				DebugDrawBox(BaseDebugContext, Box, FColor::Yellow);
 			}
@@ -114,14 +111,13 @@ void UMassReplicationGridProcessor::Execute(FMassEntityManager& EntityManager, F
 	{
 		UMassReplicationSubsystem& ReplicationSubsystem = Context.GetMutableSubsystemChecked<UMassReplicationSubsystem>();
 		FReplicationHashGrid2D& ReplicationGrid = ReplicationSubsystem.GetGridMutable();
-		const int32 NumEntities = Context.GetNumEntities();
 
 		TArrayView<FMassReplicationGridCellLocationFragment> ReplicationCellLocationList = Context.GetMutableFragmentView<FMassReplicationGridCellLocationFragment>();
-		for (int32 EntityIndex = 0; EntityIndex < NumEntities; ++EntityIndex)
+		for (FMassExecutionContext::FEntityIterator EntityIt = Context.CreateEntityIterator(); EntityIt; ++EntityIt)
 		{
-			const FMassEntityHandle EntityHandle = Context.GetEntity(EntityIndex);
-			ReplicationGrid.Remove(EntityHandle, ReplicationCellLocationList[EntityIndex].CellLoc);
-			ReplicationCellLocationList[EntityIndex].CellLoc = FReplicationHashGrid2D::FCellLocation();
+			const FMassEntityHandle EntityHandle = Context.GetEntity(EntityIt);
+			ReplicationGrid.Remove(EntityHandle, ReplicationCellLocationList[EntityIt].CellLoc);
+			ReplicationCellLocationList[EntityIt].CellLoc = FReplicationHashGrid2D::FCellLocation();
 
 			Context.Defer().RemoveTag<FMassInReplicationGridTag>(EntityHandle);
 		}
@@ -151,14 +147,13 @@ void UMassReplicationGridRemoverProcessor::Execute(FMassEntityManager& EntityMan
 	{
 		UMassReplicationSubsystem& ReplicationSubsystem = Context.GetMutableSubsystemChecked<UMassReplicationSubsystem>();
 		FReplicationHashGrid2D& ReplicationGrid = ReplicationSubsystem.GetGridMutable();
-		const int32 NumEntities = Context.GetNumEntities();
 		const TArrayView<FMassReplicationGridCellLocationFragment> ReplicationCellLocationList = Context.GetMutableFragmentView<FMassReplicationGridCellLocationFragment>();
 
-		for (int32 EntityIndex = 0; EntityIndex < NumEntities; ++EntityIndex)
+		for (FMassExecutionContext::FEntityIterator EntityIt = Context.CreateEntityIterator(); EntityIt; ++EntityIt)
 		{
-			const FMassEntityHandle EntityHandle = Context.GetEntity(EntityIndex);
-			ReplicationGrid.Remove(EntityHandle, ReplicationCellLocationList[EntityIndex].CellLoc);
-			ReplicationCellLocationList[EntityIndex].CellLoc = FReplicationHashGrid2D::FCellLocation();
+			const FMassEntityHandle EntityHandle = Context.GetEntity(EntityIt);
+			ReplicationGrid.Remove(EntityHandle, ReplicationCellLocationList[EntityIt].CellLoc);
+			ReplicationCellLocationList[EntityIt].CellLoc = FReplicationHashGrid2D::FCellLocation();
 		}
 	});
 }
\ No newline at end of file
diff --git a/Engine/Plugins/Runtime/MassGameplay/Source/MassReplication/Private/MassReplicationProcessor.cpp b/Engine/Plugins/Runtime/MassGameplay/Source/MassReplication/Private/MassReplicationProcessor.cpp
--- a/Engine/Plugins/Runtime/MassGameplay/Source/MassReplication/Private/MassReplicationProcessor.cpp
+++ b/Engine/Plugins/Runtime/MassGameplay/Source/MassReplication/Private/MassReplicationProcessor.cpp
@@ -275,12 +275,11 @@ void UMassReplicationProcessor::Execute(FMassEntityManager& EntityManager, FMass
 						const TArrayView<FMassReplicationLODFragment> ViewerLODList = Context.GetMutableFragmentView<FMassReplicationLODFragment>();
 						TArrayView<FMassReplicatedAgentFragment> ReplicatedAgentList = Context.GetMutableFragmentView<FMassReplicatedAgentFragment>();
 
-						const int32 NumEntities = Context.GetNumEntities();
-						for (int EntityIdx = 0; EntityIdx < NumEntities; EntityIdx++)
+						for (FMassExecutionContext::FEntityIterator EntityIt = Context.CreateEntityIterator(); EntityIt; ++EntityIt)
 						{
-							FMassEntityHandle EntityHandle = Context.GetEntity(EntityIdx);
-							FMassReplicatedAgentFragment& AgentFragment = ReplicatedAgentList[EntityIdx];
-							FMassReplicationLODFragment& LODFragment = ViewerLODList[EntityIdx];
+							FMassEntityHandle EntityHandle = Context.GetEntity(EntityIt);
+							FMassReplicatedAgentFragment& AgentFragment = ReplicatedAgentList[EntityIt];
+							FMassReplicationLODFragment& LODFragment = ViewerLODList[EntityIt];
 
 							if (FMassReplicatedAgentData* AgentData = ClientReplicationInfo.AgentsData.Find(EntityHandle))
 							{
@@ -367,11 +366,10 @@ void UMassReplicationProcessor::Execute(FMassEntityManager& EntityManager, FMass
 					{
 						TArrayView<FMassReplicatedAgentFragment> ReplicatedAgentList = Context.GetMutableFragmentView<FMassReplicatedAgentFragment>();
 
-						const int32 NumEntities = Context.GetNumEntities();
-						for (int EntityIdx = 0; EntityIdx < NumEntities; EntityIdx++)
+						for (FMassExecutionContext::FEntityIterator EntityIt = Context.CreateEntityIterator(); EntityIt; ++EntityIt)
 						{
-							FMassEntityHandle EntityHandle = Context.GetEntity(EntityIdx);
-							FMassReplicatedAgentFragment& AgentFragment = ReplicatedAgentList[EntityIdx];
+							FMassEntityHandle EntityHandle = Context.GetEntity(EntityIt);
+							FMassReplicatedAgentFragment& AgentFragment = ReplicatedAgentList[EntityIt];
 							ClientReplicationInfo.AgentsData.Add(EntityHandle, AgentFragment.AgentData);
 						}
 					});
diff --git a/Engine/Plugins/Runtime/MassGameplay/Source/MassReplication/Public/MassClientBubbleHandler.h b/Engine/Plugins/Runtime/MassGameplay/Source/MassReplication/Public/MassClientBubbleHandler.h
--- a/Engine/Plugins/Runtime/MassGameplay/Source/MassReplication/Public/MassClientBubbleHandler.h
+++ b/Engine/Plugins/Runtime/MassGameplay/Source/MassReplication/Public/MassClientBubbleHandler.h
@@ -543,24 +543,25 @@ void TClientBubbleHandlerBase<AgentArrayItem>::PostReplicatedAddEntitiesHelper(c
 		int32 AgentsSpawnIdx = 0;
 
 		Query.ForEachEntityChunk(FMassArchetypeEntityCollection(ArchetypeHandle, Entities, FMassArchetypeEntityCollection::NoDuplicates)
-								, ExecContext, [&AgentsSpawn, &AgentsSpawnIdx, this, ReplicationSubsystem, &ExecContext, &CacheFragmentViewsForSpawnQuery, &SetSpawnedEntityData, &EntityManager](FMassExecutionContext& Context)
+								, ExecContext, [&AgentsSpawn, &AgentsSpawnIdx, this, ReplicationSubsystem, &ExecContext, &CacheFragmentViewsForSpawnQuery
+								, &SetSpawnedEntityData, &EntityManager](FMassExecutionContext& Context)
 			{
 				CacheFragmentViewsForSpawnQuery(ExecContext);
 
 				const TArrayView<FMassNetworkIDFragment> NetworkIDList = Context.GetMutableFragmentView<FMassNetworkIDFragment>();
 
-				for (int32 i = 0; i < Context.GetNumEntities(); ++i)
+				for (FMassExecutionContext::FEntityIterator EntityIt = Context.CreateEntityIterator(); EntityIt; ++EntityIt)
 				{
 					const typename AgentArrayItem::FReplicatedAgentType& AgentSpawn = *AgentsSpawn[AgentsSpawnIdx];
-					const FMassEntityHandle Entity = Context.GetEntity(i);
+					const FMassEntityHandle Entity = Context.GetEntity(EntityIt);
 
-					FMassNetworkIDFragment& NetIDFragment = NetworkIDList[i];
+					FMassNetworkIDFragment& NetIDFragment = NetworkIDList[EntityIt];
 
 					NetIDFragment.NetID = AgentSpawn.GetNetID();
 					ReplicationSubsystem->SetEntity(NetIDFragment.NetID, Entity);
 
 					FMassEntityView EntityView(EntityManager, Entity);
-					SetSpawnedEntityData(EntityView, AgentSpawn, i);
+					SetSpawnedEntityData(EntityView, AgentSpawn, EntityIt);
 
 					++AgentsSpawnIdx;
 				}
diff --git a/Engine/Plugins/Runtime/MassGameplay/Source/MassReplication/Public/MassReplicationProcessor.h b/Engine/Plugins/Runtime/MassGameplay/Source/MassReplication/Public/MassReplicationProcessor.h
--- a/Engine/Plugins/Runtime/MassGameplay/Source/MassReplication/Public/MassReplicationProcessor.h
+++ b/Engine/Plugins/Runtime/MassGameplay/Source/MassReplication/Public/MassReplicationProcessor.h
@@ -95,9 +95,6 @@ template<typename AgentArrayItem, typename CacheViewsCallback, typename AddEntit
 void UMassReplicatorBase::CalculateClientReplication(FMassExecutionContext& Context, FMassReplicationContext& ReplicationContext, CacheViewsCallback&& CacheViews, AddEntityCallback&& AddEntity, ModifyEntityCallback&& ModifyEntity, RemoveEntityCallback&& RemoveEntity)
 {
 #if UE_REPLICATION_COMPILE_SERVER_CODE
-
-	const int32 NumEntities = Context.GetNumEntities();
-
 	TConstArrayView<FMassNetworkIDFragment> NetworkIDList = Context.GetFragmentView<FMassNetworkIDFragment>();
 	TArrayView<FMassReplicationLODFragment> ViewerLODList = Context.GetMutableFragmentView<FMassReplicationLODFragment>();
 	TArrayView<FMassReplicatedAgentFragment> ReplicatedAgentList = Context.GetMutableFragmentView<FMassReplicatedAgentFragment>();
@@ -108,9 +105,9 @@ void UMassReplicatorBase::CalculateClientReplication(FMassExecutionContext& Cont
 
 	const double Time = ReplicationContext.World.GetRealTimeSeconds();
 
-	for (int EntityIdx = 0; EntityIdx < NumEntities; EntityIdx++)
+	for (FMassExecutionContext::FEntityIterator EntityIt = Context.CreateEntityIterator(); EntityIt; ++EntityIt)
 	{
-		FMassReplicatedAgentFragment& AgentFragment = ReplicatedAgentList[EntityIdx];
+		FMassReplicatedAgentFragment& AgentFragment = ReplicatedAgentList[EntityIt];
 
 		const FMassClientHandle& ClientHandle = RepSharedFragment.CurrentClientHandle;
 		check(ClientHandle.IsValid());
@@ -119,7 +116,7 @@ void UMassReplicatorBase::CalculateClientReplication(FMassExecutionContext& Cont
 
 		FMassReplicatedAgentData& AgentData = AgentFragment.AgentData;
 
-		const EMassLOD::Type LOD = ViewerLODList[EntityIdx].LOD;
+		const EMassLOD::Type LOD = ViewerLODList[EntityIt].LOD;
 
 		if (LOD < EMassLOD::Off)
 		{
@@ -130,19 +127,19 @@ void UMassReplicatorBase::CalculateClientReplication(FMassExecutionContext& Cont
 			{
 				typename AgentArrayItem::FReplicatedAgentType ReplicatedAgent;
 
-				const FMassNetworkIDFragment& NetIDFragment = NetworkIDList[EntityIdx];
-				const FReplicationTemplateIDFragment& TemplateIDFragment = TemplateIDList[EntityIdx];
+				const FMassNetworkIDFragment& NetIDFragment = NetworkIDList[EntityIt];
+				const FReplicationTemplateIDFragment& TemplateIDFragment = TemplateIDList[EntityIt];
 
 				ReplicatedAgent.SetNetID(NetIDFragment.NetID);
 				ReplicatedAgent.SetTemplateID(TemplateIDFragment.ID);
 
-				AgentData.Handle = AddEntity(Context, EntityIdx, ReplicatedAgent, ClientHandle);
+				AgentData.Handle = AddEntity(Context, EntityIt, ReplicatedAgent, ClientHandle);
 
 				AgentData.LastUpdateTime = Time;
 			}
 			else
 			{
-				ModifyEntity(Context, EntityIdx, LOD, Time, AgentData.Handle, ClientHandle);
+				ModifyEntity(Context, EntityIt, LOD, Time, AgentData.Handle, ClientHandle);
 			}
 		}
 		else
diff --git a/Engine/Plugins/Runtime/MassGameplay/Source/MassRepresentation/Private/MassRepresentationDebug.cpp b/Engine/Plugins/Runtime/MassGameplay/Source/MassRepresentation/Private/MassRepresentationDebug.cpp
--- a/Engine/Plugins/Runtime/MassGameplay/Source/MassRepresentation/Private/MassRepresentationDebug.cpp
+++ b/Engine/Plugins/Runtime/MassGameplay/Source/MassRepresentation/Private/MassRepresentationDebug.cpp
@@ -36,14 +36,13 @@ namespace UE::Mass::Representation::Debug
 	void DebugDisplayRepresentation(FMassExecutionContext& Context, TConstArrayView<FMassRepresentationLODFragment> RepresentationLODList, TConstArrayView<FMassRepresentationFragment> RepresentationList, TConstArrayView<FTransformFragment> LocationList, UWorld* World)
 	{
 #if UE_ENABLE_DEBUG_DRAWING
-		const int32 NumEntities = Context.GetNumEntities();
-		for (int EntityIdx = 0; EntityIdx < NumEntities; EntityIdx++)
+		for (FMassExecutionContext::FEntityIterator EntityIt = Context.CreateEntityIterator(); EntityIt; ++EntityIt)
 		{
-			const FMassRepresentationLODFragment& RepresentationLOD = RepresentationLODList[EntityIdx];
+			const FMassRepresentationLODFragment& RepresentationLOD = RepresentationLODList[EntityIt];
 			if (RepresentationLOD.LODSignificance <= UE::Mass::Representation::Debug::DebugRepresentationMaxSignificance)
 			{
-				const FMassRepresentationFragment& Representation = RepresentationList[EntityIdx];
-				const FTransformFragment& EntityLocation = LocationList[EntityIdx];
+				const FMassRepresentationFragment& Representation = RepresentationList[EntityIt];
+				const FTransformFragment& EntityLocation = LocationList[EntityIt];
 				int32 CurrentRepresentationIdx = (int32)Representation.CurrentRepresentation;
 				DrawDebugSolidBox(World, EntityLocation.GetTransform().GetLocation() + FVector(0.0f, 0.0f, 150.0f), FVector(25.0f), UE::Mass::Representation::Debug::RepresentationColors[CurrentRepresentationIdx]);
 			}
@@ -54,19 +53,18 @@ namespace UE::Mass::Representation::Debug
 	void VisLogRepresentation(FMassExecutionContext& Context, TConstArrayView<FMassRepresentationLODFragment> RepresentationLODList, TConstArrayView<FMassRepresentationFragment> RepresentationList, TConstArrayView<FTransformFragment> LocationList, UObject* LogOwner)
 	{
 #if ENABLE_VISUAL_LOG
-		const int32 NumEntities = Context.GetNumEntities();
-		for (int EntityIdx = 0; EntityIdx < NumEntities; EntityIdx++)
+		for (FMassExecutionContext::FEntityIterator EntityIt = Context.CreateEntityIterator(); EntityIt; ++EntityIt)
 		{
-			const FMassRepresentationLODFragment& RepresentationLOD = RepresentationLODList[EntityIdx];
+			const FMassRepresentationLODFragment& RepresentationLOD = RepresentationLODList[EntityIt];
 			if (RepresentationLOD.LODSignificance <= UE::Mass::Representation::Debug::DebugRepresentationMaxSignificance)
 			{
-				const FTransformFragment& EntityLocation = LocationList[EntityIdx];
-				const FMassRepresentationFragment& Representation = RepresentationList[EntityIdx];
+				const FTransformFragment& EntityLocation = LocationList[EntityIt];
+				const FMassRepresentationFragment& Representation = RepresentationList[EntityIt];
 				int32 CurrentRepresentationIdx = (int32)Representation.CurrentRepresentation;
 				int32 PrevRepresentationIdx = (int32)Representation.PrevRepresentation;
 				// Add 20cm +Z offset to draw above mass.debug.RepresentationLOD so they can be viewed together
-				UE_CVLOG_LOCATION(Representation.CurrentRepresentation == Representation.PrevRepresentation, LogOwner, LogMassRepresentation, Verbose, EntityLocation.GetTransform().GetLocation() + FVector(0.0f, 0.0f, 20.0f), 20, UE::Mass::Representation::Debug::RepresentationColors[CurrentRepresentationIdx], TEXT("%s %d"), *Context.GetEntity(EntityIdx).DebugGetDescription(), CurrentRepresentationIdx);
-				UE_CVLOG_LOCATION(Representation.CurrentRepresentation != Representation.PrevRepresentation, LogOwner, LogMassRepresentation, Verbose, EntityLocation.GetTransform().GetLocation() + FVector(0.0f, 0.0f, 20.0f), 20, UE::Mass::Representation::Debug::RepresentationColors[CurrentRepresentationIdx], TEXT("%s %d -> %d"), *Context.GetEntity(EntityIdx).DebugGetDescription(), PrevRepresentationIdx, CurrentRepresentationIdx);
+				UE_CVLOG_LOCATION(Representation.CurrentRepresentation == Representation.PrevRepresentation, LogOwner, LogMassRepresentation, Verbose, EntityLocation.GetTransform().GetLocation() + FVector(0.0f, 0.0f, 20.0f), 20, UE::Mass::Representation::Debug::RepresentationColors[CurrentRepresentationIdx], TEXT("%s %d"), *Context.GetEntity(EntityIt).DebugGetDescription(), CurrentRepresentationIdx);
+				UE_CVLOG_LOCATION(Representation.CurrentRepresentation != Representation.PrevRepresentation, LogOwner, LogMassRepresentation, Verbose, EntityLocation.GetTransform().GetLocation() + FVector(0.0f, 0.0f, 20.0f), 20, UE::Mass::Representation::Debug::RepresentationColors[CurrentRepresentationIdx], TEXT("%s %d -> %d"), *Context.GetEntity(EntityIt).DebugGetDescription(), PrevRepresentationIdx, CurrentRepresentationIdx);
 			}
 		}
 #endif // ENABLE_VISUAL_LOG
diff --git a/Engine/Plugins/Runtime/MassGameplay/Source/MassRepresentation/Private/MassRepresentationProcessor.cpp b/Engine/Plugins/Runtime/MassGameplay/Source/MassRepresentation/Private/MassRepresentationProcessor.cpp
--- a/Engine/Plugins/Runtime/MassGameplay/Source/MassRepresentation/Private/MassRepresentationProcessor.cpp
+++ b/Engine/Plugins/Runtime/MassGameplay/Source/MassRepresentation/Private/MassRepresentationProcessor.cpp
@@ -76,14 +76,13 @@ void UMassRepresentationProcessor::UpdateRepresentation(FMassExecutionContext& C
 	// 		b. If EMassRepresentationType == StaticMeshInstance, sends an Actor disable request to MassRepresentationActorManagement;
 	//		c. If EMassRepresentationType == StaticMeshInstance, sends an Actor disable request to MassRepresentationActorManagement.
 	// 		NOTE: I guess this system assumes all instances are already represented by ISMs, which is why we're only dealing with Actor spawn/deactivation
-	const int32 NumEntities = Context.GetNumEntities();
-	for (int32 EntityIdx = 0; EntityIdx < NumEntities; EntityIdx++)
+	for (FMassExecutionContext::FEntityIterator EntityIt = Context.CreateEntityIterator(); EntityIt; ++EntityIt)
 	{
-		const FMassEntityHandle MassAgentEntityHandle = Context.GetEntity(EntityIdx);
-		const FTransformFragment& TransformFragment = TransformList[EntityIdx];
-		const FMassRepresentationLODFragment& RepresentationLOD = RepresentationLODList[EntityIdx];
-		FMassRepresentationFragment& Representation = RepresentationList[EntityIdx];
-		FMassActorFragment& ActorInfo = ActorList[EntityIdx];
+		const FMassEntityHandle MassAgentEntityHandle = Context.GetEntity(EntityIt);
+		const FTransformFragment& TransformFragment = TransformList[EntityIt];
+		const FMassRepresentationLODFragment& RepresentationLOD = RepresentationLODList[EntityIt];
+		FMassRepresentationFragment& Representation = RepresentationList[EntityIt];
+		FMassActorFragment& ActorInfo = ActorList[EntityIt];
 		AActor* Actor = ActorInfo.GetMutable();
 
 		// Keeping a copy of the that last calculated previous representation
@@ -160,7 +159,7 @@ void UMassRepresentationProcessor::UpdateRepresentation(FMassExecutionContext& C
 			}
 			if (Actor != nullptr)
 			{
-				RepresentationActorManagement->SetActorEnabled(EMassActorEnabledType::Disabled, *Actor, EntityIdx, Context.Defer());
+				RepresentationActorManagement->SetActorEnabled(EMassActorEnabledType::Disabled, *Actor, EntityIt, Context.Defer());
 			}
 		};
 
@@ -223,7 +222,7 @@ void UMassRepresentationProcessor::UpdateRepresentation(FMassExecutionContext& C
 							RepresentationActorManagement->TeleportActor(Representation.PrevTransform, *NewActor, Context.Defer());
 						}
 
-						RepresentationActorManagement->SetActorEnabled(bHighResActor ? EMassActorEnabledType::HighRes : EMassActorEnabledType::LowRes, *NewActor, EntityIdx, Context.Defer());
+						RepresentationActorManagement->SetActorEnabled(bHighResActor ? EMassActorEnabledType::HighRes : EMassActorEnabledType::LowRes, *NewActor, EntityIt, Context.Defer());
 						Representation.CurrentRepresentation = WantedRepresentationType;
 					}
 					else if (!Actor)
@@ -250,7 +249,7 @@ void UMassRepresentationProcessor::UpdateRepresentation(FMassExecutionContext& C
 					}
 					else
 					{
-						RepresentationActorManagement->SetActorEnabled(EMassActorEnabledType::Disabled, *Actor, EntityIdx, Context.Defer());
+						RepresentationActorManagement->SetActorEnabled(EMassActorEnabledType::Disabled, *Actor, EntityIt, Context.Defer());
 					}
 					Representation.CurrentRepresentation = EMassRepresentationType::None;
 					break;
@@ -349,12 +348,11 @@ void UMassVisualizationProcessor::UpdateVisualization(FMassExecutionContext& Con
 	const TArrayView<FMassRepresentationFragment> RepresentationList = Context.GetMutableFragmentView<FMassRepresentationFragment>();
 	const TConstArrayView<FMassRepresentationLODFragment> RepresentationLODList = Context.GetFragmentView<FMassRepresentationLODFragment>();
 
-	const int32 NumEntities = Context.GetNumEntities();
-	for (int32 EntityIdx = 0; EntityIdx < NumEntities; EntityIdx++)
+	for (FMassExecutionContext::FEntityIterator EntityIt = Context.CreateEntityIterator(); EntityIt; ++EntityIt)
 	{
-		const FMassEntityHandle Entity = Context.GetEntity(EntityIdx);
-		FMassRepresentationFragment& Representation = RepresentationList[EntityIdx];
-		const FMassRepresentationLODFragment& RepresentationLOD = RepresentationLODList[EntityIdx];
+		const FMassEntityHandle Entity = Context.GetEntity(EntityIt);
+		FMassRepresentationFragment& Representation = RepresentationList[EntityIt];
+		const FMassRepresentationLODFragment& RepresentationLOD = RepresentationLODList[EntityIt];
 		UpdateEntityVisibility(Entity, Representation, RepresentationLOD, ChunkData, Context.Defer());
 	}
 }
@@ -471,13 +469,12 @@ void UMassRepresentationFragmentDestructor::Execute(FMassEntityManager& EntityMa
 		const TArrayView<FMassRepresentationFragment> RepresentationList = Context.GetMutableFragmentView<FMassRepresentationFragment>();
 		const TArrayView<FMassActorFragment> ActorList = Context.GetMutableFragmentView<FMassActorFragment>();
 
-		const int32 NumEntities = Context.GetNumEntities();
-		for (int32 i = 0; i < NumEntities; ++i)
+		for (FMassExecutionContext::FEntityIterator EntityIt = Context.CreateEntityIterator(); EntityIt; ++EntityIt)
 		{
-			FMassRepresentationFragment& Representation = RepresentationList[i];
-			FMassActorFragment& ActorInfo = ActorList[i];
+			FMassRepresentationFragment& Representation = RepresentationList[EntityIt];
+			FMassActorFragment& ActorInfo = ActorList[EntityIt];
 
-			const FMassEntityHandle MassAgentEntityHandle = Context.GetEntity(i);
+			const FMassEntityHandle MassAgentEntityHandle = Context.GetEntity(EntityIt);
 
 			UMassRepresentationActorManagement::ReleaseAnyActorOrCancelAnySpawning(*RepresentationSubsystem, MassAgentEntityHandle, ActorInfo, Representation, ActorSubsystem);
 		}
diff --git a/Engine/Plugins/Runtime/MassGameplay/Source/MassRepresentation/Private/MassStationaryISMRepresentationFragmentDestructor.cpp b/Engine/Plugins/Runtime/MassGameplay/Source/MassRepresentation/Private/MassStationaryISMRepresentationFragmentDestructor.cpp
--- a/Engine/Plugins/Runtime/MassGameplay/Source/MassRepresentation/Private/MassStationaryISMRepresentationFragmentDestructor.cpp
+++ b/Engine/Plugins/Runtime/MassGameplay/Source/MassRepresentation/Private/MassStationaryISMRepresentationFragmentDestructor.cpp
@@ -35,16 +35,15 @@ void UMassStationaryISMRepresentationFragmentDestructor::Execute(FMassEntityMana
 
 		const TArrayView<FMassRepresentationFragment> RepresentationList = Context.GetMutableFragmentView<FMassRepresentationFragment>();
 
-		const int32 NumEntities = Context.GetNumEntities();
-		for (int32 EntityIdx = 0; EntityIdx < NumEntities; ++EntityIdx)
+		for (FMassExecutionContext::FEntityIterator EntityIt = Context.CreateEntityIterator(); EntityIt; ++EntityIt)
 		{
-			FMassRepresentationFragment& Representation = RepresentationList[EntityIdx];
+			FMassRepresentationFragment& Representation = RepresentationList[EntityIt];
 			if (Representation.CurrentRepresentation == EMassRepresentationType::StaticMeshInstance)
 			{
 				FMassInstancedStaticMeshInfo& ISMInfo = ISMInfosView[Representation.StaticMeshDescHandle.ToIndex()];
 				if (FMassLODSignificanceRange* OldRange = ISMInfo.GetLODSignificanceRange(Representation.PrevLODSignificance))
 				{
-					const FMassEntityHandle EntityHandle = Context.GetEntity(EntityIdx);
+					const FMassEntityHandle EntityHandle = Context.GetEntity(EntityIt);
 					if (OldRange)
 					{
 						OldRange->RemoveInstance(EntityHandle);
diff --git a/Engine/Plugins/Runtime/MassGameplay/Source/MassRepresentation/Private/MassStationaryISMSwitcherProcessor.cpp b/Engine/Plugins/Runtime/MassGameplay/Source/MassRepresentation/Private/MassStationaryISMSwitcherProcessor.cpp
--- a/Engine/Plugins/Runtime/MassGameplay/Source/MassRepresentation/Private/MassStationaryISMSwitcherProcessor.cpp
+++ b/Engine/Plugins/Runtime/MassGameplay/Source/MassRepresentation/Private/MassStationaryISMSwitcherProcessor.cpp
@@ -49,13 +49,12 @@ void UMassStationaryISMSwitcherProcessor::ProcessContext(FMassExecutionContext&
 	const FMassRepresentationParameters& RepresentationParams = Context.GetConstSharedFragment<FMassRepresentationParameters>();
 	const bool bDoKeepActorExtraFrame = UE::Mass::Representation::bAllowKeepActorExtraFrame ? RepresentationParams.bKeepLowResActors : false;
 
-	const int32 NumEntities = Context.GetNumEntities();
-	for (int32 EntityIdx = 0; EntityIdx < NumEntities; EntityIdx++)
+	for (FMassExecutionContext::FEntityIterator EntityIt = Context.CreateEntityIterator(); EntityIt; ++EntityIt)
 	{
-		const FMassEntityHandle EntityHandle = Context.GetEntity(EntityIdx);
-		const FMassRepresentationLODFragment& RepresentationLOD = RepresentationLODList[EntityIdx];
-		const FTransformFragment& TransformFragment = TransformList[EntityIdx];
-		FMassRepresentationFragment& Representation = RepresentationList[EntityIdx];
+		const FMassEntityHandle EntityHandle = Context.GetEntity(EntityIt);
+		const FMassRepresentationLODFragment& RepresentationLOD = RepresentationLODList[EntityIt];
+		const FTransformFragment& TransformFragment = TransformList[EntityIt];
+		FMassRepresentationFragment& Representation = RepresentationList[EntityIt];
 
 		if (Representation.PrevRepresentation != EMassRepresentationType::StaticMeshInstance
 			&& Representation.CurrentRepresentation != EMassRepresentationType::StaticMeshInstance)
diff --git a/Engine/Plugins/Runtime/MassGameplay/Source/MassRepresentation/Private/MassUpdateISMProcessor.cpp b/Engine/Plugins/Runtime/MassGameplay/Source/MassRepresentation/Private/MassUpdateISMProcessor.cpp
--- a/Engine/Plugins/Runtime/MassGameplay/Source/MassRepresentation/Private/MassUpdateISMProcessor.cpp
+++ b/Engine/Plugins/Runtime/MassGameplay/Source/MassRepresentation/Private/MassUpdateISMProcessor.cpp
@@ -44,16 +44,15 @@ void UMassUpdateISMProcessor::Execute(FMassEntityManager& EntityManager, FMassEx
 		const TArrayView<FMassRepresentationFragment> RepresentationList = Context.GetMutableFragmentView<FMassRepresentationFragment>();
 		const TConstArrayView<FMassRepresentationLODFragment> RepresentationLODList = Context.GetFragmentView<FMassRepresentationLODFragment>();
 
-		const int32 NumEntities = Context.GetNumEntities();
-		for (int32 EntityIdx = 0; EntityIdx < NumEntities; EntityIdx++)
+		for (FMassExecutionContext::FEntityIterator EntityIt = Context.CreateEntityIterator(); EntityIt; ++EntityIt)
 		{
-			const FTransformFragment& TransformFragment = TransformList[EntityIdx];
-			const FMassRepresentationLODFragment& RepresentationLOD = RepresentationLODList[EntityIdx];
-			FMassRepresentationFragment& Representation = RepresentationList[EntityIdx];
+			const FTransformFragment& TransformFragment = TransformList[EntityIt];
+			const FMassRepresentationLODFragment& RepresentationLOD = RepresentationLODList[EntityIt];
+			FMassRepresentationFragment& Representation = RepresentationList[EntityIt];
 
 			if (Representation.CurrentRepresentation == EMassRepresentationType::StaticMeshInstance)
 			{
-				UpdateISMTransform(Context.GetEntity(EntityIdx), ISMInfo[Representation.StaticMeshDescHandle.ToIndex()], TransformFragment.GetTransform(), Representation.PrevTransform, RepresentationLOD.LODSignificance, Representation.PrevLODSignificance);
+				UpdateISMTransform(Context.GetEntity(EntityIt), ISMInfo[Representation.StaticMeshDescHandle.ToIndex()], TransformFragment.GetTransform(), Representation.PrevTransform, RepresentationLOD.LODSignificance, Representation.PrevLODSignificance);
 			}
 			Representation.PrevTransform = TransformFragment.GetTransform();
 			Representation.PrevLODSignificance = RepresentationLOD.LODSignificance;
diff --git a/Engine/Plugins/Runtime/MassGameplay/Source/MassSmartObjects/Private/MassSmartObjectProcessor.cpp b/Engine/Plugins/Runtime/MassGameplay/Source/MassSmartObjects/Private/MassSmartObjectProcessor.cpp
--- a/Engine/Plugins/Runtime/MassGameplay/Source/MassSmartObjects/Private/MassSmartObjectProcessor.cpp
+++ b/Engine/Plugins/Runtime/MassGameplay/Source/MassSmartObjects/Private/MassSmartObjectProcessor.cpp
@@ -53,10 +53,10 @@ UMassSmartObjectCandidatesFinderProcessor::UMassSmartObjectCandidatesFinderProce
 	ExecutionOrder.ExecuteBefore.Add(UE::Mass::ProcessorGroupNames::Behavior);
 }
 
-void UMassSmartObjectCandidatesFinderProcessor::Execute(FMassEntityManager& EntityManager, FMassExecutionContext& Context)
+void UMassSmartObjectCandidatesFinderProcessor::Execute(FMassEntityManager& EntityManager, FMassExecutionContext& ExecutionContext)
 {
-	UMassSignalSubsystem& SignalSubsystem = Context.GetMutableSubsystemChecked<UMassSignalSubsystem>();
-	const UZoneGraphAnnotationSubsystem& AnnotationSubsystem = Context.GetSubsystemChecked<UZoneGraphAnnotationSubsystem>();
+	UMassSignalSubsystem& SignalSubsystem = ExecutionContext.GetMutableSubsystemChecked<UMassSignalSubsystem>();
+	const UZoneGraphAnnotationSubsystem& AnnotationSubsystem = ExecutionContext.GetSubsystemChecked<UZoneGraphAnnotationSubsystem>();
 	
 	// Create filter
 	FSmartObjectRequestFilter Filter;
@@ -86,7 +86,7 @@ void UMassSmartObjectCandidatesFinderProcessor::Execute(FMassEntityManager& Enti
 	};
 
 	// Process world location based requests
-	WorldRequestQuery.ForEachEntityChunk(Context, [this, &Filter, &EntitiesToSignal, &BeginRequestProcessing, &EndRequestProcessing](FMassExecutionContext& Context)
+	WorldRequestQuery.ForEachEntityChunk(ExecutionContext, [this, &Filter, &EntitiesToSignal, &BeginRequestProcessing, &EndRequestProcessing](FMassExecutionContext& Context)
 	{
 		const USmartObjectSubsystem& SmartObjectSubsystem = Context.GetSubsystemChecked<USmartObjectSubsystem>();
 
@@ -99,17 +99,17 @@ void UMassSmartObjectCandidatesFinderProcessor::Execute(FMassEntityManager& Enti
 		TArray<FSmartObjectRequestResult> QueryResults;
 		TArray<FSmartObjectCandidateSlot> SortedCandidateSlots;
 
-		for (int32 i = 0; i < NumEntities; ++i)
+		for (FMassExecutionContext::FEntityIterator EntityIt = Context.CreateEntityIterator(); EntityIt; ++EntityIt)
 		{
-			const FMassSmartObjectWorldLocationRequestFragment& RequestFragment = RequestList[i];
-			FMassSmartObjectRequestResultFragment& Result = ResultList[i];
+			const FMassSmartObjectWorldLocationRequestFragment& RequestFragment = RequestList[EntityIt];
+			FMassSmartObjectRequestResultFragment& Result = ResultList[EntityIt];
 			
 			EntitiesToSignal.Add(RequestFragment.RequestingEntity);
 
 			const FVector& SearchOrigin = RequestFragment.SearchOrigin;
 			const FBox& SearchBounds = FBox::BuildAABB(SearchOrigin, FVector(SearchExtents));
 
-			const FMassEntityHandle Entity = Context.GetEntity(i);
+			const FMassEntityHandle Entity = Context.GetEntity(EntityIt);
 			bool bDisplayDebug = false;
 			FColor DebugColor(FColor::White);
 
@@ -155,7 +155,7 @@ void UMassSmartObjectCandidatesFinderProcessor::Execute(FMassEntityManager& Enti
 	const FZoneGraphTag SmartObjectTag = GetDefault<UMassSmartObjectSettings>()->SmartObjectTag;
 	USmartObjectZoneAnnotations* Annotations = Cast<USmartObjectZoneAnnotations>(AnnotationSubsystem.GetFirstAnnotationForTag(SmartObjectTag));
 
-	LaneRequestQuery.ForEachEntityChunk(Context,
+	LaneRequestQuery.ForEachEntityChunk(ExecutionContext,
 		[&AnnotationSubsystem, Annotations, &Filter, SmartObjectTag, &EntitiesToSignal, &BeginRequestProcessing, &EndRequestProcessing](FMassExecutionContext& Context)
 		{
 #if WITH_MASSGAMEPLAY_DEBUG
@@ -173,17 +173,17 @@ void UMassSmartObjectCandidatesFinderProcessor::Execute(FMassEntityManager& Enti
 			FZoneGraphDataHandle LastUsedDataHandle;
 			const FSmartObjectAnnotationData* GraphData = nullptr;
 
-			for (int32 i = 0; i < NumEntities; ++i)
+			for (FMassExecutionContext::FEntityIterator EntityIt = Context.CreateEntityIterator(); EntityIt; ++EntityIt)
 			{
-				const FMassSmartObjectLaneLocationRequestFragment& RequestFragment = RequestList[i];
-				FMassSmartObjectRequestResultFragment& Result = ResultList[i];
+				const FMassSmartObjectLaneLocationRequestFragment& RequestFragment = RequestList[EntityIt];
+				FMassSmartObjectRequestResultFragment& Result = ResultList[EntityIt];
 				
 				EntitiesToSignal.Add(RequestFragment.RequestingEntity);
 
 				const FZoneGraphCompactLaneLocation RequestLocation = RequestFragment.CompactLaneLocation;
 				const FZoneGraphLaneHandle RequestLaneHandle = RequestLocation.LaneHandle;
 
-				const FMassEntityHandle Entity = Context.GetEntity(i);
+				const FMassEntityHandle Entity = Context.GetEntity(EntityIt);
 				bool bDisplayDebug = false;
 #if WITH_MASSGAMEPLAY_DEBUG
 				FColor DebugColor(FColor::White);
@@ -328,15 +328,14 @@ void UMassSmartObjectTimedBehaviorProcessor::Execute(FMassEntityManager& EntityM
 	{
 		const USmartObjectSubsystem& SmartObjectSubsystem = Context.GetSubsystemChecked<USmartObjectSubsystem>();
 
-		const int32 NumEntities = Context.GetNumEntities();
 		const TArrayView<FMassSmartObjectUserFragment> UserList = Context.GetMutableFragmentView<FMassSmartObjectUserFragment>();
 		const TArrayView<FMassSmartObjectTimedBehaviorFragment> TimedBehaviorFragments = Context.GetMutableFragmentView<FMassSmartObjectTimedBehaviorFragment>();
 
-		for (int32 i = 0; i < NumEntities; ++i)
+		for (FMassExecutionContext::FEntityIterator EntityIt = Context.CreateEntityIterator(); EntityIt; ++EntityIt)
 		{
-			FMassSmartObjectUserFragment& SOUser = UserList[i];
-			FMassSmartObjectTimedBehaviorFragment& TimedBehaviorFragment = TimedBehaviorFragments[i];
-			ensureMsgf(SOUser.InteractionStatus == EMassSmartObjectInteractionStatus::InProgress, TEXT("TimedBehavior fragment should only be present for in-progress interactions: %s"), *Context.GetEntity(i).DebugGetDescription());
+			FMassSmartObjectUserFragment& SOUser = UserList[EntityIt];
+			FMassSmartObjectTimedBehaviorFragment& TimedBehaviorFragment = TimedBehaviorFragments[EntityIt];
+			ensureMsgf(SOUser.InteractionStatus == EMassSmartObjectInteractionStatus::InProgress, TEXT("TimedBehavior fragment should only be present for in-progress interactions: %s"), *Context.GetEntity(EntityIt).DebugGetDescription());
 
 			const float DT = Context.GetDeltaTimeSeconds();
 			float& UseTime = TimedBehaviorFragment.UseTime;
@@ -344,7 +343,7 @@ void UMassSmartObjectTimedBehaviorProcessor::Execute(FMassEntityManager& EntityM
 			const bool bMustRelease = UseTime <= 0.f;
 
 #if WITH_MASSGAMEPLAY_DEBUG
-			const FMassEntityHandle Entity = Context.GetEntity(i);
+			const FMassEntityHandle Entity = Context.GetEntity(EntityIt);
 			FColor DebugColor(FColor::White);
 			const bool bIsDebuggingEntity = UE::Mass::Debug::IsDebuggingEntity(Entity, &DebugColor);
 			if (bIsDebuggingEntity)
@@ -368,7 +367,7 @@ void UMassSmartObjectTimedBehaviorProcessor::Execute(FMassEntityManager& EntityM
 			if (bMustRelease)
 			{
 				SOUser.InteractionStatus = EMassSmartObjectInteractionStatus::BehaviorCompleted;
-				ToRelease.Add(Context.GetEntity(i));
+				ToRelease.Add(Context.GetEntity(EntityIt));
 			}
 		}
 	});
@@ -402,12 +401,11 @@ void UMassSmartObjectUserFragmentDeinitializer::Execute(FMassEntityManager& Enti
 	EntityQuery.ForEachEntityChunk(Context, [this](FMassExecutionContext& Context)
 		{
 			USmartObjectSubsystem& SmartObjectSubsystem = Context.GetMutableSubsystemChecked<USmartObjectSubsystem>();
-			const int32 NumEntities = Context.GetNumEntities();
 			const TArrayView<FMassSmartObjectUserFragment> SmartObjectUserFragments = Context.GetMutableFragmentView<FMassSmartObjectUserFragment>();
 
-			for (int32 i = 0; i < NumEntities; ++i)
+			for (FMassExecutionContext::FEntityIterator EntityIt = Context.CreateEntityIterator(); EntityIt; ++EntityIt)
 			{
-				SmartObjectSubsystem.UnregisterSlotInvalidationCallback(SmartObjectUserFragments[i].InteractionHandle);
+				SmartObjectSubsystem.UnregisterSlotInvalidationCallback(SmartObjectUserFragments[EntityIt].InteractionHandle);
 			}
 		});
 }
\ No newline at end of file
diff --git a/Engine/Plugins/Runtime/MassGameplay/Source/MassSmartObjects/Private/MassSmartObjectRegistration.cpp b/Engine/Plugins/Runtime/MassGameplay/Source/MassSmartObjects/Private/MassSmartObjectRegistration.cpp
--- a/Engine/Plugins/Runtime/MassGameplay/Source/MassSmartObjects/Private/MassSmartObjectRegistration.cpp
+++ b/Engine/Plugins/Runtime/MassGameplay/Source/MassSmartObjects/Private/MassSmartObjectRegistration.cpp
@@ -31,12 +31,12 @@ void UMassSmartObjectInitializerBase::ConfigureQueries(const TSharedRef<FMassEnt
 	EntityQuery.AddRequirement<FSmartObjectRegistrationFragment>(EMassFragmentAccess::ReadWrite);
 }
 
-void UMassSmartObjectInitializerBase::Execute(FMassEntityManager& EntityManager, FMassExecutionContext& Context)
+void UMassSmartObjectInitializerBase::Execute(FMassEntityManager& EntityManager, FMassExecutionContext& ExecutionContext)
 {
-	UMassSignalSubsystem& SignalSubsystem = Context.GetMutableSubsystemChecked<UMassSignalSubsystem>();
+	UMassSignalSubsystem& SignalSubsystem = ExecutionContext.GetMutableSubsystemChecked<UMassSignalSubsystem>();
 	TArray<FMassEntityHandle> EntitiesToSignal;
 
-	EntityQuery.ForEachEntityChunk(Context,[&EntitiesToSignal](FMassExecutionContext& Context)
+	EntityQuery.ForEachEntityChunk(ExecutionContext,[&EntitiesToSignal](FMassExecutionContext& Context)
 	{
 		const TConstArrayView<FMassActorInstanceFragment> InstancedActorFragments = Context.GetFragmentView<FMassActorInstanceFragment>();
 		const TArrayView<FSmartObjectRegistrationFragment> RegistrationFragments = Context.GetMutableFragmentView<FSmartObjectRegistrationFragment>();
@@ -44,10 +44,10 @@ void UMassSmartObjectInitializerBase::Execute(FMassEntityManager& EntityManager,
 		const int32 NumEntities = Context.GetNumEntities();
 		EntitiesToSignal.Reserve(EntitiesToSignal.Num() + NumEntities);
 
-		for (int32 EntityIndex = 0; EntityIndex < NumEntities; ++EntityIndex)
+		for (FMassExecutionContext::FEntityIterator EntityIt = Context.CreateEntityIterator(); EntityIt; ++EntityIt)
 		{
-			FSmartObjectRegistrationFragment& RegistrationFragment = RegistrationFragments[EntityIndex];
-			const FMassActorInstanceFragment& InstancedActorFragment = InstancedActorFragments[EntityIndex];
+			FSmartObjectRegistrationFragment& RegistrationFragment = RegistrationFragments[EntityIt];
+			const FMassActorInstanceFragment& InstancedActorFragment = InstancedActorFragments[EntityIt];
 
 			checkf(!RegistrationFragment.Handle.IsValid(), TEXT("Should create smartobject only once."))
 			if (RegistrationFragment.Asset.Get() == nullptr)
@@ -60,7 +60,7 @@ void UMassSmartObjectInitializerBase::Execute(FMassEntityManager& EntityManager,
 				continue;
 			}
 
-			EntitiesToSignal.Add(Context.GetEntity(EntityIndex));
+			EntitiesToSignal.Add(Context.GetEntity(EntityIt));
 		}
 	});
 
@@ -86,18 +86,18 @@ void UMassSmartObjectDeinitializerBase::ConfigureQueries(const TSharedRef<FMassE
 	ProcessorRequirements.AddSubsystemRequirement<UMassSignalSubsystem>(EMassFragmentAccess::ReadWrite);
 }
 
-void UMassSmartObjectDeinitializerBase::Execute(FMassEntityManager& EntityManager, FMassExecutionContext& Context)
+void UMassSmartObjectDeinitializerBase::Execute(FMassEntityManager& EntityManager, FMassExecutionContext& ExecutionContext)
 {
-	UMassSignalSubsystem& SignalSubsystem = Context.GetMutableSubsystemChecked<UMassSignalSubsystem>();
+	UMassSignalSubsystem& SignalSubsystem = ExecutionContext.GetMutableSubsystemChecked<UMassSignalSubsystem>();
 	TArray<FMassEntityHandle> EntitiesToSignal;
 
-	EntityQuery.ForEachEntityChunk(Context, [&EntitiesToSignal](FMassExecutionContext& Context)
+	EntityQuery.ForEachEntityChunk(ExecutionContext, [&EntitiesToSignal](FMassExecutionContext& Context)
 	{
 		const int32 NumEntities = Context.GetNumEntities();
 		EntitiesToSignal.Reserve(EntitiesToSignal.Num() + NumEntities);
-		for (int32 EntityIndex = 0; EntityIndex < NumEntities; ++EntityIndex)
+		for (FMassExecutionContext::FEntityIterator EntityIt = Context.CreateEntityIterator(); EntityIt; ++EntityIt)
 		{
-			EntitiesToSignal.Add(Context.GetEntity(EntityIndex));
+			EntitiesToSignal.Add(Context.GetEntity(EntityIt));
 		}
 	});
 	
@@ -188,12 +188,11 @@ void UMassActiveSmartObjectSignalProcessor::SignalEntities(FMassEntityManager& E
 		const TConstArrayView<FMassActorInstanceFragment> InstancedActorFragments = Context.GetFragmentView<FMassActorInstanceFragment>();
 		const TArrayView<FSmartObjectRegistrationFragment> RegistrationFragments = Context.GetMutableFragmentView<FSmartObjectRegistrationFragment>();
 
-		const int32 NumEntities = Context.GetNumEntities();
-		for (int32 EntityIndex = 0; EntityIndex < NumEntities; ++EntityIndex)
+		for (FMassExecutionContext::FEntityIterator EntityIt = Context.CreateEntityIterator(); EntityIt; ++EntityIt)
 		{
-			FSmartObjectRegistrationFragment& RegistrationFragment = RegistrationFragments[EntityIndex];
-			const FMassActorInstanceFragment& InstancedActorFragment = InstancedActorFragments[EntityIndex];
-			const FTransformFragment& TransformFragment = TransformFragments[EntityIndex];
+			FSmartObjectRegistrationFragment& RegistrationFragment = RegistrationFragments[EntityIt];
+			const FMassActorInstanceFragment& InstancedActorFragment = InstancedActorFragments[EntityIt];
+			const FTransformFragment& TransformFragment = TransformFragments[EntityIt];
 
 			// Creation
 			if (!RegistrationFragment.Handle.IsValid() && InstancedActorFragment.Handle.IsValid())
@@ -222,10 +221,9 @@ void UMassActiveSmartObjectSignalProcessor::SignalEntities(FMassEntityManager& E
 		USmartObjectSubsystem& Subsystem = Context.GetMutableSubsystemChecked<USmartObjectSubsystem>();
 		const TArrayView<FSmartObjectRegistrationFragment> RegistrationFragments = Context.GetMutableFragmentView<FSmartObjectRegistrationFragment>();
 
-		const int32 NumEntities = Context.GetNumEntities();
-		for (int32 EntityIndex = 0; EntityIndex < NumEntities; ++EntityIndex)
+		for (FMassExecutionContext::FEntityIterator EntityIt = Context.CreateEntityIterator(); EntityIt; ++EntityIt)
 		{
-			FSmartObjectRegistrationFragment& RegistrationFragment = RegistrationFragments[EntityIndex];
+			FSmartObjectRegistrationFragment& RegistrationFragment = RegistrationFragments[EntityIt];
 
 			if (RegistrationFragment.Handle.IsValid())
 			{
diff --git a/Engine/Plugins/Runtime/MassGameplay/Source/MassSpawner/Private/MassSpawnLocationProcessor.cpp b/Engine/Plugins/Runtime/MassGameplay/Source/MassSpawner/Private/MassSpawnLocationProcessor.cpp
--- a/Engine/Plugins/Runtime/MassGameplay/Source/MassSpawner/Private/MassSpawnLocationProcessor.cpp
+++ b/Engine/Plugins/Runtime/MassGameplay/Source/MassSpawner/Private/MassSpawnLocationProcessor.cpp
@@ -24,9 +24,9 @@ void UMassSpawnLocationProcessor::ConfigureQueries(const TSharedRef<FMassEntityM
 	EntityQuery.AddRequirement<FTransformFragment>(EMassFragmentAccess::ReadWrite);
 }
 
-void UMassSpawnLocationProcessor::Execute(FMassEntityManager& EntityManager, FMassExecutionContext& Context)
+void UMassSpawnLocationProcessor::Execute(FMassEntityManager& EntityManager, FMassExecutionContext& ExecutionContext)
 {
-	if (!ensure(Context.ValidateAuxDataType<FMassTransformsSpawnData>()))
+	if (!ensure(ExecutionContext.ValidateAuxDataType<FMassTransformsSpawnData>()))
 	{
 		UE_VLOG_UELOG(this, LogMass, Log, TEXT("Execution context has invalid AuxData or it's not FMassSpawnAuxData. Entity transforms won't be initialized."));
 		return;
@@ -37,7 +37,7 @@ void UMassSpawnLocationProcessor::Execute(FMassEntityManager& EntityManager, FMa
 
 	const ENetMode NetMode = World->GetNetMode();
 
-	FMassTransformsSpawnData& AuxData = Context.GetMutableAuxData().GetMutable<FMassTransformsSpawnData>();
+	FMassTransformsSpawnData& AuxData = ExecutionContext.GetMutableAuxData().GetMutable<FMassTransformsSpawnData>();
 	TArray<FTransform>& Transforms = AuxData.Transforms;
 
 	const int32 NumSpawnTransforms = Transforms.Num();
@@ -48,7 +48,7 @@ void UMassSpawnLocationProcessor::Execute(FMassEntityManager& EntityManager, FMa
 	}
 
 	int32 NumRequiredSpawnTransforms = 0;
-	EntityQuery.ForEachEntityChunk(Context, [&NumRequiredSpawnTransforms](const FMassExecutionContext& Context)
+	EntityQuery.ForEachEntityChunk(ExecutionContext, [&NumRequiredSpawnTransforms](const FMassExecutionContext& Context)
 		{
 			NumRequiredSpawnTransforms += Context.GetNumEntities();
 		});
@@ -61,22 +61,21 @@ void UMassSpawnLocationProcessor::Execute(FMassEntityManager& EntityManager, FMa
 			NumSpawnTransforms, NumRequiredSpawnTransforms, NumToAdd);
 
 		Transforms.AddUninitialized(NumToAdd);
-		for (int i = 0; i < NumToAdd; ++i)
+		for (int EntityIndex = 0; EntityIndex < NumToAdd; ++EntityIndex)
 		{
-			Transforms[NumSpawnTransforms + i] = Transforms[RandomStream.RandRange(0, NumSpawnTransforms - 1)];
+			Transforms[NumSpawnTransforms + EntityIndex] = Transforms[RandomStream.RandRange(0, NumSpawnTransforms - 1)];
 		}
 	}
 
 	if (AuxData.bRandomize && !UE::Mass::Utils::IsDeterministic())
 	{
-		EntityQuery.ForEachEntityChunk(Context, [&Transforms, this](FMassExecutionContext& Context)
+		EntityQuery.ForEachEntityChunk(ExecutionContext, [&Transforms, this](FMassExecutionContext& Context)
 			{
 				const TArrayView<FTransformFragment> LocationList = Context.GetMutableFragmentView<FTransformFragment>();
-				const int32 NumEntities = Context.GetNumEntities();
-				for (int32 i = 0; i < NumEntities; ++i)
+				for (FMassExecutionContext::FEntityIterator EntityIt = Context.CreateEntityIterator(); EntityIt; ++EntityIt)
 				{
 					const int32 AuxIndex = RandomStream.RandRange(0, Transforms.Num() - 1);
-					LocationList[i].GetMutableTransform() = Transforms[AuxIndex];
+					LocationList[EntityIt].GetMutableTransform() = Transforms[AuxIndex];
 					Transforms.RemoveAtSwap(AuxIndex, EAllowShrinking::No);
 				}
 			});
@@ -84,7 +83,7 @@ void UMassSpawnLocationProcessor::Execute(FMassEntityManager& EntityManager, FMa
 	else
 	{
 		int32 NextTransformIndex = 0;
-		EntityQuery.ForEachEntityChunk(Context, [&Transforms, &NextTransformIndex, this](FMassExecutionContext& Context)
+		EntityQuery.ForEachEntityChunk(ExecutionContext, [&Transforms, &NextTransformIndex, this](FMassExecutionContext& Context)
 			{
 				const int32 NumEntities = Context.GetNumEntities();
 				TArrayView<FTransformFragment> LocationList = Context.GetMutableFragmentView<FTransformFragment>();
diff --git a/Engine/Source/Developer/MassEntityTestSuite/Private/MassEntityTestFarmPlot.cpp b/Engine/Source/Developer/MassEntityTestSuite/Private/MassEntityTestFarmPlot.cpp
--- a/Engine/Source/Developer/MassEntityTestSuite/Private/MassEntityTestFarmPlot.cpp
+++ b/Engine/Source/Developer/MassEntityTestSuite/Private/MassEntityTestFarmPlot.cpp
@@ -59,26 +59,26 @@ void UFarmHarvestTimerSystem_Flowers::Execute(FMassEntityManager& EntityManager,
 {
 	QUICK_SCOPE_CYCLE_COUNTER(UFarmHarvestTimerSystem_Flowers_Run);
 
-	EntityQuery.ForEachEntityChunk(Context, [this](FMassExecutionContext& Context) {
-		const int32 NumEntities = Context.GetNumEntities();
-		const float WellWateredThreshold = 0.25f;
-		TArrayView<FHarvestTimerFragment> TimerList = Context.GetMutableFragmentView<FHarvestTimerFragment>();
-		TConstArrayView<FFarmWaterFragment> WaterList = Context.GetFragmentView<FFarmWaterFragment>();
-		TArrayView<FFarmFlowerFragment> FlowerList = Context.GetMutableFragmentView<FFarmFlowerFragment>();
-
-		for (int32 i = 0; i < NumEntities; ++i)
+	EntityQuery.ForEachEntityChunk(Context, [this](FMassExecutionContext& Context) 
 		{
-			if (TimerList[i].NumSecondsLeft > 0)
-			{
-				--TimerList[i].NumSecondsLeft;
+			const float WellWateredThreshold = 0.25f;
+			TArrayView<FHarvestTimerFragment> TimerList = Context.GetMutableFragmentView<FHarvestTimerFragment>();
+			TConstArrayView<FFarmWaterFragment> WaterList = Context.GetFragmentView<FFarmWaterFragment>();
+			TArrayView<FFarmFlowerFragment> FlowerList = Context.GetMutableFragmentView<FFarmFlowerFragment>();
 
-				if (WaterList[i].CurrentWater > WellWateredThreshold)
+			for (FMassExecutionContext::FEntityIterator EntityIt = Context.CreateEntityIterator(); EntityIt; ++EntityIt)
+			{
+				if (TimerList[EntityIt].NumSecondsLeft > 0)
 				{
-					++FlowerList[i].NumBonusTicks;
+					--TimerList[EntityIt].NumSecondsLeft;
+
+					if (WaterList[EntityIt].CurrentWater > WellWateredThreshold)
+					{
+						++FlowerList[EntityIt].NumBonusTicks;
+					}
 				}
 			}
-		}
-	});
+		});
 }
 
 //-----------------------------------------------------------------------------
@@ -95,19 +95,18 @@ void UFarmHarvestTimerSystem_Crops::Execute(FMassEntityManager & EntityManager,
 {
 	QUICK_SCOPE_CYCLE_COUNTER(UFarmHarvestTimerSystem_Crops_Run);
 
-	EntityQuery.ForEachEntityChunk(Context, [this](FMassExecutionContext& Context) {
-
-		const int32 NumEntities = Context.GetNumEntities();
-		const float WellWateredThreshold = 0.25f;
-		TArrayView<FHarvestTimerFragment> TimerList = Context.GetMutableFragmentView<FHarvestTimerFragment>();
-		TConstArrayView<FFarmWaterFragment> WaterList = Context.GetFragmentView<FFarmWaterFragment>();
-
-		for (int32 i = 0; i < NumEntities; ++i)
+	EntityQuery.ForEachEntityChunk(Context, [this](FMassExecutionContext& Context) 
 		{
-			const uint32 TimeToSubtract = (WaterList[i].CurrentWater > WellWateredThreshold) ? 2 : 1;
-			TimerList[i].NumSecondsLeft = (TimerList[i].NumSecondsLeft >= TimeToSubtract) ? (TimerList[i].NumSecondsLeft - TimeToSubtract) : 0;
-		}
-	});
+			const float WellWateredThreshold = 0.25f;
+			TArrayView<FHarvestTimerFragment> TimerList = Context.GetMutableFragmentView<FHarvestTimerFragment>();
+			TConstArrayView<FFarmWaterFragment> WaterList = Context.GetFragmentView<FFarmWaterFragment>();
+
+			for (FMassExecutionContext::FEntityIterator EntityIt = Context.CreateEntityIterator(); EntityIt; ++EntityIt)
+			{
+				const uint32 TimeToSubtract = (WaterList[EntityIt].CurrentWater > WellWateredThreshold) ? 2 : 1;
+				TimerList[EntityIt].NumSecondsLeft = (TimerList[EntityIt].NumSecondsLeft >= TimeToSubtract) ? (TimerList[EntityIt].NumSecondsLeft - TimeToSubtract) : 0;
+			}
+		});
 }
 
 //-----------------------------------------------------------------------------
@@ -124,18 +123,18 @@ void UFarmHarvestTimerExpired::Execute(FMassEntityManager& EntityManager, FMassE
 {
 	QUICK_SCOPE_CYCLE_COUNTER(UFarmHarvestTimerExpired_Run);
 
-	EntityQuery.ForEachEntityChunk(Context, [this](FMassExecutionContext& Context) {
-		const int32 NumEntities = Context.GetNumEntities();
-		TConstArrayView<FHarvestTimerFragment> TimerList = Context.GetFragmentView<FHarvestTimerFragment>();
-
-		for (int32 i = 0; i < NumEntities; ++i)
+	EntityQuery.ForEachEntityChunk(Context, [this](FMassExecutionContext& Context) 
 		{
-			if (TimerList[i].NumSecondsLeft == 0)
+			TConstArrayView<FHarvestTimerFragment> TimerList = Context.GetFragmentView<FHarvestTimerFragment>();
+
+			for (FMassExecutionContext::FEntityIterator EntityIt = Context.CreateEntityIterator(); EntityIt; ++EntityIt)
 			{
-				Context.Defer().AddTag<FFarmJustBecameReadyToHarvestTag>(Context.GetEntity(i));
+				if (TimerList[EntityIt].NumSecondsLeft == 0)
+				{
+					Context.Defer().AddTag<FFarmJustBecameReadyToHarvestTag>(Context.GetEntity(EntityIt));
+				}
 			}
-		}
-	});
+		});
 }
 
 //-----------------------------------------------------------------------------
@@ -157,26 +156,25 @@ void UFarmHarvestTimerSetIcon::Execute(FMassEntityManager& EntityManager, FMassE
 
 	QUICK_SCOPE_CYCLE_COUNTER(SET_ICON_SET_ICON_SET_ICON);
 
-	EntityQuery.ForEachEntityChunk(Context, [this](FMassExecutionContext& Context) {
-
-		const int32 NumEntities = Context.GetNumEntities();
-		TConstArrayView<FFarmGridCellData> GridCoordList = Context.GetFragmentView<FFarmGridCellData>();
-		TArrayView<FFarmVisualFragment> VisualList = Context.GetMutableFragmentView<FFarmVisualFragment>();
-
-		for (int32 i = 0; i < NumEntities; ++i)
+	EntityQuery.ForEachEntityChunk(Context, [this](FMassExecutionContext& Context) 
 		{
-			const FFarmGridCellData& GridCells = GridCoordList[i];
+			TConstArrayView<FFarmGridCellData> GridCoordList = Context.GetFragmentView<FFarmGridCellData>();
+			TArrayView<FFarmVisualFragment> VisualList = Context.GetMutableFragmentView<FFarmVisualFragment>();
 
-			const FVector IconPosition(GridCells.CellX*GridCellWidth, GridCells.CellY*GridCellHeight, HarvestIconHeight);
-			const FTransform IconTransform(FQuat::Identity, IconPosition, FVector(HarvestIconScale, HarvestIconScale, HarvestIconScale));
+			for (FMassExecutionContext::FEntityIterator EntityIt = Context.CreateEntityIterator(); EntityIt; ++EntityIt)
+			{
+				const FFarmGridCellData& GridCells = GridCoordList[EntityIt];
 
-			VisualList[i].HarvestIconIndex = HarvestIconISMC->AddInstance(IconTransform);
+				const FVector IconPosition(GridCells.CellX*GridCellWidth, GridCells.CellY*GridCellHeight, HarvestIconHeight);
+				const FTransform IconTransform(FQuat::Identity, IconPosition, FVector(HarvestIconScale, HarvestIconScale, HarvestIconScale));
 
-			FMassEntityHandle ThisEntity = Context.GetEntity(i);
-			Context.Defer().RemoveTag<FFarmJustBecameReadyToHarvestTag>(ThisEntity);
-			Context.Defer().AddTag<FFarmReadyToHarvestTag>(ThisEntity);
-		}
-	});
+				VisualList[EntityIt].HarvestIconIndex = HarvestIconISMC->AddInstance(IconTransform);
+
+				FMassEntityHandle ThisEntity = Context.GetEntity(EntityIt);
+				Context.Defer().RemoveTag<FFarmJustBecameReadyToHarvestTag>(ThisEntity);
+				Context.Defer().AddTag<FFarmReadyToHarvestTag>(ThisEntity);
+			}
+		});
 }
 
 //-----------------------------------------------------------------------------
