Fix typo in description

#rnx

diff --git a/Engine/Plugins/AI/MassAI/Source/MassNavigation/Public/Avoidance/MassAvoidanceFragments.h b/Engine/Plugins/AI/MassAI/Source/MassNavigation/Public/Avoidance/MassAvoidanceFragments.h
--- a/Engine/Plugins/AI/MassAI/Source/MassNavigation/Public/Avoidance/MassAvoidanceFragments.h
+++ b/Engine/Plugins/AI/MassAI/Source/MassNavigation/Public/Avoidance/MassAvoidanceFragments.h
@@ -149,7 +149,7 @@ struct MASSNAVIGATION_API FMassStandingAvoidanceParameters : public FMassConstSh
 	UPROPERTY(EditAnywhere, Category = "Ghost", meta = (ClampMin = "0", ForceUnits="N/cm"))
 	float GhostSeparationStiffness = 200.0f;
 
-	/** Much much avoidance is scaled for moving obstacles. Range: 1..5. */
+	/** How much avoidance is scaled for moving obstacles. Range: 1..5. */
 	UPROPERTY(EditAnywhere, Category = "Ghost", meta = (ClampMin = "0", ForceUnits="x"))
 	float MovingObstacleAvoidanceScale = 3.0f;
 
