Add optional display of mass navmesh short path in the gameplay debugger tool.

diff --git a/Engine/Plugins/AI/MassAI/Source/MassAIDebug/MassAIDebug.Build.cs b/Engine/Plugins/AI/MassAI/Source/MassAIDebug/MassAIDebug.Build.cs
--- a/Engine/Plugins/AI/MassAI/Source/MassAIDebug/MassAIDebug.Build.cs
+++ b/Engine/Plugins/AI/MassAI/Source/MassAIDebug/MassAIDebug.Build.cs
@@ -25,6 +25,7 @@ namespace UnrealBuildTool.Rules
 					"MassCommon",
 					"MassMovement",
 					"MassNavigation",
+					"MassNavMeshNavigation",
 					"MassZoneGraphNavigation",
 					"MassAIReplication",
 					"MassSmartObjects",
diff --git a/Engine/Plugins/AI/MassAI/Source/MassAIDebug/Private/GameplayDebuggerCategory_Mass.cpp b/Engine/Plugins/AI/MassAI/Source/MassAIDebug/Private/GameplayDebuggerCategory_Mass.cpp
--- a/Engine/Plugins/AI/MassAI/Source/MassAIDebug/Private/GameplayDebuggerCategory_Mass.cpp
+++ b/Engine/Plugins/AI/MassAI/Source/MassAIDebug/Private/GameplayDebuggerCategory_Mass.cpp
@@ -19,6 +19,7 @@
 #include "MassStateTreeFragments.h"
 #include "MassStateTreeExecutionContext.h"
 #include "MassZoneGraphNavigationFragments.h"
+#include "MassNavMeshNavigationFragments.h"
 #include "MassSmartObjectFragments.h"
 #include "SmartObjectSubsystem.h"
 #include "MassSimulationLOD.h"
@@ -507,6 +508,7 @@ void FGameplayDebuggerCategory_Mass::CollectData(APlayerController* OwnerPC, AAc
 		EntityQuery.AddRequirement<FMassLookAtFragment>(EMassFragmentAccess::ReadOnly, EMassFragmentPresence::Optional);
 		EntityQuery.AddRequirement<FMassSimulationLODFragment>(EMassFragmentAccess::ReadOnly, EMassFragmentPresence::Optional);
 		EntityQuery.AddRequirement<FMassZoneGraphShortPathFragment>(EMassFragmentAccess::ReadOnly, EMassFragmentPresence::Optional);
+		EntityQuery.AddRequirement<FMassNavMeshShortPathFragment>(EMassFragmentAccess::ReadOnly, EMassFragmentPresence::Optional);
 		EntityQuery.AddRequirement<FMassSmartObjectUserFragment>(EMassFragmentAccess::ReadOnly, EMassFragmentPresence::Optional);
 
 		const double CurrentTime = World->GetTimeSeconds();
@@ -534,7 +536,8 @@ void FGameplayDebuggerCategory_Mass::CollectData(APlayerController* OwnerPC, AAc
 				const TConstArrayView<FMassMoveTargetFragment> MoveTargetList = Context.GetFragmentView<FMassMoveTargetFragment>();
 				const TConstArrayView<FMassLookAtFragment> LookAtList = Context.GetFragmentView<FMassLookAtFragment>();
 				const TConstArrayView<FMassSimulationLODFragment> SimLODList = Context.GetFragmentView<FMassSimulationLODFragment>();
-				const TConstArrayView<FMassZoneGraphShortPathFragment> ShortPathList = Context.GetFragmentView<FMassZoneGraphShortPathFragment>();
+				const TConstArrayView<FMassZoneGraphShortPathFragment> ZoneGraphShortPathList = Context.GetFragmentView<FMassZoneGraphShortPathFragment>();
+				const TConstArrayView<FMassNavMeshShortPathFragment> NavMeshShortPathList = Context.GetFragmentView<FMassNavMeshShortPathFragment>();
 				const TConstArrayView<FMassSmartObjectUserFragment> SOUserList = Context.GetFragmentView<FMassSmartObjectUserFragment>();
 				const FMassStateTreeSharedFragment& SharedStateTree = Context.GetConstSharedFragment<FMassStateTreeSharedFragment>();
 
@@ -642,24 +645,48 @@ void FGameplayDebuggerCategory_Mass::CollectData(APlayerController* OwnerPC, AAc
 					}
 
 					// Path
-					if (bShowNearEntityPath && ShortPathList.Num() > 0)
+					if (bShowNearEntityPath)
 					{
-						const FMassZoneGraphShortPathFragment& ShortPath = ShortPathList[EntityIndex];
-
-						const FVector ZOffset = FVector(0.0f , 0.0f , 25.0f );
-						for (uint8 PointIndex = 0; PointIndex < ShortPath.NumPoints - 1; PointIndex++)
+						const FVector ZOffset = FVector(0.f , 0.f , 25.f );
+					
+						if (!ZoneGraphShortPathList.IsEmpty())
 						{
-							const FMassZoneGraphPathPoint& CurrPoint = ShortPath.Points[PointIndex];
-							const FMassZoneGraphPathPoint& NextPoint = ShortPath.Points[PointIndex + 1];
-							AddShape(FGameplayDebuggerShape::MakeSegment(CurrPoint.Position + ZOffset, NextPoint.Position + ZOffset, 3.0f, FColorList::Grey));
-						}
+							const FMassZoneGraphShortPathFragment& ShortPath = ZoneGraphShortPathList[EntityIndex];
+							
+							for (uint8 PointIndex = 0; PointIndex < ShortPath.NumPoints - 1; PointIndex++)
+							{
+								const FMassZoneGraphPathPoint& CurrPoint = ShortPath.Points[PointIndex];
+								const FMassZoneGraphPathPoint& NextPoint = ShortPath.Points[PointIndex + 1];
+								AddShape(FGameplayDebuggerShape::MakeSegment(CurrPoint.Position + ZOffset, NextPoint.Position + ZOffset, /*Thickness*/3.f, FColorList::Grey));
+							}
 					
-						for (uint8 PointIndex = 0; PointIndex < ShortPath.NumPoints; PointIndex++)
+							for (uint8 PointIndex = 0; PointIndex < ShortPath.NumPoints; PointIndex++)
+							{
+								const FMassZoneGraphPathPoint& CurrPoint = ShortPath.Points[PointIndex];
+								const FVector CurrBase = CurrPoint.Position + ZOffset;
+								// Lane tangents
+								AddShape(FGameplayDebuggerShape::MakeSegment(CurrBase, CurrBase + CurrPoint.Tangent.GetVector() * 50.f, /*Thickness*/1.f, FColorList::LightGrey));
+							}
+						}
+
+						if (!NavMeshShortPathList.IsEmpty())
 						{
-							const FMassZoneGraphPathPoint& CurrPoint = ShortPath.Points[PointIndex];
-							const FVector CurrBase = CurrPoint.Position + ZOffset;
-							// Lane tangents
-							AddShape(FGameplayDebuggerShape::MakeSegment(CurrBase, CurrBase + CurrPoint.Tangent.GetVector() * 50.0f, 1.0f, FColorList::LightGrey));
+							const FMassNavMeshShortPathFragment& ShortPath = NavMeshShortPathList[EntityIndex];
+
+							for (uint8 PointIndex = 0; PointIndex < ShortPath.NumPoints - 1; PointIndex++)
+							{
+								const FMassNavMeshPathPoint& CurrPoint = ShortPath.Points[PointIndex];
+								const FMassNavMeshPathPoint& NextPoint = ShortPath.Points[PointIndex + 1];
+								AddShape(FGameplayDebuggerShape::MakeSegment(CurrPoint.Position + ZOffset, NextPoint.Position + ZOffset, /*Thickness*/3.f, FColorList::GreenYellow));
+							}
+					
+							for (uint8 PointIndex = 0; PointIndex < ShortPath.NumPoints; PointIndex++)
+							{
+								const FMassNavMeshPathPoint& CurrPoint = ShortPath.Points[PointIndex];
+								const FVector CurrBase = CurrPoint.Position + ZOffset;
+								// Path tangents
+								AddShape(FGameplayDebuggerShape::MakeSegment(CurrBase, CurrBase + CurrPoint.Tangent.GetVector() * 50.f, /*Thickness*/1.f, FColorList::LightGrey));
+							}
 						}
 					}
 
