Keep FMassNavMeshPathFollowTask running until entity reached it's destination (within provided threshold)

diff --git a/Engine/Plugins/AI/MassAI/Source/MassAIBehavior/Private/Tasks/MassNavMeshPathFollowTask.cpp b/Engine/Plugins/AI/MassAI/Source/MassAIBehavior/Private/Tasks/MassNavMeshPathFollowTask.cpp
--- a/Engine/Plugins/AI/MassAI/Source/MassAIBehavior/Private/Tasks/MassNavMeshPathFollowTask.cpp
+++ b/Engine/Plugins/AI/MassAI/Source/MassAIBehavior/Private/Tasks/MassNavMeshPathFollowTask.cpp
@@ -88,7 +88,7 @@ bool FMassNavMeshPathFollowTask::RequestPath(FStateTreeExecutionContext& Context
 
 		// Update short path
 		FMassNavMeshShortPathFragment& ShortPathFragment = Context.GetExternalData(ShortPathHandle);
-		ShortPathFragment.RequestShortPath(CachedPathFragment.Corridor, /*NavCorridorStartIndex*/0, /*NumLeadingPoints*/0);
+		ShortPathFragment.RequestShortPath(CachedPathFragment.Corridor, /*NavCorridorStartIndex*/0, /*NumLeadingPoints*/0, InstanceData.EndDistanceThreshold);
 
 		CachedPathFragment.NavPathNextStartIndex = (uint16)FMath::Max(ShortPathFragment.NumPoints - FMassNavMeshShortPathFragment::NumPointsBeyondUpdate - FMassNavMeshCachedPathFragment::NumLeadingPoints, 0);
 		
@@ -111,13 +111,14 @@ bool FMassNavMeshPathFollowTask::RequestPath(FStateTreeExecutionContext& Context
 
 bool FMassNavMeshPathFollowTask::UpdateShortPath(FStateTreeExecutionContext& Context) const
 {
+	FMassNavMeshPathFollowTaskInstanceData& InstanceData = Context.GetInstanceData<FMassNavMeshPathFollowTaskInstanceData>(*this);
+	
 	FMassNavMeshCachedPathFragment& CachedPathFragment = Context.GetExternalData(CachedPathHandle);
-
 	MASSBEHAVIOR_LOG(Verbose, TEXT("updating short path, starting at index %i"), CachedPathFragment.NavPathNextStartIndex);
 	
 	FMassNavMeshShortPathFragment& ShortPathFragment = Context.GetExternalData(ShortPathHandle);
 
-	ShortPathFragment.RequestShortPath(CachedPathFragment.Corridor, CachedPathFragment.NavPathNextStartIndex, CachedPathFragment.NumLeadingPoints);
+	ShortPathFragment.RequestShortPath(CachedPathFragment.Corridor, CachedPathFragment.NavPathNextStartIndex, CachedPathFragment.NumLeadingPoints, InstanceData.EndDistanceThreshold);
 
 	CachedPathFragment.NavPathNextStartIndex += (uint16)FMath::Max(ShortPathFragment.NumPoints - FMassNavMeshShortPathFragment::NumPointsBeyondUpdate - FMassNavMeshCachedPathFragment::NumLeadingPoints, 0);
 
diff --git a/Engine/Plugins/AI/MassAI/Source/MassAIBehavior/Public/Tasks/MassNavMeshPathfollowTask.h b/Engine/Plugins/AI/MassAI/Source/MassAIBehavior/Public/Tasks/MassNavMeshPathfollowTask.h
--- a/Engine/Plugins/AI/MassAI/Source/MassAIBehavior/Public/Tasks/MassNavMeshPathfollowTask.h
+++ b/Engine/Plugins/AI/MassAI/Source/MassAIBehavior/Public/Tasks/MassNavMeshPathfollowTask.h
@@ -36,6 +36,10 @@ struct MASSAIBEHAVIOR_API FMassNavMeshPathFollowTaskInstanceData
 	/** Amount to offset corridor sides from navmesh borders. Used to keep movement away for navmesh borders. */
 	UPROPERTY(EditAnywhere, Category = Parameter)
 	float OffsetFromBoundaries = 10.f;
+
+	/** Distance from the end of path used to confirm that the destination is reached. */
+	UPROPERTY(EditAnywhere, Category = Parameter)
+	float EndDistanceThreshold = 20.f;
 };
 
 /** Finds a path to TargetLocation, requests a short path, starts a move action and follow the path by updating the short path when needed. */ 
diff --git a/Engine/Plugins/AI/MassAI/Source/MassNavMeshNavigation/Private/MassNavMeshNavigationFragments.cpp b/Engine/Plugins/AI/MassAI/Source/MassNavMeshNavigation/Private/MassNavMeshNavigationFragments.cpp
--- a/Engine/Plugins/AI/MassAI/Source/MassNavMeshNavigation/Private/MassNavMeshNavigationFragments.cpp
+++ b/Engine/Plugins/AI/MassAI/Source/MassNavMeshNavigation/Private/MassNavMeshNavigationFragments.cpp
@@ -6,17 +6,19 @@
 #include "NavCorridor.h"
 
 // Make a new short corridor
-bool FMassNavMeshShortPathFragment::RequestShortPath(const TSharedPtr<FNavCorridor>& NavCorridor, const int32 NavCorridorStartIndex, const uint8 NumLeadingPoints)
+bool FMassNavMeshShortPathFragment::RequestShortPath(const TSharedPtr<FNavCorridor>& NavCorridor, const int32 InNavCorridorStartIndex, const uint8 InNumLeadingPoints,
+	const float InEndReachedDistance)
 {
 	Reset();
+	EndReachedDistance = InEndReachedDistance;
 
 	TArray<FNavCorridorPortal>& NavCorridorPortals = NavCorridor->Portals;
 	int32 PortalNum = NavCorridorPortals.Num();
 
 	// Copy Portals from the StartIndex
-	for(; NumPoints < MaxPoints && (NavCorridorStartIndex + NumPoints) < PortalNum; NumPoints++)
+	for(; NumPoints < MaxPoints && (InNavCorridorStartIndex + NumPoints) < PortalNum; NumPoints++)
 	{
-		const FNavCorridorPortal& SourcePortal = NavCorridorPortals[NavCorridorStartIndex + NumPoints];
+		const FNavCorridorPortal& SourcePortal = NavCorridorPortals[InNavCorridorStartIndex + NumPoints];
 		
 		FMassNavMeshPathPoint& Point = Points[NumPoints];
 		Point.Left = SourcePortal.Left;
@@ -26,11 +28,12 @@ bool FMassNavMeshShortPathFragment::RequestShortPath(const TSharedPtr<FNavCorrid
 
 	checkf(NumPoints >= 2, TEXT("Path should have at least 2 points at this stage, has %d."), NumPoints);
 
-	// We are not done with the main path, mark as partial.
+	// If we are not done with the main path, mark as partial.
 	const int32 LastNavPathIndex = PortalNum-1;
-	if (NavCorridorStartIndex + (NumPoints-1) < LastNavPathIndex)
+	if (InNavCorridorStartIndex + (NumPoints-1) < LastNavPathIndex)
 	{
 		bPartialResult = true;
+		ensure(NumPoints == MaxPoints);
 	}
 
 	// Calculate movement distance and tangent at each point.
@@ -52,7 +55,7 @@ bool FMassNavMeshShortPathFragment::RequestShortPath(const TSharedPtr<FNavCorrid
 	}
 
 	// Update ProgressDistance to account for leading points. 
-	ProgressDistance = Points[NumLeadingPoints].Distance.Get();
+	ProgressDistance = Points[InNumLeadingPoints].Distance.Get();
 	
 	return true;
 }
diff --git a/Engine/Plugins/AI/MassAI/Source/MassNavMeshNavigation/Private/MassNavMeshNavigationProcessors.cpp b/Engine/Plugins/AI/MassAI/Source/MassNavMeshNavigation/Private/MassNavMeshNavigationProcessors.cpp
--- a/Engine/Plugins/AI/MassAI/Source/MassNavMeshNavigation/Private/MassNavMeshNavigationProcessors.cpp
+++ b/Engine/Plugins/AI/MassAI/Source/MassNavMeshNavigation/Private/MassNavMeshNavigationProcessors.cpp
@@ -1,6 +1,8 @@
 // Copyright Epic Games, Inc. All Rights Reserved.
 
 #include "MassNavMeshNavigationProcessors.h"
+
+#include "MassCommonFragments.h"
 #include "Avoidance/MassAvoidanceFragments.h"
 #include "MassCommonTypes.h"
 #include "MassDebugger.h"
@@ -31,6 +33,7 @@ void UMassNavMeshPathFollowProcessor::ConfigureQueries(const TSharedRef<FMassEnt
 {
 	EntityQuery_Conditional.AddRequirement<FMassNavMeshShortPathFragment>(EMassFragmentAccess::ReadWrite);
 	EntityQuery_Conditional.AddRequirement<FMassMoveTargetFragment>(EMassFragmentAccess::ReadWrite);
+	EntityQuery_Conditional.AddRequirement<FTransformFragment>(EMassFragmentAccess::ReadOnly);
 
 	// @todo: validate LOD and variable ticking
 	EntityQuery_Conditional.AddRequirement<FMassSimulationLODFragment>(EMassFragmentAccess::ReadOnly, EMassFragmentPresence::Optional);
@@ -50,6 +53,7 @@ void UMassNavMeshPathFollowProcessor::Execute(FMassEntityManager& EntityManager,
 	{
 		const TArrayView<FMassNavMeshShortPathFragment> ShortPathList = Context.GetMutableFragmentView<FMassNavMeshShortPathFragment>();
 		const TArrayView<FMassMoveTargetFragment> MoveTargetList = Context.GetMutableFragmentView<FMassMoveTargetFragment>();
+		const TConstArrayView<FTransformFragment> TransformList = Context.GetFragmentView<FTransformFragment>();
 
 		// @todo: validate LOD and variable ticking
 		const TConstArrayView<FMassSimulationLODFragment> SimLODList = Context.GetFragmentView<FMassSimulationLODFragment>();
@@ -92,6 +96,9 @@ void UMassNavMeshPathFollowProcessor::Execute(FMassEntityManager& EntityManager,
 					// Point index on the short path where we should update the path.
 					const uint8 UpdatePointIndex =  ShortPath.bPartialResult ? (ShortPath.NumPoints - ShortPath.NumPointsBeyondUpdate) : LastIndex;
 
+					// If the shortpath is partial, it's expected to be full, meaning MaxPoints-NumPointsBeyondUpdate >= 1 so the check below must be valid.
+					check(UpdatePointIndex >= 1);
+
 #if WITH_MASSGAMEPLAY_DEBUG && UNSAFE_FOR_MT
 					UE_CVLOG_UELOG(bDisplayDebug, this, LogMassNavMeshNavigation, Verbose, TEXT("Entity [%s]   ProgressDistance: %f, UpdatePointIndex: %i"),
 						*Entity.DebugGetDescription(), ShortPath.ProgressDistance, UpdatePointIndex);
@@ -101,7 +108,7 @@ void UMassNavMeshPathFollowProcessor::Execute(FMassEntityManager& EntityManager,
 						// Display update point
 						const FVector ZOffset(0,0,20);
 						UE_VLOG_CIRCLE(this, LogMassNavMeshNavigation, Display, ShortPath.Points[UpdatePointIndex].Position + ZOffset,
-							FVector(0,0,1), /*Radius*/10.f, FColor::Red, TEXT("%i"), UpdatePointIndex);
+							FVector(0,0,1), /*Radius*/10.f, FColor::Green, TEXT("%i"), UpdatePointIndex);
 					}
 #endif // WITH_MASSGAMEPLAY_DEBUG
 
@@ -150,9 +157,29 @@ void UMassNavMeshPathFollowProcessor::Execute(FMassEntityManager& EntityManager,
 						MoveTarget.Forward = ShortPath.Points[UpdatePointIndex].Tangent.GetVector();
 						MoveTarget.DistanceToGoal = 0.f;
 
-						UE_CVLOG_UELOG(bDisplayDebug, this, LogMassNavMeshNavigation, Log, TEXT("Entity [%s]   Finished path follow on short path."), *Entity.DebugGetDescription());
+						if (ShortPath.bPartialResult)
+						{
+							UE_CVLOG_UELOG(bDisplayDebug, this, LogMassNavMeshNavigation, Log, TEXT("Entity [%s]   Finished path follow on short path."), *Entity.DebugGetDescription());
+							ShortPath.bDone = true;
+						}
+						else
+						{
+							// Last section of the path, wait for the steering to complete the movement.
+							const FTransformFragment& TransformFragment = TransformList[EntityIt];
+							const FVector& EntityLocation = TransformFragment.GetTransform().GetLocation();
+
+							const float SegmentSize = ShortPath.Points[UpdatePointIndex].Distance.Get() - ShortPath.Points[UpdatePointIndex-1].Distance.Get();
+							const FVector PreviousPointLocation = ShortPath.Points[UpdatePointIndex-1].Position;
+							FVector PathDir = ShortPath.Points[UpdatePointIndex-1].Tangent.GetVector();
 
-						ShortPath.bDone = true;
+							const FVector RelativeLocation = EntityLocation - PreviousPointLocation;
+							const float ProjectionOnSegment = RelativeLocation.Dot(PathDir);
+							if (ProjectionOnSegment > (SegmentSize - ShortPath.EndReachedDistance))
+							{
+								UE_CVLOG_UELOG(bDisplayDebug, this, LogMassNavMeshNavigation, Log, TEXT("Entity [%s]   Finished path follow on short path, end of path."), *Entity.DebugGetDescription());
+								ShortPath.bDone = true;
+							}
+						}
 					}
 				}
 
@@ -208,6 +235,14 @@ void UMassNavMeshPathFollowProcessor::Execute(FMassEntityManager& EntityManager,
 					// Draw MoveTarget
 					constexpr float Radius = 20.f;
 					UE_VLOG_CIRCLE(this, LogMassNavMeshNavigation, Display, MoveTarget.Center + ZOffset, FVector(0,0,1), Radius, LightEntityColor, TEXT(""));
+
+					// Draw tolerance distance at end (shown as a circle instead of an infinite perpendicular line)
+					if (!ShortPath.bPartialResult)
+					{
+						const float EndRadius = ShortPath.EndReachedDistance;
+						UE_VLOG_WIRECIRCLE(this, LogMassNavMeshNavigation, Display, ShortPath.Points[ShortPath.NumPoints-1].Position + ZOffset,
+							FVector(0,0,1), EndRadius, FColor::Black, TEXT("End"));
+					}
 				}
 #endif // WITH_MASSGAMEPLAY_DEBUG
 			}
diff --git a/Engine/Plugins/AI/MassAI/Source/MassNavMeshNavigation/Public/MassNavMeshNavigationFragments.h b/Engine/Plugins/AI/MassAI/Source/MassNavMeshNavigation/Public/MassNavMeshNavigationFragments.h
--- a/Engine/Plugins/AI/MassAI/Source/MassNavMeshNavigation/Public/MassNavMeshNavigationFragments.h
+++ b/Engine/Plugins/AI/MassAI/Source/MassNavMeshNavigation/Public/MassNavMeshNavigationFragments.h
@@ -50,6 +50,7 @@ struct MASSNAVMESHNAVIGATION_API FMassNavMeshShortPathFragment : public FMassFra
 	void Reset()
 	{
 		ProgressDistance = 0.f;
+		EndReachedDistance = 20.f;
 		NumPoints = 0;
 		EndOfPathIntent = EMassMovementAction::Stand;
 		bPartialResult = false;
@@ -63,13 +64,16 @@ struct MASSNAVMESHNAVIGATION_API FMassNavMeshShortPathFragment : public FMassFra
 	}
 
 	/** Fill Points from Corridor. */
-	bool RequestShortPath(const TSharedPtr<FNavCorridor>& Corridor, const int32 NavCorridorStartIndex, const uint8 NumLeadingPoints);
+	bool RequestShortPath(const TSharedPtr<FNavCorridor>& InCorridor, const int32 InNavCorridorStartIndex, const uint8 InNumLeadingPoints, const float InEndReachedDistance);
 
 	/** Portal points */
 	TStaticArray<FMassNavMeshPathPoint, MaxPoints> Points;
 
 	/** Current progress distance along the short path. */
 	float ProgressDistance = 0.f;
+
+	/** Distance from the end of path used to confirm that the destination is reached. */
+	float EndReachedDistance = 20.f;
 	
 	/** Number of points on path. */
 	uint8 NumPoints = 0;
