Fix compilation error introduced with CL 39512303

#rnx

diff --git a/Engine/Plugins/AI/MassAI/Source/MassAIBehavior/Private/MassStateTreeExecutionContext.cpp b/Engine/Plugins/AI/MassAI/Source/MassAIBehavior/Private/MassStateTreeExecutionContext.cpp
--- a/Engine/Plugins/AI/MassAI/Source/MassAIBehavior/Private/MassStateTreeExecutionContext.cpp
+++ b/Engine/Plugins/AI/MassAI/Source/MassAIBehavior/Private/MassStateTreeExecutionContext.cpp
@@ -174,7 +174,7 @@ FString FMassExecutionExtension::GetInstanceDescription(const FContextParameters
 void FMassExecutionExtension::OnLinkedStateTreeOverridesSet(const FContextParameters& Context, const FStateTreeReferenceOverrides& Overrides)
 {
 	const uint32 NewLinkedStateTreeOverridesHash = GetTypeHash(Overrides);
-	if (LinkedStateTreeOverridesHash != LinkedStateTreeOverridesHash)
+	if (NewLinkedStateTreeOverridesHash != LinkedStateTreeOverridesHash)
 	{
 		LinkedStateTreeOverridesHash = NewLinkedStateTreeOverridesHash;
 		//@todo update mass dependencies
