Addition of vislog for standing avoidance
Addition of vislog to UMassSteerToMoveTargetProcessor
Small fix in debug display of UMassNavMeshNavigationBoundaryProcessor

diff --git a/Engine/Plugins/AI/MassAI/Source/MassNavMeshNavigation/Private/MassNavMeshNavigationBoundaryProcessor.cpp b/Engine/Plugins/AI/MassAI/Source/MassNavMeshNavigation/Private/MassNavMeshNavigationBoundaryProcessor.cpp
--- a/Engine/Plugins/AI/MassAI/Source/MassNavMeshNavigation/Private/MassNavMeshNavigationBoundaryProcessor.cpp
+++ b/Engine/Plugins/AI/MassAI/Source/MassNavMeshNavigation/Private/MassNavMeshNavigationBoundaryProcessor.cpp
@@ -61,10 +61,11 @@ void UMassNavMeshNavigationBoundaryProcessor::Execute(FMassEntityManager& Entity
 
 			if (bDisplayDebug)
 			{
-				// Draw MoveTarget
+				// Draw LastUpdatePosition
 				const FVector ZOffset(0,0,10);
-				constexpr float Radius = 20.f;
-				UE_VLOG_CIRCLE(this, LogMassNavMeshNavigation, Display, NavmeshBoundary.LastUpdatePosition + ZOffset, FVector(0,0,1), Radius, EntityColor, TEXT(""));
+				constexpr float Radius = 5.f;
+				UE_VLOG_WIRECIRCLE(this, LogMassNavMeshNavigation, Display, NavmeshBoundary.LastUpdatePosition + ZOffset,
+					FVector(0,0,1), Radius, FColor::Blue, TEXT("Boundary update"));
 			}
 #endif // WITH_MASSGAMEPLAY_DEBUG
 
diff --git a/Engine/Plugins/AI/MassAI/Source/MassNavigation/Private/MassAvoidanceProcessors.cpp b/Engine/Plugins/AI/MassAI/Source/MassNavigation/Private/MassAvoidanceProcessors.cpp
--- a/Engine/Plugins/AI/MassAI/Source/MassNavigation/Private/MassAvoidanceProcessors.cpp
+++ b/Engine/Plugins/AI/MassAI/Source/MassNavigation/Private/MassAvoidanceProcessors.cpp
@@ -176,6 +176,10 @@ namespace UE::MassAvoidance
 	static const FColor ObstacleAvoidForceColor = FColor::Magenta;
 	static const FColor ObstacleAvoidFuturHitPosColor = FColor::Magenta;
 	static constexpr FColor ObstacleSeparationForceColor = FColor(255, 66, 66);	// Bright red
+
+	// Ghost colors
+	static const FColor GhostColor = FColor::Silver;
+	static const FColor GhostSteeringForceColor = FColor::Silver;
 	
 	static const FVector DebugAgentHeightOffset = FVector(0., 0., 185.);
 	static const FVector DebugLowCylinderOffset = FVector(0., 0., 20.);
@@ -1123,6 +1127,18 @@ void UMassStandingAvoidanceProcessor::Execute(FMassEntityManager& EntityManager,
 			ClosestObstacles.Sort([](const FSortedObstacle& A, const FSortedObstacle& B) { return A.DistSq < B.DistSq; });
 
 			const FVector::FReal GhostRadius = AgentRadius * StandingParams.GhostSeparationRadiusScale;
+
+#if WITH_MASSGAMEPLAY_DEBUG && UNSAFE_FOR_MT			
+			const UE::MassAvoidance::FDebugContext AgentDebugContext(this, LogAvoidanceAgents, World, Entity);
+
+			// Draw entity
+			UE::MassAvoidance::DebugDrawCylinder(AgentDebugContext, AgentLocation, AgentLocation + UE::MassAvoidance::DebugAgentHeightOffset,
+				AgentRadius, UE::MassAvoidance::CurrentAgentColor);
+
+			// Draw ghost
+			UE::MassAvoidance::DebugDrawCylinder(AgentDebugContext, Ghost.Location, Ghost.Location + UE::MassAvoidance::DebugAgentHeightOffset,
+				AgentRadius, UE::MassAvoidance::GhostColor);
+#endif // WITH_MASSGAMEPLAY_DEBUG
 			
 			// Compute forces
 			constexpr int32 MaxCloseObstacleTreated = 6;
@@ -1135,6 +1151,12 @@ void UMassStandingAvoidanceProcessor::Execute(FMassEntityManager& EntityManager,
 				const FVector::FReal OtherRadius = OtherEntityView.GetFragmentData<FAgentRadiusFragment>().Radius;
 				const FVector::FReal TotalRadius = GhostRadius + OtherRadius;
 
+#if WITH_MASSGAMEPLAY_DEBUG && UNSAFE_FOR_MT
+				// Draw other agent
+				UE::MassAvoidance::DebugDrawCylinder(AgentDebugContext, OtherAgent.Location, OtherAgent.Location + UE::MassAvoidance::DebugAgentHeightOffset,
+					OtherRadius, UE::MassAvoidance::AgentsColor);
+#endif // WITH_MASSGAMEPLAY_DEBUG
+				
 				// @todo: this is heavy fragment to access, see if we could handle this differently.
 				const FMassMoveTargetFragment* OtherMoveTarget = OtherEntityView.GetFragmentDataPtr<FMassMoveTargetFragment>();
 				const FMassGhostLocationFragment* OtherGhost = OtherEntityView.GetFragmentDataPtr<FMassGhostLocationFragment>();
@@ -1163,6 +1185,14 @@ void UMassStandingAvoidanceProcessor::Execute(FMassEntityManager& EntityManager,
 					const FVector SeparationForce = ConNorm * SeparationStiffness * SeparationMag;
 
 					GhostSteeringForce += SeparationForce;
+
+#if WITH_MASSGAMEPLAY_DEBUG && UNSAFE_FOR_MT					
+					// Draw agent avoid force
+					UE::MassAvoidance::DebugDrawForce(AgentDebugContext,
+						Ghost.Location + UE::MassAvoidance::DebugAgentHeightOffset,
+						Ghost.Location + UE::MassAvoidance::DebugAgentHeightOffset + SeparationForce,
+						UE::MassAvoidance::AgentAvoidForceColor);
+#endif // WITH_MASSGAMEPLAY_DEBUG
 				}
 				else
 				{
@@ -1187,11 +1217,28 @@ void UMassStandingAvoidanceProcessor::Execute(FMassEntityManager& EntityManager,
 					const FVector SeparationForce = ConNorm * MovingSeparationStiffness * SeparationMag;
 
 					GhostSteeringForce += SeparationForce;
+
+#if WITH_MASSGAMEPLAY_DEBUG && UNSAFE_FOR_MT
+					// Draw agent avoid force
+					UE::MassAvoidance::DebugDrawForce(AgentDebugContext,
+						OtherBasePosition + UE::MassAvoidance::DebugAgentHeightOffset,
+						OtherBasePosition + UE::MassAvoidance::DebugAgentHeightOffset + SeparationForce,
+						UE::MassAvoidance::AgentAvoidForceColor);
+#endif // WITH_MASSGAMEPLAY_DEBUG
 				}
 			}
 
 			GhostSteeringForce.Z = 0.;
 			GhostSteeringForce = UE::MassNavigation::ClampVector(GhostSteeringForce, StandingParams.GhostMaxAcceleration); // Assume unit mass
+
+#if WITH_MASSGAMEPLAY_DEBUG && UNSAFE_FOR_MT
+			// Draw total steering force
+			UE::MassAvoidance::DebugDrawSummedForce(AgentDebugContext,
+				AgentLocation + UE::MassAvoidance::DebugAgentHeightOffset,
+				AgentLocation + UE::MassAvoidance::DebugAgentHeightOffset + GhostSteeringForce,
+				UE::MassAvoidance::GhostSteeringForceColor);
+#endif // WITH_MASSGAMEPLAY_DEBUG
+			
 			Ghost.Velocity += GhostSteeringForce * DeltaTime;
 			Ghost.Velocity.Z = 0.;
 			
diff --git a/Engine/Plugins/AI/MassAI/Source/MassNavigation/Private/MassSteeringProcessors.cpp b/Engine/Plugins/AI/MassAI/Source/MassNavigation/Private/MassSteeringProcessors.cpp
--- a/Engine/Plugins/AI/MassAI/Source/MassNavigation/Private/MassSteeringProcessors.cpp
+++ b/Engine/Plugins/AI/MassAI/Source/MassNavigation/Private/MassSteeringProcessors.cpp
@@ -4,12 +4,14 @@
 #include "MassCommonUtils.h"
 #include "MassCommandBuffer.h"
 #include "MassCommonFragments.h"
+#include "MassDebugger.h"
 #include "MassExecutionContext.h"
 #include "MassMovementFragments.h"
 #include "MassNavigationFragments.h"
 #include "Steering/MassSteeringFragments.h"
 #include "Math/UnrealMathUtility.h"
 #include "MassSimulationLOD.h"
+#include "MassNavigationDebug.h"
 
 #define UNSAFE_FOR_MT 0
 #define MOVEMENT_DEBUGDRAW 0	// Set to 1 to see heading debugdraw
@@ -255,6 +257,15 @@ void UMassSteerToMoveTargetProcessor::Execute(FMassEntityManager& EntityManager,
 				const FVector::FReal Distance = Delta.Size();
 				if (Distance > StandingSteeringParams.DeadZoneRadius)
 				{
+#if WITH_MASSGAMEPLAY_DEBUG && UNSAFE_FOR_MT
+					FColor EntityColor = FColor::White;
+					const bool bDisplayDebug = UE::Mass::Debug::IsDebuggingEntity(Entity, &EntityColor);
+					if (bDisplayDebug)
+					{
+						UE_VLOG_UELOG(this, LogMassNavigation, Verbose, TEXT("Standing steering: out of deadzone (Distance: %.2f)"), Distance);
+					}
+#endif
+					
 					SteerDirection = Delta / Distance;
 					if (StandingSteering.bEnteredFromMoveAction)
 					{
@@ -312,12 +323,15 @@ void UMassSteerToMoveTargetProcessor::Execute(FMassEntityManager& EntityManager,
 				UE_VLOG_CIRCLE_THICK(this, LogMassNavigation, Log, MoveTargetCenter, FVector::UpVector, 5.f, EntityColor, /*Thickness*/2,
 					TEXT("MoveTarget\ngoal: %s"), *UEnum::GetDisplayValueAsText(MoveTarget.IntentAtGoal).ToString());
 
+				// Display ghost location
+				UE_VLOG_CIRCLE_THICK(this, LogMassNavigation, Log, Ghost.Location, FVector::UpVector, 5.f, FColor::Silver, /*Thickness*/2, TEXT("Ghost"));
+
 				// Display MoveTarget orientation
 				UE_VLOG_SEGMENT_THICK(this, LogMassNavigation, Log, MoveTargetCenter, MoveTargetCenter + MoveTarget.Forward * 100, EntityColor, /*Thickness*/1,
 					TEXT("MoveTarget\nforward"));
 
 				// Display MoveTarget - current location relation
-				if (FVector::Dist2D(CurrentLocation, MoveTarget.Center) > LookAheadDistance * 1.5f)
+				if (MoveTarget.DesiredSpeed.Get() > 0.f && FVector::Dist2D(CurrentLocation, MoveTarget.Center) > LookAheadDistance * 1.5f)
 				{
 					UE_VLOG_SEGMENT_THICK(this, LogMassNavigation, Log, MoveTargetCenter, CurrentLocation + ZOffset, FColor::Red, /*Thickness*/1, TEXT("LOST"));
 				}
diff --git a/Engine/Plugins/AI/MassAI/Source/MassNavigation/Public/Steering/MassSteeringFragments.h b/Engine/Plugins/AI/MassAI/Source/MassNavigation/Public/Steering/MassSteeringFragments.h
--- a/Engine/Plugins/AI/MassAI/Source/MassNavigation/Public/Steering/MassSteeringFragments.h
+++ b/Engine/Plugins/AI/MassAI/Source/MassNavigation/Public/Steering/MassSteeringFragments.h
@@ -69,11 +69,14 @@ struct FMassStandingSteeringParameters : public FMassConstSharedFragment
 {
 	GENERATED_BODY()
 
+#if WITH_EDITORONLY_DATA 	
 	/** Steering reaction time in seconds. */
-	UPROPERTY(EditAnywhere, Category = "Standing", meta = (ClampMin = "0.05", ForceUnits="s"))
+	UE_DEPRECATED(5.6, "Unused")
+	UPROPERTY()
 	float ReactionTime = 0.3f;
+#endif	
 
-	/** How much the target should deviate from the ghost location before update */
+	/** How much the ghost should deviate from the target before updating the target. */
 	UPROPERTY(EditAnywhere, Category = "Standing", meta = (ClampMin = "0.05", ForceUnits="cm"))
 	float TargetMoveThreshold = 15.0f;
 	
@@ -95,6 +98,7 @@ struct FMassStandingSteeringParameters : public FMassConstSharedFragment
 	UPROPERTY(EditAnywhere, Category = "Standing")
 	float TargetSelectionCooldownVariance = 0.5f;
 
+	/** How much the target should deviate from the current location before updating the force on the agent. */
 	UPROPERTY(EditAnywhere, Category = "Standing", meta = (ForceUnits="cm"))
 	float DeadZoneRadius = 15.0f;
 };
