This is in an experimental module, the deprecation is not really needed.

#rnx

diff --git a/Engine/Plugins/AI/MassAI/Source/MassNavigation/Public/Steering/MassSteeringFragments.h b/Engine/Plugins/AI/MassAI/Source/MassNavigation/Public/Steering/MassSteeringFragments.h
--- a/Engine/Plugins/AI/MassAI/Source/MassNavigation/Public/Steering/MassSteeringFragments.h
+++ b/Engine/Plugins/AI/MassAI/Source/MassNavigation/Public/Steering/MassSteeringFragments.h
@@ -69,13 +69,6 @@ struct FMassStandingSteeringParameters : public FMassConstSharedFragment
 {
 	GENERATED_BODY()
 
-#if WITH_EDITORONLY_DATA 	
-	/** Steering reaction time in seconds. */
-	UE_DEPRECATED(5.6, "Unused")
-	UPROPERTY()
-	float ReactionTime = 0.3f;
-#endif	
-
 	/** How much the ghost should deviate from the target before updating the target. */
 	UPROPERTY(EditAnywhere, Category = "Standing", meta = (ClampMin = "0.05", ForceUnits="cm"))
 	float TargetMoveThreshold = 15.0f;
