[MassAI] Mass dynamic StateTree processors enabled by default.

#rnx

diff --git a/Engine/Plugins/AI/MassAI/Source/MassAIBehavior/Private/MassStateTreeSubsystem.cpp b/Engine/Plugins/AI/MassAI/Source/MassAIBehavior/Private/MassStateTreeSubsystem.cpp
--- a/Engine/Plugins/AI/MassAI/Source/MassAIBehavior/Private/MassStateTreeSubsystem.cpp
+++ b/Engine/Plugins/AI/MassAI/Source/MassAIBehavior/Private/MassStateTreeSubsystem.cpp
@@ -14,7 +14,7 @@
 
 namespace UE::Mass::StateTree
 {
-	bool bDynamicSTProcessorsEnabled = false;
+	bool bDynamicSTProcessorsEnabled = true;
 	FAutoConsoleVariableRef CVarDynamicSTEnabled(TEXT("ai.mass.DynamicSTProcessorsEnabled"), bDynamicSTProcessorsEnabled
 		, TEXT("Whether Dynamic ST processors will be created per distinct ST requirements. Can only be set via code or ini.")
 		, ECVF_ReadOnly);
