PR #12813: [Mass Entity] Fix UMassZoneGraphLaneCacheBoundaryProcessor incorrectly calculates segment normals

diff --git a/Engine/Plugins/AI/MassAI/Source/MassNavigation/Public/MassNavigationUtils.h b/Engine/Plugins/AI/MassAI/Source/MassNavigation/Public/MassNavigationUtils.h
--- a/Engine/Plugins/AI/MassAI/Source/MassNavigation/Public/MassNavigationUtils.h
+++ b/Engine/Plugins/AI/MassAI/Source/MassNavigation/Public/MassNavigationUtils.h
@@ -96,7 +96,7 @@ namespace UE::MassNavigation
 	// Returns left direction from forward and up directions.
 	inline FVector GetLeftDirection(const FVector Forward, const FVector Up)
 	{
-		return FVector::CrossProduct(Forward, Up);
+		return FVector::CrossProduct(Forward, Up).GetSafeNormal();
 	}
 
 	// Computes miter normal in XY plane from two neighbour edge normals.
