Temporary fix for UMassMoveTargetFragmentInitializer using default transform value
(instead of actual location) when initializing MoveTarget.Center.

diff --git a/Engine/Plugins/AI/MassAI/Source/MassAIBehavior/Private/Tasks/MassNavMeshStandTask.cpp b/Engine/Plugins/AI/MassAI/Source/MassAIBehavior/Private/Tasks/MassNavMeshStandTask.cpp
--- a/Engine/Plugins/AI/MassAI/Source/MassAIBehavior/Private/Tasks/MassNavMeshStandTask.cpp
+++ b/Engine/Plugins/AI/MassAI/Source/MassAIBehavior/Private/Tasks/MassNavMeshStandTask.cpp
@@ -9,6 +9,7 @@
 #include "MassStateTreeExecutionContext.h"
 #include "MassSignalSubsystem.h"
 #include "StateTreeLinker.h"
+#include "MassCommonFragments.h"
 
 bool FMassNavMeshStandTask::Link(FStateTreeLinker& Linker)
 {
@@ -16,6 +17,7 @@ bool FMassNavMeshStandTask::Link(FStateTreeLinker& Linker)
 	Linker.LinkExternalData(ShortPathHandle);
 	Linker.LinkExternalData(MovementParamsHandle);
 	Linker.LinkExternalData(MassSignalSubsystemHandle);
+	Linker.LinkExternalData(TransformHandle);
 
 	return true;
 }
@@ -34,6 +36,8 @@ EStateTreeRunStatus FMassNavMeshStandTask::EnterState(FStateTreeExecutionContext
 	const UWorld* World = Context.GetWorld();
 	checkf(World != nullptr, TEXT("A valid world is expected from the execution context."));
 
+	const FVector AgentNavLocation = Context.GetExternalData(TransformHandle).GetTransform().GetLocation();
+	MoveTarget.Center = AgentNavLocation;
 	MoveTarget.CreateNewAction(EMassMovementAction::Stand, *World);
 	
 	const bool bSuccess = UE::MassNavigation::ActivateActionStand(Context.GetOwner(), MassStateTreeContext.GetEntity(), 
diff --git a/Engine/Plugins/AI/MassAI/Source/MassAIBehavior/Public/Tasks/MassNavMeshStandTask.h b/Engine/Plugins/AI/MassAI/Source/MassAIBehavior/Public/Tasks/MassNavMeshStandTask.h
--- a/Engine/Plugins/AI/MassAI/Source/MassAIBehavior/Public/Tasks/MassNavMeshStandTask.h
+++ b/Engine/Plugins/AI/MassAI/Source/MassAIBehavior/Public/Tasks/MassNavMeshStandTask.h
@@ -9,6 +9,7 @@ struct FStateTreeExecutionContext;
 struct FMassMoveTargetFragment;
 struct FMassNavMeshShortPathFragment;
 struct FMassMovementParameters;
+struct FTransformFragment;
 class UZoneGraphSubsystem;
 class UMassSignalSubsystem;
 
@@ -50,4 +51,5 @@ protected:
 	TStateTreeExternalDataHandle<FMassNavMeshShortPathFragment> ShortPathHandle;
 	TStateTreeExternalDataHandle<FMassMovementParameters> MovementParamsHandle;
 	TStateTreeExternalDataHandle<UMassSignalSubsystem> MassSignalSubsystemHandle;
+	TStateTreeExternalDataHandle<FTransformFragment> TransformHandle;
 };
