Fix inverted force generated by environment edges
Added text display to DebugDrawLine and display environment edges index when using avoidance debug display
Change verbosity of path tangent display

diff --git a/Engine/Plugins/AI/MassAI/Source/MassNavMeshNavigation/Private/MassNavMeshNavigationProcessors.cpp b/Engine/Plugins/AI/MassAI/Source/MassNavMeshNavigation/Private/MassNavMeshNavigationProcessors.cpp
--- a/Engine/Plugins/AI/MassAI/Source/MassNavMeshNavigation/Private/MassNavMeshNavigationProcessors.cpp
+++ b/Engine/Plugins/AI/MassAI/Source/MassNavMeshNavigation/Private/MassNavMeshNavigationProcessors.cpp
@@ -229,7 +229,7 @@ void UMassNavMeshPathFollowProcessor::Execute(FMassEntityManager& EntityManager,
 					{
 						const FMassNavMeshPathPoint& CurrPoint = ShortPath.Points[PointIndex];
 						const FVector CurrBase = CurrPoint.Position + ZOffset + FVector(0,0,1);
-						UE_VLOG_SEGMENT_THICK(this, LogMassNavMeshNavigation, Display, CurrBase, CurrBase + CurrPoint.Tangent.GetVector() * 100.0f, LightEntityColor, /*Thickness*/1, TEXT(""));
+						UE_VLOG_SEGMENT_THICK(this, LogMassNavMeshNavigation, Verbose, CurrBase, CurrBase + CurrPoint.Tangent.GetVector() * 100.0f, LightEntityColor, /*Thickness*/1, TEXT(""));
 					}
 
 					// Draw MoveTarget
diff --git a/Engine/Plugins/AI/MassAI/Source/MassNavigation/Private/MassAvoidanceProcessors.cpp b/Engine/Plugins/AI/MassAI/Source/MassNavigation/Private/MassAvoidanceProcessors.cpp
--- a/Engine/Plugins/AI/MassAI/Source/MassNavigation/Private/MassAvoidanceProcessors.cpp
+++ b/Engine/Plugins/AI/MassAI/Source/MassNavigation/Private/MassAvoidanceProcessors.cpp
@@ -436,6 +436,8 @@ void UMassMovingAvoidanceProcessor::Execute(FMassEntityManager& EntityManager, F
 
 				const FVector::FReal AgentRadiusForEnvironment = NavEdges.bExtrudedEdges ? 0 : AgentRadius;
 
+				int32 EdgeIndex = 0;
+				
 				// Collect potential contacts between agent and environment edges.
 				for (const FNavigationAvoidanceEdge& Edge : NavEdges.AvoidanceEdges)
 				{
@@ -508,7 +510,8 @@ void UMassMovingAvoidanceProcessor::Execute(FMassEntityManager& EntityManager, F
 					{
 						// Draw active edges
 						FVector ZOffset = FVector(0., 0., 7.);
-						DebugDrawLine(ObstacleDebugContext, ZOffset + Edge.Start, ZOffset + Edge.End, FColor::Yellow, /*Thickness=*/3.f);
+						DebugDrawLine(ObstacleDebugContext, ZOffset + Edge.Start, ZOffset + Edge.End, FColor::Yellow, /*Thickness=*/3.f,
+							/*bPersistent*/false, *LexToString(EdgeIndex));
 
 						// Draw environment separation distance
 						const FVector SeparationOffset = MovingAvoidanceParams.EnvironmentSeparationDistance * Edge.LeftDir;
@@ -545,7 +548,7 @@ void UMassMovingAvoidanceProcessor::Execute(FMassEntityManager& EntityManager, F
 						FVector AvoidRelPos = HitAgentPos - HitObPos;
 						AvoidRelPos.Z = 0.;	// @todo AT: ignore the z component for now until we clamp the height of obstacles
 						const FVector::FReal AvoidDist = AvoidRelPos.Size();
-						const FVector AvoidNormal = AvoidDist > 0. ? (AvoidRelPos / AvoidDist) : Edge.LeftDir;
+						const FVector AvoidNormal = AvoidDist > UE_KINDA_SMALL_NUMBER ? (AvoidRelPos / AvoidDist) : Edge.LeftDir;
 
 						const FVector::FReal AvoidPen = (EnvironmentPredictiveAvoidanceAgentRadius + MovingAvoidanceParams.PredictiveAvoidanceDistance) - AvoidDist;
 						FVector::FReal AvoidMag = 1.;
@@ -578,6 +581,9 @@ void UMassMovingAvoidanceProcessor::Execute(FMassEntityManager& EntityManager, F
 						}
 #endif // WITH_MASSGAMEPLAY_DEBUG
 					}
+
+					EdgeIndex++;
+					
 				} // edge loop
 
 #if WITH_MASSGAMEPLAY_DEBUG && UNSAFE_FOR_MT
@@ -874,7 +880,7 @@ void UMassMovingAvoidanceProcessor::Execute(FMassEntityManager& EntityManager, F
 				// Calculate penetration at CPA
 				const FVector AvoidRelPos = RelPos + RelVel * CPA;
 				const FVector::FReal AvoidDist = AvoidRelPos.Size();
-				const FVector AvoidConNormal = AvoidDist > 0. ? (AvoidRelPos / AvoidDist) : FVector::ForwardVector;
+				const FVector AvoidConNormal = AvoidDist > UE_KINDA_SMALL_NUMBER ? (AvoidRelPos / AvoidDist) : FVector::ForwardVector;
 
 				FVector AvoidNormal = AvoidConNormal;
 				if (bHasForcedNormal)
diff --git a/Engine/Plugins/AI/MassAI/Source/MassNavigation/Private/MassNavigationDebug.cpp b/Engine/Plugins/AI/MassAI/Source/MassNavigation/Private/MassNavigationDebug.cpp
--- a/Engine/Plugins/AI/MassAI/Source/MassNavigation/Private/MassNavigationDebug.cpp
+++ b/Engine/Plugins/AI/MassAI/Source/MassNavigation/Private/MassNavigationDebug.cpp
@@ -35,14 +35,14 @@ namespace UE::MassNavigation::Debug
 #endif		
 	}
 
-	void DebugDrawLine(const FDebugContext& Context, const FVector& Start, const FVector& End, const FColor& Color, const float Thickness /*= 0.f*/, const bool bPersistent /*= false*/)
+	void DebugDrawLine(const FDebugContext& Context, const FVector& Start, const FVector& End, const FColor& Color, const float Thickness /*= 0.f*/, const bool bPersistent /*= false*/, const FString& Text /*= FString()*/)
 	{
 		if (!DebugIsSelected(Context.Entity))
 		{
 			return;
 		}
 
-		UE_VLOG_SEGMENT_THICK(Context.LogOwner, Context.Category, Log, Start, End, Color, (int16)Thickness, TEXT(""));
+		UE_VLOG_SEGMENT_THICK(Context.LogOwner, Context.Category, Log, Start, End, Color, (int16)Thickness, TEXT("%s"), *Text);
 
 		if (UseDrawDebugHelper() && Context.World)
 		{
diff --git a/Engine/Plugins/AI/MassAI/Source/MassNavigation/Public/MassNavigationDebug.h b/Engine/Plugins/AI/MassAI/Source/MassNavigation/Public/MassNavigationDebug.h
--- a/Engine/Plugins/AI/MassAI/Source/MassNavigation/Public/MassNavigationDebug.h
+++ b/Engine/Plugins/AI/MassAI/Source/MassNavigation/Public/MassNavigationDebug.h
@@ -26,7 +26,7 @@ namespace UE::MassNavigation::Debug
 	
 	MASSNAVIGATION_API bool UseDrawDebugHelper();
 	MASSNAVIGATION_API bool DebugIsSelected(const FMassEntityHandle Entity);
-	MASSNAVIGATION_API void DebugDrawLine(const FDebugContext& Context, const FVector& Start, const FVector& End, const FColor& Color, const float Thickness = 0.f, const bool bPersistent = false);
+	MASSNAVIGATION_API void DebugDrawLine(const FDebugContext& Context, const FVector& Start, const FVector& End, const FColor& Color, const float Thickness = 0.f, const bool bPersistent = false, const FString& Text = FString());
 	MASSNAVIGATION_API void DebugDrawArrow(const FDebugContext& Context, const FVector& Start, const FVector& End, const FColor& Color, const float HeadSize = 8.f, const float Thickness = 1.5f);
 	MASSNAVIGATION_API void DebugDrawSphere(const FDebugContext& Context, const FVector& Center, const FVector::FReal InRadius, const FColor& Color);
 	MASSNAVIGATION_API void DebugDrawBox(const FDebugContext& Context, const FBox& Box, const FColor& Color);
