[MassSmartObject]
PR #11620: forget set InstanceData.LastLane

diff --git a/Engine/Plugins/AI/MassAI/Source/MassAIBehavior/Private/Tasks/MassFindSmartObjectTask.cpp b/Engine/Plugins/AI/MassAI/Source/MassAIBehavior/Private/Tasks/MassFindSmartObjectTask.cpp
--- a/Engine/Plugins/AI/MassAI/Source/MassAIBehavior/Private/Tasks/MassFindSmartObjectTask.cpp
+++ b/Engine/Plugins/AI/MassAI/Source/MassAIBehavior/Private/Tasks/MassFindSmartObjectTask.cpp
@@ -139,6 +139,7 @@ EStateTreeRunStatus FMassFindSmartObjectTask::Tick(FStateTreeExecutionContext& C
 						*LexToString(LaneLocation->DistanceAlongLane),
 						*LexToString(LaneLocation->LaneLength));
 
+					InstanceData.LastLane = LaneLocation->LaneHandle;
 					InstanceData.SearchRequestID = MassSmartObjectHandler.FindCandidatesAsync(RequestingEntity, SOUser.UserTags, ActivityRequirements, { LaneLocation->LaneHandle, LaneLocation->DistanceAlongLane });
 				}
 			}
