[Mass Navigation] Prevent MoveTarget progress to continue without it's entity.
Improved tangent computation on FMassNavMeshShortPathFragment.
Change verbosity of navmesh boundary update log.

diff --git a/Engine/Plugins/AI/MassAI/Source/MassAIReplication/Private/MassReplicationPathHandlers.cpp b/Engine/Plugins/AI/MassAI/Source/MassAIReplication/Private/MassReplicationPathHandlers.cpp
--- a/Engine/Plugins/AI/MassAI/Source/MassAIReplication/Private/MassReplicationPathHandlers.cpp
+++ b/Engine/Plugins/AI/MassAI/Source/MassAIReplication/Private/MassReplicationPathHandlers.cpp
@@ -85,6 +85,7 @@ void FReplicatedAgentPathData::InitEntity(const UWorld& InWorld,
 	OutMoveTarget.Center = LaneLocation.Position;
 	OutMoveTarget.Forward = LaneLocation.Tangent;
 	OutMoveTarget.DistanceToGoal = 0.0f;
+	OutMoveTarget.EntityDistanceToGoal = FMassMoveTargetFragment::UnsetDistance;
 	OutMoveTarget.SlackRadius = 0.0f;
 
 	// Setup initial action request
diff --git a/Engine/Plugins/AI/MassAI/Source/MassNavMeshNavigation/Private/MassNavMeshNavigationBoundaryProcessor.cpp b/Engine/Plugins/AI/MassAI/Source/MassNavMeshNavigation/Private/MassNavMeshNavigationBoundaryProcessor.cpp
--- a/Engine/Plugins/AI/MassAI/Source/MassNavMeshNavigation/Private/MassNavMeshNavigationBoundaryProcessor.cpp
+++ b/Engine/Plugins/AI/MassAI/Source/MassNavMeshNavigation/Private/MassNavMeshNavigationBoundaryProcessor.cpp
@@ -64,7 +64,7 @@ void UMassNavMeshNavigationBoundaryProcessor::Execute(FMassEntityManager& Entity
 				// Draw LastUpdatePosition
 				const FVector ZOffset(0,0,10);
 				constexpr float Radius = 5.f;
-				UE_VLOG_WIRECIRCLE(this, LogMassNavMeshNavigation, Display, NavmeshBoundary.LastUpdatePosition + ZOffset,
+				UE_VLOG_WIRECIRCLE(this, LogMassNavMeshNavigation, Verbose, NavmeshBoundary.LastUpdatePosition + ZOffset,
 					FVector(0,0,1), Radius, FColor::Blue, TEXT("Boundary update"));
 			}
 #endif // WITH_MASSGAMEPLAY_DEBUG
diff --git a/Engine/Plugins/AI/MassAI/Source/MassNavMeshNavigation/Private/MassNavMeshNavigationFragments.cpp b/Engine/Plugins/AI/MassAI/Source/MassNavMeshNavigation/Private/MassNavMeshNavigationFragments.cpp
--- a/Engine/Plugins/AI/MassAI/Source/MassNavMeshNavigation/Private/MassNavMeshNavigationFragments.cpp
+++ b/Engine/Plugins/AI/MassAI/Source/MassNavMeshNavigation/Private/MassNavMeshNavigationFragments.cpp
@@ -39,6 +39,8 @@ bool FMassNavMeshShortPathFragment::RequestShortPath(const TSharedPtr<FNavCorrid
 	// Calculate movement distance and tangent at each point.
 	float PathDistance = 0.f;
 	Points[0].Distance.Set(PathDistance);
+	Points[0].Tangent = FMassSnorm8Vector2D((Points[1].Position - Points[0].Position).GetSafeNormal());
+	
 	for (uint8 PointIndex = 1; PointIndex < NumPoints; PointIndex++)
 	{
 		FMassNavMeshPathPoint& PrevPoint = Points[PointIndex - 1];
@@ -49,13 +51,21 @@ bool FMassNavMeshShortPathFragment::RequestShortPath(const TSharedPtr<FNavCorrid
 		PathDistance += DeltaDistance;
 		Point.Distance.Set(PathDistance);
 
-		// Set tangent of PrevPoint
-		const FVector& Tangent = (Point.Position - PrevPoint.Position).GetSafeNormal();
-		PrevPoint.Tangent = FMassSnorm8Vector2D(Tangent);
+		// Set tangent of Point
+		if (PointIndex + 1 < NumPoints)
+		{
+			FMassNavMeshPathPoint& NextPoint = Points[PointIndex + 1];
+			Point.Tangent = FMassSnorm8Vector2D(0.5 * ((Point.Position - PrevPoint.Position).GetSafeNormal() + (NextPoint.Position - Point.Position).GetSafeNormal()));
+		}
+		else
+		{
+			// Last point
+			Point.Tangent = FMassSnorm8Vector2D((Point.Position - PrevPoint.Position).GetSafeNormal());
+		}
 	}
 
 	// Update ProgressDistance to account for leading points. 
-	ProgressDistance = Points[InNumLeadingPoints].Distance.Get();
+	MoveTargetProgressDistance = Points[InNumLeadingPoints].Distance.Get();
 	
 	return true;
 }
diff --git a/Engine/Plugins/AI/MassAI/Source/MassNavMeshNavigation/Private/MassNavMeshNavigationProcessors.cpp b/Engine/Plugins/AI/MassAI/Source/MassNavMeshNavigation/Private/MassNavMeshNavigationProcessors.cpp
--- a/Engine/Plugins/AI/MassAI/Source/MassNavMeshNavigation/Private/MassNavMeshNavigationProcessors.cpp
+++ b/Engine/Plugins/AI/MassAI/Source/MassNavMeshNavigation/Private/MassNavMeshNavigationProcessors.cpp
@@ -86,11 +86,14 @@ void UMassNavMeshPathFollowProcessor::Execute(FMassEntityManager& EntityManager,
 				if (!bHasSteering || !MoveTarget.bSteeringFallingBehind)
 				{
 					// Update progress
-					ShortPath.ProgressDistance += MoveTarget.DesiredSpeed.Get() * WorldDeltaTime;
+					ShortPath.MoveTargetProgressDistance += MoveTarget.DesiredSpeed.Get() * WorldDeltaTime;
 				}
 				
 				if (!bWasDone)
 				{
+					const FTransformFragment& TransformFragment = TransformList[EntityIt];
+					const FVector EntityLocation = TransformFragment.GetTransform().GetLocation();
+					
 					const uint8 LastIndex = ShortPath.NumPoints - 1;
 
 					// Point index on the short path where we should update the path.
@@ -99,9 +102,29 @@ void UMassNavMeshPathFollowProcessor::Execute(FMassEntityManager& EntityManager,
 					// If the shortpath is partial, it's expected to be full, meaning MaxPoints-NumPointsBeyondUpdate >= 1 so the check below must be valid.
 					check(UpdatePointIndex >= 1);
 
+					// Update entity progress on path (EntityDistanceToGoal)
+					uint8 ClosestSegmentIndex = 0;
+					float ClosestDistanceSquare = FLT_MAX;
+					for (uint8 PointIndex = 0; PointIndex < UpdatePointIndex && PointIndex+1 <= LastIndex; PointIndex++)
+					{
+						// Project on closest segment
+						const float DistanceSquared = FMath::PointDistToSegmentSquared(EntityLocation, ShortPath.Points[PointIndex].Position, ShortPath.Points[PointIndex+1].Position);
+						if (DistanceSquared < ClosestDistanceSquare)
+						{
+							ClosestDistanceSquare = DistanceSquared;
+							ClosestSegmentIndex = PointIndex;
+						}
+					}
+
+					const FMassNavMeshPathPoint PointA = ShortPath.Points[ClosestSegmentIndex];
+					const FMassNavMeshPathPoint PointB = ShortPath.Points[ClosestSegmentIndex+1];
+					const float DistanceOnClosestSegment =  (EntityLocation - PointA.Position).Dot((PointB.Position - PointA.Position).GetSafeNormal());
+					const float EntityProgress = ShortPath.Points[ClosestSegmentIndex].Distance.Get() + DistanceOnClosestSegment;
+					MoveTarget.EntityDistanceToGoal = ShortPath.Points[UpdatePointIndex].Distance.Get() - EntityProgress;
+					
 #if WITH_MASSGAMEPLAY_DEBUG && UNSAFE_FOR_MT
 					UE_CVLOG_UELOG(bDisplayDebug, this, LogMassNavMeshNavigation, Verbose, TEXT("Entity [%s]   ProgressDistance: %f, UpdatePointIndex: %i"),
-						*Entity.DebugGetDescription(), ShortPath.ProgressDistance, UpdatePointIndex);
+						*Entity.DebugGetDescription(), ShortPath.MoveTargetProgressDistance, UpdatePointIndex);
 
 					if (bDisplayDebug)
 					{
@@ -109,28 +132,36 @@ void UMassNavMeshPathFollowProcessor::Execute(FMassEntityManager& EntityManager,
 						const FVector ZOffset(0,0,20);
 						UE_VLOG_CIRCLE(this, LogMassNavMeshNavigation, Display, ShortPath.Points[UpdatePointIndex].Position + ZOffset,
 							FVector(0,0,1), /*Radius*/10.f, FColor::Green, TEXT("%i"), UpdatePointIndex);
+
+						// Display entity progress
+						const float T = (EntityProgress - ShortPath.Points[ClosestSegmentIndex].Distance.Get()) / (PointB.Distance.Get() - PointA.Distance.Get());
+						const FVector ProjectedEntityPosition = FMath::Lerp(PointA.Position, PointB.Position, T);
+						UE_VLOG_SEGMENT(this, LogMassNavMeshNavigation, Display, EntityLocation + ZOffset, ProjectedEntityPosition + ZOffset, FColor::Silver, TEXT(""));
+						UE_VLOG_CIRCLE(this, LogMassNavMeshNavigation, Display, ProjectedEntityPosition + ZOffset,
+							FVector(0,0,1), /*Radius*/10.f, FColor::Silver, TEXT(""));
 					}
 #endif // WITH_MASSGAMEPLAY_DEBUG
 
-					if (ShortPath.ProgressDistance <= 0.0f)
+					if (ShortPath.MoveTargetProgressDistance <= 0.0f)
 					{
 						// Before the start of the path
 						MoveTarget.Center = ShortPath.Points[0].Position;
 						MoveTarget.Forward = ShortPath.Points[0].Tangent.GetVector();
 						MoveTarget.DistanceToGoal = ShortPath.Points[UpdatePointIndex].Distance.Get();
 
-						UE_CVLOG_UELOG(bDisplayDebug, this, LogMassNavMeshNavigation, Log, TEXT("Entity [%s]   before start of path. Distance to goal: %.1f."),
+						UE_CVLOG_UELOG(bDisplayDebug, this, LogMassNavMeshNavigation, Log, TEXT("Entity [%s]   before start of path. EntityDistanceToGoal: %.1f, DistanceToGoal: %.1f."),
 							*Entity.DebugGetDescription(),
+							MoveTarget.EntityDistanceToGoal,
 							MoveTarget.DistanceToGoal);
 					}
-					else if (ShortPath.ProgressDistance <= ShortPath.Points[UpdatePointIndex].Distance.Get())
+					else if (ShortPath.MoveTargetProgressDistance <= ShortPath.Points[UpdatePointIndex].Distance.Get())
 					{
 						// Along the path, interpolate.
 						uint8 PointIndex = 0;
 						while (PointIndex < (ShortPath.NumPoints - 2))
 						{
 							const FMassNavMeshPathPoint& NextPoint = ShortPath.Points[PointIndex + 1];
-							if (ShortPath.ProgressDistance <= NextPoint.Distance.Get())
+							if (ShortPath.MoveTargetProgressDistance <= NextPoint.Distance.Get())
 							{
 								break;
 							}
@@ -139,7 +170,7 @@ void UMassNavMeshPathFollowProcessor::Execute(FMassEntityManager& EntityManager,
 						const FMassNavMeshPathPoint& CurrPoint = ShortPath.Points[PointIndex];
 						const FMassNavMeshPathPoint& NextPoint = ShortPath.Points[PointIndex + 1];
 						
-						const float T = (ShortPath.ProgressDistance - CurrPoint.Distance.Get()) / (NextPoint.Distance.Get() - CurrPoint.Distance.Get());
+						const float T = (ShortPath.MoveTargetProgressDistance - CurrPoint.Distance.Get()) / (NextPoint.Distance.Get() - CurrPoint.Distance.Get());
 
 						// Set move target location using the new progress distance.
 						MoveTarget.Center = FMath::Lerp(CurrPoint.Position, NextPoint.Position, T);
@@ -165,9 +196,6 @@ void UMassNavMeshPathFollowProcessor::Execute(FMassEntityManager& EntityManager,
 						else
 						{
 							// Last section of the path, wait for the steering to complete the movement.
-							const FTransformFragment& TransformFragment = TransformList[EntityIt];
-							const FVector& EntityLocation = TransformFragment.GetTransform().GetLocation();
-
 							const float SegmentSize = ShortPath.Points[UpdatePointIndex].Distance.Get() - ShortPath.Points[UpdatePointIndex-1].Distance.Get();
 							const FVector PreviousPointLocation = ShortPath.Points[UpdatePointIndex-1].Position;
 							FVector PathDir = ShortPath.Points[UpdatePointIndex-1].Tangent.GetVector();
diff --git a/Engine/Plugins/AI/MassAI/Source/MassNavMeshNavigation/Private/MassNavMeshNavigationUtils.cpp b/Engine/Plugins/AI/MassAI/Source/MassNavMeshNavigation/Private/MassNavMeshNavigationUtils.cpp
--- a/Engine/Plugins/AI/MassAI/Source/MassNavMeshNavigation/Private/MassNavMeshNavigationUtils.cpp
+++ b/Engine/Plugins/AI/MassAI/Source/MassNavMeshNavigation/Private/MassNavMeshNavigationUtils.cpp
@@ -16,6 +16,7 @@ namespace UE::MassNavigation
 	{
 		OutShortPath.Reset();
 		InOutMoveTarget.DistanceToGoal = 0.f;
+		InOutMoveTarget.EntityDistanceToGoal = FMassMoveTargetFragment::UnsetDistance;
 		InOutMoveTarget.DesiredSpeed.Set(0.f);
 		
 		if (!ensureMsgf(InOutMoveTarget.GetCurrentAction() == EMassMovementAction::Stand,
diff --git a/Engine/Plugins/AI/MassAI/Source/MassNavMeshNavigation/Public/MassNavMeshNavigationFragments.h b/Engine/Plugins/AI/MassAI/Source/MassNavMeshNavigation/Public/MassNavMeshNavigationFragments.h
--- a/Engine/Plugins/AI/MassAI/Source/MassNavMeshNavigation/Public/MassNavMeshNavigationFragments.h
+++ b/Engine/Plugins/AI/MassAI/Source/MassNavMeshNavigation/Public/MassNavMeshNavigationFragments.h
@@ -49,7 +49,7 @@ struct MASSNAVMESHNAVIGATION_API FMassNavMeshShortPathFragment : public FMassFra
 
 	void Reset()
 	{
-		ProgressDistance = 0.f;
+		MoveTargetProgressDistance = 0.f;
 		EndReachedDistance = 20.f;
 		NumPoints = 0;
 		EndOfPathIntent = EMassMovementAction::Stand;
@@ -68,9 +68,9 @@ struct MASSNAVMESHNAVIGATION_API FMassNavMeshShortPathFragment : public FMassFra
 
 	/** Portal points */
 	TStaticArray<FMassNavMeshPathPoint, MaxPoints> Points;
-
+	
 	/** Current progress distance along the short path. */
-	float ProgressDistance = 0.f;
+	float MoveTargetProgressDistance = 0.f;
 
 	/** Distance from the end of path used to confirm that the destination is reached. */
 	float EndReachedDistance = 20.f;
diff --git a/Engine/Plugins/AI/MassAI/Source/MassNavigation/Private/MassNavigationFragments.cpp b/Engine/Plugins/AI/MassAI/Source/MassNavigation/Private/MassNavigationFragments.cpp
--- a/Engine/Plugins/AI/MassAI/Source/MassNavigation/Private/MassNavigationFragments.cpp
+++ b/Engine/Plugins/AI/MassAI/Source/MassNavigation/Private/MassNavigationFragments.cpp
@@ -4,6 +4,8 @@
 #include "GameFramework/GameStateBase.h"
 #include "Engine/World.h"
 
+const float FMassMoveTargetFragment::UnsetDistance = -1.f;
+
 void FMassMoveTargetFragment::CreateNewAction(const EMassMovementAction InAction, const UWorld& InWorld)
 {
 	ensureMsgf(InWorld.GetNetMode() != NM_Client, TEXT("This version of SetDesiredAction should only be called on the authority."));
diff --git a/Engine/Plugins/AI/MassAI/Source/MassNavigation/Private/MassNavigationProcessors.cpp b/Engine/Plugins/AI/MassAI/Source/MassNavigation/Private/MassNavigationProcessors.cpp
--- a/Engine/Plugins/AI/MassAI/Source/MassNavigation/Private/MassNavigationProcessors.cpp
+++ b/Engine/Plugins/AI/MassAI/Source/MassNavigation/Private/MassNavigationProcessors.cpp
@@ -149,6 +149,7 @@ void UMassMoveTargetFragmentInitializer::Execute(FMassEntityManager& EntityManag
 			MoveTarget.Center = Location.GetTransform().GetLocation();
 			MoveTarget.Forward = Location.GetTransform().GetRotation().Vector();
 			MoveTarget.DistanceToGoal = 0.0f;
+			MoveTarget.EntityDistanceToGoal = FMassMoveTargetFragment::UnsetDistance;
 			MoveTarget.SlackRadius = 0.0f;
 		}
 	});
diff --git a/Engine/Plugins/AI/MassAI/Source/MassNavigation/Private/MassSteeringProcessors.cpp b/Engine/Plugins/AI/MassAI/Source/MassNavigation/Private/MassSteeringProcessors.cpp
--- a/Engine/Plugins/AI/MassAI/Source/MassNavigation/Private/MassSteeringProcessors.cpp
+++ b/Engine/Plugins/AI/MassAI/Source/MassNavigation/Private/MassSteeringProcessors.cpp
@@ -189,7 +189,16 @@ void UMassSteerToMoveTargetProcessor::Execute(FMassEntityManager& EntityManager,
 					DesiredSpeed *= UE::MassNavigation::ArrivalSpeedEnvelope(FMath::Max(ArrivalFade, NormalizedDistanceToSteerTarget));
 				}
 
-				MoveTarget.bSteeringFallingBehind = ForwardOffset < -LookAheadDistance * 0.8;
+				constexpr FVector::FReal FallingBehindScale = 0.8;
+				if (MoveTarget.EntityDistanceToGoal != FMassMoveTargetFragment::UnsetDistance)
+				{
+					MoveTarget.bSteeringFallingBehind = (MoveTarget.EntityDistanceToGoal - MoveTarget.DistanceToGoal) > LookAheadDistance * FallingBehindScale;	
+				}
+				else
+				{
+					// If EntityDistanceToGoal is not available, use ForwardOffset. 
+					MoveTarget.bSteeringFallingBehind = ForwardOffset < -LookAheadDistance * FallingBehindScale;
+				}
 
 				// @todo: This current completely overrides steering, we probably should have one processor that resets the steering at the beginning of the frame.
 				Steering.DesiredVelocity = SteerDirection * DesiredSpeed;
@@ -351,7 +360,7 @@ void UMassSteerToMoveTargetProcessor::Execute(FMassEntityManager& EntityManager,
 
 				// Display Force
 				UE_VLOG_SEGMENT_THICK(this, LogMassNavigation, Log, CurrentLocation + ZOffset, CurrentLocation + Force.Value + ZOffset, FColor::Emerald, /*Thickness*/4,
-					TEXT("Force %.1f"), Force.Value.Length());
+					TEXT("Steering Force %.1f"), Force.Value.Length());
 			}
 #endif // WITH_MASSGAMEPLAY_DEBUG
 			
diff --git a/Engine/Plugins/AI/MassAI/Source/MassNavigation/Public/MassNavigationFragments.h b/Engine/Plugins/AI/MassAI/Source/MassNavigation/Public/MassNavigationFragments.h
--- a/Engine/Plugins/AI/MassAI/Source/MassNavigation/Public/MassNavigationFragments.h
+++ b/Engine/Plugins/AI/MassAI/Source/MassNavigation/Public/MassNavigationFragments.h
@@ -45,6 +45,8 @@ public:
 	double GetCurrentActionServerStartTime() const { return CurrentActionServerStartTime; }
 	uint16 GetCurrentActionID() const { return CurrentActionID; }
 
+	static const float UnsetDistance;
+
 	/** Center of the move target. */
 	FVector Center = FVector::ZeroVector;
 
@@ -54,6 +56,9 @@ public:
 	/** Distance remaining to the movement goal. */
 	float DistanceToGoal = 0.0f;
 
+	/** Projected progress distance of the entity using the path. */
+	float EntityDistanceToGoal = UnsetDistance;
+
 	/** Allowed deviation around the movement target. */
 	float SlackRadius = 0.0f;
 
diff --git a/Engine/Plugins/AI/MassAI/Source/MassZoneGraphNavigation/Private/MassZoneGraphNavigationProcessors.cpp b/Engine/Plugins/AI/MassAI/Source/MassZoneGraphNavigation/Private/MassZoneGraphNavigationProcessors.cpp
--- a/Engine/Plugins/AI/MassAI/Source/MassZoneGraphNavigation/Private/MassZoneGraphNavigationProcessors.cpp
+++ b/Engine/Plugins/AI/MassAI/Source/MassZoneGraphNavigation/Private/MassZoneGraphNavigationProcessors.cpp
@@ -74,8 +74,6 @@ void UMassZoneGraphLocationInitializer::Execute(FMassEntityManager& EntityManage
 				
 				MoveTarget.Center = AgentLocation;
 				MoveTarget.Forward = NearestLane.Tangent;
-				MoveTarget.DistanceToGoal = 0.0f;
-				MoveTarget.SlackRadius = 0.0f;
 			}
 			else
 			{
@@ -85,9 +83,11 @@ void UMassZoneGraphLocationInitializer::Execute(FMassEntityManager& EntityManage
 
 				MoveTarget.Center = AgentLocation;
 				MoveTarget.Forward = FVector::ForwardVector;
-				MoveTarget.DistanceToGoal = 0.0f;
-				MoveTarget.SlackRadius = 0.0f;
 			}
+
+			MoveTarget.DistanceToGoal = 0.0f;
+			MoveTarget.EntityDistanceToGoal = FMassMoveTargetFragment::UnsetDistance;
+			MoveTarget.SlackRadius = 0.0f;
 		}
 	});
 }
diff --git a/Engine/Plugins/AI/MassAI/Source/MassZoneGraphNavigation/Private/MassZoneGraphNavigationUtils.cpp b/Engine/Plugins/AI/MassAI/Source/MassZoneGraphNavigation/Private/MassZoneGraphNavigationUtils.cpp
--- a/Engine/Plugins/AI/MassAI/Source/MassZoneGraphNavigation/Private/MassZoneGraphNavigationUtils.cpp
+++ b/Engine/Plugins/AI/MassAI/Source/MassZoneGraphNavigation/Private/MassZoneGraphNavigationUtils.cpp
@@ -27,6 +27,7 @@ namespace UE::MassNavigation
 		OutShortPath.Reset();
 		OutCachedLane.Reset();
 		InOutMoveTarget.DistanceToGoal = 0.0f;
+		InOutMoveTarget.EntityDistanceToGoal = FMassMoveTargetFragment::UnsetDistance;
 		InOutMoveTarget.DesiredSpeed.Set(0.0f);
 
 		if (!ensureMsgf(InOutMoveTarget.GetCurrentAction() == EMassMovementAction::Move, TEXT("Expecting action 'Move': Invalid action %u"), InOutMoveTarget.GetCurrentAction()))
@@ -89,6 +90,7 @@ namespace UE::MassNavigation
 		ShortPath.Reset();
 		CachedLane.Reset();
 		MoveTarget.DistanceToGoal = 0.0f;
+		MoveTarget.EntityDistanceToGoal = FMassMoveTargetFragment::UnsetDistance;
 		MoveTarget.DesiredSpeed.Set(0.0f);
 
 		if (!ensureMsgf(MoveTarget.GetCurrentAction() == EMassMovementAction::Stand, TEXT("Expecting action 'Stand': Invalid action %u"), MoveTarget.GetCurrentAction()))
@@ -113,6 +115,7 @@ namespace UE::MassNavigation
 							   FMassMoveTargetFragment& MoveTarget)
 	{
 		MoveTarget.DistanceToGoal = 0.0f;
+		MoveTarget.EntityDistanceToGoal = FMassMoveTargetFragment::UnsetDistance;
 		MoveTarget.DesiredSpeed.Set(0.0f);
 
 		if (!ensureMsgf(MoveTarget.GetCurrentAction() == EMassMovementAction::Animate, TEXT("Expecting action 'Animate': Invalid action %u"), MoveTarget.GetCurrentAction()))
