[Mass Navigation] Clamp steering force to respect MaxAcceleration prameter from FMassMovementParameters.

diff --git a/Engine/Plugins/AI/MassAI/Source/MassNavigation/Private/MassSteeringProcessors.cpp b/Engine/Plugins/AI/MassAI/Source/MassNavigation/Private/MassSteeringProcessors.cpp
--- a/Engine/Plugins/AI/MassAI/Source/MassNavigation/Private/MassSteeringProcessors.cpp
+++ b/Engine/Plugins/AI/MassAI/Source/MassNavigation/Private/MassSteeringProcessors.cpp
@@ -294,6 +294,7 @@ void UMassSteerToMoveTargetProcessor::Execute(FMassEntityManager& EntityManager,
 					// @todo: This current completely overrides steering, we probably should have one processor that resets the steering at the beginning of the frame.
 					Steering.DesiredVelocity = SteerDirection * DesiredSpeed;
 					Force.Value = SteerK * (Steering.DesiredVelocity - Velocity.Value); // Goal force
+					Force.Value = Force.Value.GetClampedToMaxSize(MovementParams.MaxAcceleration); 
 				}
 				else
 				{
