[Mass Avoidance] Simplify avoidance debug display, improved consistency and revisit coloring.

diff --git a/Engine/Plugins/AI/MassAI/Source/MassNavigation/Private/MassAvoidanceProcessors.cpp b/Engine/Plugins/AI/MassAI/Source/MassNavigation/Private/MassAvoidanceProcessors.cpp
--- a/Engine/Plugins/AI/MassAI/Source/MassNavigation/Private/MassAvoidanceProcessors.cpp
+++ b/Engine/Plugins/AI/MassAI/Source/MassNavigation/Private/MassAvoidanceProcessors.cpp
@@ -33,6 +33,7 @@ namespace UE::MassAvoidance
 		bool bEnableSettingsForExtendingColliders = true;
 		bool bUseAdjacentCorridors = true;
 		bool bUseDrawDebugHelpers = false;
+		bool bEnableDetailedDebug = false;
 	} // Tweakables
 
 	FAutoConsoleVariableRef Vars[] = 
@@ -40,7 +41,8 @@ namespace UE::MassAvoidance
 		FAutoConsoleVariableRef(TEXT("ai.mass.avoidance.EnableEnvironmentAvoidance"), Tweakables::bEnableEnvironmentAvoidance, TEXT("Set to false to disable avoidance forces for environment (for debug purposes)."), ECVF_Cheat),
 		FAutoConsoleVariableRef(TEXT("ai.mass.avoidance.EnableSettingsForExtendingColliders"), Tweakables::bEnableSettingsForExtendingColliders, TEXT("Set to false to disable using different settings for extending obstacles (for debug purposes)."), ECVF_Cheat),
 		FAutoConsoleVariableRef(TEXT("ai.mass.avoidance.UseAdjacentCorridors"), Tweakables::bUseAdjacentCorridors, TEXT("Set to false to disable usage of adjacent lane width."), ECVF_Cheat),
-		FAutoConsoleVariableRef(TEXT("ai.mass.avoidance.UseDrawDebugHelpers"), Tweakables::bUseDrawDebugHelpers, TEXT("Use debug draw helpers in addition to visual logs."), ECVF_Cheat)
+		FAutoConsoleVariableRef(TEXT("ai.mass.avoidance.UseDrawDebugHelpers"), Tweakables::bUseDrawDebugHelpers, TEXT("Use debug draw helpers in addition to visual logs."), ECVF_Cheat),
+		FAutoConsoleVariableRef(TEXT("ai.mass.avoidance.EnableDetailedDebug"), Tweakables::bEnableDetailedDebug, TEXT("Display additional avoidance debug information."), ECVF_Cheat)
 	};
 
 	constexpr int32 MaxExpectedAgentsPerCell = 6;
@@ -157,31 +159,51 @@ namespace UE::MassAvoidance
 	using namespace UE::MassNavigation::Debug;
 	
 	// Colors
-	static const FColor CurrentAgentColor = FColor::Emerald;
-
-	static const FColor VelocityColor = FColor::Black;
-	static const FColor DesiredVelocityColor = FColor::Yellow;
-	static const FColor FinalSteeringForceColor = FColor::Cyan;
-	static constexpr float BigArrowThickness = 6.f;
-	static constexpr float BigArrowHeadSize = 12.f;
+	static constexpr FColor Amber(255,179,0);
+	static constexpr FColor Orange(251,140,0);
+	static constexpr FColor OrangeRed(244,81,30);
+	static constexpr FColor Cyan(0,172,193);
+	static constexpr FColor Blue(2,155,229);
+	static constexpr FColor Indigo(57,73,171);
+	static constexpr FColor Yellow(253, 216, 53);
+	static constexpr FColor Teal(0, 137, 123);
+	static constexpr FColor Lime(172, 222, 51);
+
+	static constexpr FColor CurrentAgentColor = Lime;
+	static const	 FColor VelocityColor = FColor::Black;
+	static constexpr FColor DesiredVelocityColor = Yellow;
+	static constexpr FColor FinalSteeringForceColor = Teal;
 
 	// Agents colors
-	static const FColor AgentsColor = FColor::Orange;
-	static constexpr FColor AgentSeparationForceColor = FColor(255, 145, 71);	// Orange red
-	static const FColor AgentAvoidForceColor = AgentsColor;
+	static constexpr FColor AgentsColor = Amber;
+	static constexpr FColor AgentAvoidForceColor = Orange;
+	static constexpr FColor AgentSeparationForceColor = OrangeRed;
 	
 	// Obstacles colors
-	static const FColor ObstacleColor = FColor::Blue;
-	static const FColor ObstacleContactNormalColor = FColor::Silver;
-	static const FColor ObstacleAvoidForceColor = FColor::Magenta;
-	static const FColor ObstacleAvoidFuturHitPosColor = FColor::Magenta;
-	static constexpr FColor ObstacleSeparationForceColor = FColor(255, 66, 66);	// Bright red
+	static constexpr FColor ObstacleColor = Cyan;	// edges
+	static constexpr FColor ObstacleAvoidForceColor = Blue;
+	static constexpr FColor ObstacleSeparationForceColor = Indigo;
+	static const	FColor ObstacleContactNormalColor = FColor::Silver;
 
 	// Ghost colors
 	static const FColor GhostColor = FColor::Silver;
 	static const FColor GhostSteeringForceColor = FColor::Silver;
-	
-	static const FVector DebugAgentHeightOffset = FVector(0., 0., 185.);
+
+	// Forces thickness
+	static constexpr float AvoidThickness = 3.f;
+	static constexpr float SeparationThickness = 3.f;
+	static constexpr float SummedForcesThickness = 5.f;
+
+	// Output force
+	static constexpr float SteeringThickness = 8.f;
+	static constexpr float SteeringArrowHeadSize = 12.f;
+
+	// Height offsets
+	static const FVector DebugAgentHeightOffset = FVector(0., 0., 175.);
+	static const FVector DebugInputForceHeight = FVector(0., 0., 181.);
+	static const FVector DebugAgentAvoidHeightOffset = FVector(0., 0., 182.);
+	static const FVector DebugAgentSeparationHeightOffset = FVector(0., 0., 183.);
+	static const FVector DebugOutputForcesHeight = FVector(0., 0., 184.);
 	static const FVector DebugLowCylinderOffset = FVector(0., 0., 20.);
 	
 	// Local debug utils
@@ -215,14 +237,14 @@ namespace UE::MassAvoidance
 		}
 	}
 
-	static void DebugDrawForce(const FDebugContext& Context, const FVector& Start, const FVector& End, const FColor& Color)
+	static void DebugDrawForce(const FDebugContext& Context, const FVector& Start, const FVector& End, const FColor& Color, const float Thickness)
 	{
-		DebugDrawArrow(Context, Start, End, Color, /*HeadSize*/4.f, /*Thickness*/3.f);
+		DebugDrawArrow(Context, Start, End, Color, /*HeadSize*/4.f, Thickness);
 	}
 
 	static void DebugDrawSummedForce(const FDebugContext& Context, const FVector& Start, const FVector& End, const FColor& Color)
 	{
-		DebugDrawArrow(Context, Start + FVector(0.,0.,1.), End + FVector(0., 0., 1.), Color, /*HeadSize*/8.f, /*Thickness*/6.f);
+		DebugDrawArrow(Context, Start + FVector(0.,0.,1.), End + FVector(0., 0., 1.), Color, /*HeadSize*/8.f, SummedForcesThickness);
 	}
 
 #endif // WITH_MASSGAMEPLAY_DEBUG
@@ -398,13 +420,13 @@ void UMassMovingAvoidanceProcessor::Execute(FMassEntityManager& EntityManager, F
 					TEXT("%s"), *Entity.DebugGetDescription(), TEXT("%s"), *Entity.DebugGetDescription());
 
 				// Draw current velocity (black)
-				UE::MassAvoidance::DebugDrawVelocity(VelocitiesDebugContext, AgentLocation + UE::MassAvoidance::DebugAgentHeightOffset,
-				                                     AgentLocation + UE::MassAvoidance::DebugAgentHeightOffset + AgentVelocity, UE::MassAvoidance::VelocityColor);
+				UE::MassAvoidance::DebugDrawVelocity(VelocitiesDebugContext, AgentLocation + UE::MassAvoidance::DebugInputForceHeight,
+				                                     AgentLocation + UE::MassAvoidance::DebugInputForceHeight + AgentVelocity, UE::MassAvoidance::VelocityColor);
 
 				// Draw initial steering force
-				DebugDrawArrow(BaseDebugContext, AgentLocation + UE::MassAvoidance::DebugAgentHeightOffset,
-					AgentLocation +UE::MassAvoidance:: DebugAgentHeightOffset + SteeringForce,
-					UE::MassAvoidance::CurrentAgentColor, UE::MassAvoidance::BigArrowHeadSize, UE::MassAvoidance::BigArrowThickness);
+				DebugDrawArrow(BaseDebugContext, AgentLocation + UE::MassAvoidance::DebugInputForceHeight,
+					AgentLocation + UE::MassAvoidance:: DebugInputForceHeight + SteeringForce,
+					UE::MassAvoidance::CurrentAgentColor, UE::MassAvoidance::SteeringArrowHeadSize, UE::MassAvoidance::SteeringThickness);
 
 				// Draw center
 				DebugDrawSphere(BaseDebugContext, AgentLocation, /*Radius*/2.f, UE::MassAvoidance::CurrentAgentColor);
@@ -427,8 +449,8 @@ void UMassMovingAvoidanceProcessor::Execute(FMassEntityManager& EntityManager, F
 
 #if WITH_MASSGAMEPLAY_DEBUG && UNSAFE_FOR_MT
 				// Draw desired velocity (yellow)
-				UE::MassAvoidance::DebugDrawVelocity(VelocitiesDebugContext, AgentLocation + UE::MassAvoidance::DebugAgentHeightOffset,
-					AgentLocation + UE::MassAvoidance::DebugAgentHeightOffset + DesiredVelocity, UE::MassAvoidance::DesiredVelocityColor);
+				UE::MassAvoidance::DebugDrawVelocity(VelocitiesDebugContext, AgentLocation + UE::MassAvoidance::DebugInputForceHeight,
+					AgentLocation + UE::MassAvoidance::DebugInputForceHeight + DesiredVelocity, UE::MassAvoidance::DesiredVelocityColor);
 #endif // WITH_MASSGAMEPLAY_DEBUG
 
 				OldSteeringForce = SteeringForce;
@@ -438,7 +460,7 @@ void UMassMovingAvoidanceProcessor::Execute(FMassEntityManager& EntityManager, F
 
 				int32 EdgeIndex = 0;
 				
-				// Collect potential contacts between agent and environment edges.
+				// Collect potential contacts between agent and environment edges (obstacles).
 				for (const FNavigationAvoidanceEdge& Edge : NavEdges.AvoidanceEdges)
 				{
 					const FVector EdgeDiff = Edge.End - Edge.Start;
@@ -506,23 +528,26 @@ void UMassMovingAvoidanceProcessor::Execute(FMassEntityManager& EntityManager, F
 					Contacts.Add(Contact);
 
 #if WITH_MASSGAMEPLAY_DEBUG && UNSAFE_FOR_MT
-					if (bAlongTheEdge)
+					if (UE::MassAvoidance::Tweakables::bEnableDetailedDebug)
 					{
-						// Draw active edges
-						FVector ZOffset = FVector(0., 0., 7.);
-						DebugDrawLine(ObstacleDebugContext, ZOffset + Edge.Start, ZOffset + Edge.End, FColor::Yellow, /*Thickness=*/3.f,
-							/*bPersistent*/false, *LexToString(EdgeIndex));
-
-						// Draw environment separation distance
-						const FVector SeparationOffset = MovingAvoidanceParams.EnvironmentSeparationDistance * Edge.LeftDir;
-						DebugDrawLine(ObstacleDebugContext, ZOffset + Edge.Start + SeparationOffset, ZOffset + Edge.End + SeparationOffset, FColor::Yellow, /*Thickness=*/1.f);
-					}
+						if (bAlongTheEdge)
+						{
+							// Draw active edges
+							FVector ZOffset = FVector(0., 0., 7.);
+							DebugDrawLine(ObstacleDebugContext, ZOffset + Edge.Start, ZOffset + Edge.End, FColor::Yellow, /*Thickness=*/3.f,
+								/*bPersistent*/false, *LexToString(EdgeIndex));
+
+							// Draw environment separation distance
+							const FVector SeparationOffset = MovingAvoidanceParams.EnvironmentSeparationDistance * Edge.LeftDir;
+							DebugDrawLine(ObstacleDebugContext, ZOffset + Edge.Start + SeparationOffset, ZOffset + Edge.End + SeparationOffset, FColor::Yellow, /*Thickness=*/1.f);
+						}
 					
-					if (bBehindTheEdge)
-					{
-						// Draw inactive eges
-						FVector Offset = FVector(0., 0., 9.);
-						DebugDrawLine(ObstacleDebugContext, Offset + Edge.Start, Offset + Edge.End, FColor::Silver, /*Thickness=*/2.f);	
+						if (bBehindTheEdge)
+						{
+							// Draw inactive eges
+							FVector Offset = FVector(0., 0., 9.);
+							DebugDrawLine(ObstacleDebugContext, Offset + Edge.Start, Offset + Edge.End, FColor::Silver, /*Thickness=*/2.f);	
+						}
 					}
 #endif // WITH_MASSGAMEPLAY_DEBUG
 					
@@ -532,10 +557,14 @@ void UMassMovingAvoidanceProcessor::Execute(FMassEntityManager& EntityManager, F
 						// Avoid edges
 						
 #if WITH_MASSGAMEPLAY_DEBUG && UNSAFE_FOR_MT
-						// Draw environment predictive avoidance distance
-						FVector ZOffset = FVector(0., 0., 7.);
-						const FVector AvoidanceOffset = MovingAvoidanceParams.PredictiveAvoidanceDistance * Edge.LeftDir;
-						DebugDrawLine(ObstacleDebugContext, ZOffset + Edge.Start + AvoidanceOffset, ZOffset + Edge.End + AvoidanceOffset, UE::MassAvoidance::ObstacleAvoidFuturHitPosColor, /*Thickness=*/1.f);
+						if (UE::MassAvoidance::Tweakables::bEnableDetailedDebug)
+						{
+							// Draw environment predictive avoidance distance
+							FVector ZOffset = FVector(0., 0., 7.);
+							const FVector AvoidanceOffset = MovingAvoidanceParams.PredictiveAvoidanceDistance * Edge.LeftDir;
+							DebugDrawLine(ObstacleDebugContext, ZOffset + Edge.Start + AvoidanceOffset, ZOffset + Edge.End + AvoidanceOffset,
+								UE::MassAvoidance::ObstacleAvoidForceColor, /*Thickness=*/1.f);
+						}
 #endif // WITH_MASSGAMEPLAY_DEBUG
 
 						const FVector::FReal CPA = UE::MassAvoidance::ComputeClosestPointOfApproach(FVector2D(AgentLocation), FVector2D(DesiredVelocity), AgentRadiusForEnvironment,
@@ -567,17 +596,21 @@ void UMassMovingAvoidanceProcessor::Execute(FMassEntityManager& EntityManager, F
 #if WITH_MASSGAMEPLAY_DEBUG && UNSAFE_FOR_MT
 						if (!AvoidForce.IsNearlyZero())
 						{
-							// Draw contact normal
-							DebugDrawArrow(ObstacleDebugContext, ConPos, ConPos + (Contact.Distance * Contact.Normal), UE::MassAvoidance::ObstacleContactNormalColor, /*HeadSize=*/ 5.f);
-							DebugDrawSphere(ObstacleDebugContext, ConPos, 2.5f, UE::MassAvoidance::ObstacleContactNormalColor);
+							if (UE::MassAvoidance::Tweakables::bEnableDetailedDebug)
+							{
+								// Draw contact normal
+								DebugDrawArrow(ObstacleDebugContext, ConPos, ConPos + (Contact.Distance * Contact.Normal), UE::MassAvoidance::ObstacleContactNormalColor, /*HeadSize=*/ 5.f);
+								DebugDrawSphere(ObstacleDebugContext, ConPos, 2.5f, UE::MassAvoidance::ObstacleContactNormalColor);
+							}
 
 							// Draw future hit pos with edge
-							DebugDrawSphere(ObstacleDebugContext, HitAgentPos, 1.f, UE::MassAvoidance::ObstacleAvoidFuturHitPosColor);
-							DebugDrawCircle(ObstacleDebugContext, HitAgentPos, AgentRadius, UE::MassAvoidance::ObstacleAvoidFuturHitPosColor);
-							DebugDrawLine(ObstacleDebugContext, AgentLocation, HitAgentPos, UE::MassAvoidance::ObstacleAvoidFuturHitPosColor);
+							DebugDrawSphere(ObstacleDebugContext, HitAgentPos, 1.f, UE::MassAvoidance::ObstacleAvoidForceColor);
+							DebugDrawCircle(ObstacleDebugContext, HitAgentPos, AgentRadius, UE::MassAvoidance::ObstacleAvoidForceColor);
+							DebugDrawLine(ObstacleDebugContext, AgentLocation, HitAgentPos, UE::MassAvoidance::ObstacleAvoidForceColor);
 
-							// Draw individual predictive avoidance obstacle force
-							UE::MassAvoidance::DebugDrawForce(ObstacleDebugContext, HitObPos, HitObPos + AvoidForce, UE::MassAvoidance::ObstacleAvoidForceColor);
+							// Draw individual predictive obstacle avoidance forces
+							UE::MassAvoidance::DebugDrawForce(ObstacleDebugContext, HitObPos, HitObPos + AvoidForce,
+								UE::MassAvoidance::ObstacleAvoidForceColor, UE::MassAvoidance::AvoidThickness);
 						}
 #endif // WITH_MASSGAMEPLAY_DEBUG
 					}
@@ -590,14 +623,18 @@ void UMassMovingAvoidanceProcessor::Execute(FMassEntityManager& EntityManager, F
 				// Draw total steering force to avoid obstacles
 				const FVector EnvironmentAvoidSteeringForce = SteeringForce - OldSteeringForce;
 				UE::MassAvoidance::DebugDrawSummedForce(ObstacleDebugContext,
-					AgentLocation + UE::MassAvoidance::DebugAgentHeightOffset,
-					AgentLocation + UE::MassAvoidance::DebugAgentHeightOffset + EnvironmentAvoidSteeringForce,
+					AgentLocation + UE::MassAvoidance::DebugAgentAvoidHeightOffset,
+					AgentLocation + UE::MassAvoidance::DebugAgentAvoidHeightOffset + EnvironmentAvoidSteeringForce,
 					UE::MassAvoidance::ObstacleAvoidForceColor);
 
-				// Draw all contact points
-				for (const FEnvironmentContact& Contact : Contacts) 
+				if (UE::MassAvoidance::Tweakables::bEnableDetailedDebug)
 				{
-					DebugDrawSphere(ObstacleDebugContext, Contact.Position + FVector(0.f, 0.f, Contact.bBehindTheEdge ? 10.f : 0.f), 5.f, Contact.bBehindTheEdge ? FColor::Black: FColor::Cyan);					
+					// Draw all contact points
+					for (const FEnvironmentContact& Contact : Contacts) 
+					{
+						DebugDrawSphere(ObstacleDebugContext, Contact.Position + FVector(0.f, 0.f, Contact.bBehindTheEdge ? 10.f : 0.f),
+							5.f, Contact.bBehindTheEdge ? FColor::Black: FColor::Cyan);					
+					}
 				}
 #endif // WITH_MASSGAMEPLAY_DEBUG
 				
@@ -632,22 +669,16 @@ void UMassMovingAvoidanceProcessor::Execute(FMassEntityManager& EntityManager, F
 							SeparationMag,
 							*SeparationForce.ToCompactString()
 							);
-
-						FVector ZOffset = FVector(0., 0., 7.);
 						
 						// Draw contact normal
-						float SphereRadius = 2.5f;
-						FColor ContactColor = FColor::Yellow;
 						if (!SeparationForce.IsNearlyZero())
 						{
-							SphereRadius = 5.f;
-							ContactColor = UE::MassAvoidance::ObstacleSeparationForceColor;
-							
-							// Draw individual contact forces
-							UE::MassAvoidance::DebugDrawForce(ObstacleDebugContext, Contact.Position + ZOffset, Contact.Position + SeparationForce + ZOffset, ContactColor);
+							// Draw individual separation forces
+							const FVector ZOffset = FVector(0., 0., 7.);
+							UE::MassAvoidance::DebugDrawForce(ObstacleDebugContext, Contact.Position + ZOffset,
+								Contact.Position + SeparationForce + ZOffset,
+								UE::MassAvoidance::ObstacleSeparationForceColor, UE::MassAvoidance::SeparationThickness);
 						}
-						DebugDrawSphere(ObstacleDebugContext, Contact.Position + ZOffset, SphereRadius, ContactColor);
-						DebugDrawArrow(ObstacleDebugContext, Contact.Position + ZOffset, Contact.Position + (Contact.Distance * Contact.Normal) + ZOffset, ContactColor, /*HeadSize=*/ 5.f);
 	#endif // WITH_MASSGAMEPLAY_DEBUG
 					}
 				}
@@ -656,8 +687,8 @@ void UMassMovingAvoidanceProcessor::Execute(FMassEntityManager& EntityManager, F
 				// Draw total steering force to separate from close edges
 				const FVector TotalSeparationForce = SteeringForce - SteeringForceBeforeSeparation;
 				UE::MassAvoidance::DebugDrawSummedForce(ObstacleDebugContext,
-					AgentLocation + UE::MassAvoidance::DebugAgentHeightOffset,
-					AgentLocation + UE::MassAvoidance::DebugAgentHeightOffset + TotalSeparationForce,
+					AgentLocation + UE::MassAvoidance::DebugAgentSeparationHeightOffset,
+					AgentLocation + UE::MassAvoidance::DebugAgentSeparationHeightOffset + TotalSeparationForce,
 					UE::MassAvoidance::ObstacleSeparationForceColor);
 
 				// Display close obstacle edges
@@ -910,10 +941,13 @@ void UMassMovingAvoidanceProcessor::Execute(FMassEntityManager& EntityManager, F
 				}
 
 				// Draw agent contact separation force
-				UE::MassAvoidance::DebugDrawSummedForce(AgentDebugContext,
-					Collider.Location + UE::MassAvoidance::DebugAgentHeightOffset,
-					Collider.Location + UE::MassAvoidance::DebugAgentHeightOffset + SeparationForce,
-					UE::MassAvoidance::AgentSeparationForceColor); 
+				if (!SeparationForce.IsNearlyZero())
+				{
+					UE::MassAvoidance::DebugDrawForce(AgentDebugContext,
+						Collider.Location + UE::MassAvoidance::DebugAgentSeparationHeightOffset,
+						Collider.Location + UE::MassAvoidance::DebugAgentSeparationHeightOffset + SeparationForce,
+						UE::MassAvoidance::AgentSeparationForceColor, UE::MassAvoidance::SeparationThickness); 
+				}
 				
 				if (AvoidForce.Size() > 0.)
 				{
@@ -940,15 +974,15 @@ void UMassMovingAvoidanceProcessor::Execute(FMassEntityManager& EntityManager, F
 
 					// Draw agent avoid force
 					UE::MassAvoidance::DebugDrawForce(AgentDebugContext,
-						OtherHitPosition + UE::MassAvoidance::DebugAgentHeightOffset,
-						OtherHitPosition + UE::MassAvoidance::DebugAgentHeightOffset + AvoidForce,
-						UE::MassAvoidance::AgentAvoidForceColor);
+						OtherHitPosition + UE::MassAvoidance::DebugAgentAvoidHeightOffset,
+						OtherHitPosition + UE::MassAvoidance::DebugAgentAvoidHeightOffset + AvoidForce,
+						UE::MassAvoidance::AgentAvoidForceColor, UE::MassAvoidance::AvoidThickness);
 				}
 #endif // WITH_MASSGAMEPLAY_DEBUG
 			} // close entities loop
 
 			SteeringForce *= NearStartScaling * NearEndScaling;
-			
+
 			Force.Value = UE::MassNavigation::ClampVector(SteeringForce, MaxSteerAccel); // Assume unit mass
 
 #if WITH_MASSGAMEPLAY_DEBUG && UNSAFE_FOR_MT
@@ -956,21 +990,21 @@ void UMassMovingAvoidanceProcessor::Execute(FMassEntityManager& EntityManager, F
 
 			// Draw total steering force to separate agents
 			UE::MassAvoidance::DebugDrawSummedForce(AgentDebugContext,
-				AgentLocation + UE::MassAvoidance::DebugAgentHeightOffset,
-				AgentLocation + UE::MassAvoidance::DebugAgentHeightOffset + TotalAgentSeparationForce,
+				AgentLocation + UE::MassAvoidance::DebugAgentSeparationHeightOffset,
+				AgentLocation + UE::MassAvoidance::DebugAgentSeparationHeightOffset + TotalAgentSeparationForce,
 				UE::MassAvoidance::AgentSeparationForceColor);
 
 			// Draw total steering force to avoid agents
 			UE::MassAvoidance::DebugDrawSummedForce(AgentDebugContext,
-				AgentLocation + UE::MassAvoidance::DebugAgentHeightOffset,
-				AgentLocation + UE::MassAvoidance::DebugAgentHeightOffset + AgentAvoidSteeringForce,
+				AgentLocation + UE::MassAvoidance::DebugAgentAvoidHeightOffset,
+				AgentLocation + UE::MassAvoidance::DebugAgentAvoidHeightOffset + AgentAvoidSteeringForce,
 				UE::MassAvoidance::AgentAvoidForceColor);
 
-			// Draw final steering force adding to the agent velocity
+			// Draw final steering force adding
 			UE::MassAvoidance::DebugDrawArrow(BaseDebugContext, 
-				AgentLocation + AgentVelocity + UE::MassAvoidance::DebugAgentHeightOffset,
-				AgentLocation + AgentVelocity + UE::MassAvoidance::DebugAgentHeightOffset + Force.Value,
-				UE::MassAvoidance::FinalSteeringForceColor, UE::MassAvoidance::BigArrowHeadSize, UE::MassAvoidance::BigArrowThickness);
+				AgentLocation + UE::MassAvoidance::DebugOutputForcesHeight,
+				AgentLocation + UE::MassAvoidance::DebugOutputForcesHeight + Force.Value,
+				UE::MassAvoidance::FinalSteeringForceColor, UE::MassAvoidance::SteeringArrowHeadSize, UE::MassAvoidance::SteeringThickness);
 #endif // WITH_MASSGAMEPLAY_DEBUG
 		}
 	});
@@ -1197,7 +1231,7 @@ void UMassStandingAvoidanceProcessor::Execute(FMassEntityManager& EntityManager,
 					UE::MassAvoidance::DebugDrawForce(AgentDebugContext,
 						Ghost.Location + UE::MassAvoidance::DebugAgentHeightOffset,
 						Ghost.Location + UE::MassAvoidance::DebugAgentHeightOffset + SeparationForce,
-						UE::MassAvoidance::AgentAvoidForceColor);
+						UE::MassAvoidance::AgentAvoidForceColor, UE::MassAvoidance::AvoidThickness);
 #endif // WITH_MASSGAMEPLAY_DEBUG
 				}
 				else
@@ -1229,7 +1263,7 @@ void UMassStandingAvoidanceProcessor::Execute(FMassEntityManager& EntityManager,
 					UE::MassAvoidance::DebugDrawForce(AgentDebugContext,
 						OtherBasePosition + UE::MassAvoidance::DebugAgentHeightOffset,
 						OtherBasePosition + UE::MassAvoidance::DebugAgentHeightOffset + SeparationForce,
-						UE::MassAvoidance::AgentAvoidForceColor);
+						UE::MassAvoidance::AgentAvoidForceColor, UE::MassAvoidance::AvoidThickness);
 #endif // WITH_MASSGAMEPLAY_DEBUG
 				}
 			}
