[Mass Navmesh Navigation] Made the end of path detection more robust by using the last segment direction instead of the tangent.
Clamp EntityDistanceToGoal to 0 (avoiding negative values).
Added more logs.

diff --git a/Engine/Plugins/AI/MassAI/Source/MassNavMeshNavigation/Private/MassNavMeshNavigationProcessors.cpp b/Engine/Plugins/AI/MassAI/Source/MassNavMeshNavigation/Private/MassNavMeshNavigationProcessors.cpp
--- a/Engine/Plugins/AI/MassAI/Source/MassNavMeshNavigation/Private/MassNavMeshNavigationProcessors.cpp
+++ b/Engine/Plugins/AI/MassAI/Source/MassNavMeshNavigation/Private/MassNavMeshNavigationProcessors.cpp
@@ -120,11 +120,11 @@ void UMassNavMeshPathFollowProcessor::Execute(FMassEntityManager& EntityManager,
 					const FMassNavMeshPathPoint PointB = ShortPath.Points[ClosestSegmentIndex+1];
 					const float DistanceOnClosestSegment =  (EntityLocation - PointA.Position).Dot((PointB.Position - PointA.Position).GetSafeNormal());
 					const float EntityProgress = ShortPath.Points[ClosestSegmentIndex].Distance.Get() + DistanceOnClosestSegment;
-					MoveTarget.EntityDistanceToGoal = ShortPath.Points[UpdatePointIndex].Distance.Get() - EntityProgress;
+					MoveTarget.EntityDistanceToGoal = FMath::Max(0.f, ShortPath.Points[UpdatePointIndex].Distance.Get() - EntityProgress);
 					
 #if WITH_MASSGAMEPLAY_DEBUG && UNSAFE_FOR_MT
-					UE_CVLOG_UELOG(bDisplayDebug, this, LogMassNavMeshNavigation, Verbose, TEXT("Entity [%s]   ProgressDistance: %f, UpdatePointIndex: %i"),
-						*Entity.DebugGetDescription(), ShortPath.MoveTargetProgressDistance, UpdatePointIndex);
+					UE_CVLOG_UELOG(bDisplayDebug, this, LogMassNavMeshNavigation, Verbose, TEXT("Entity [%s]   ProgressDistance: %.2f, UpdatePointIndex: %i, EntityDistanceToGoal: %.2f"),
+						*Entity.DebugGetDescription(), ShortPath.MoveTargetProgressDistance, UpdatePointIndex, MoveTarget.EntityDistanceToGoal);
 
 					if (bDisplayDebug)
 					{
@@ -196,9 +196,11 @@ void UMassNavMeshPathFollowProcessor::Execute(FMassEntityManager& EntityManager,
 						else
 						{
 							// Last section of the path, wait for the steering to complete the movement.
+							UE_CVLOG_UELOG(bDisplayDebug, this, LogMassNavMeshNavigation, Log, TEXT("Entity [%s]   Waiting to reach the end of path."), *Entity.DebugGetDescription());
+							
 							const float SegmentSize = ShortPath.Points[UpdatePointIndex].Distance.Get() - ShortPath.Points[UpdatePointIndex-1].Distance.Get();
 							const FVector PreviousPointLocation = ShortPath.Points[UpdatePointIndex-1].Position;
-							FVector PathDir = ShortPath.Points[UpdatePointIndex-1].Tangent.GetVector();
+							const FVector PathDir = (ShortPath.Points[UpdatePointIndex].Position - PreviousPointLocation).GetSafeNormal();
 
 							const FVector RelativeLocation = EntityLocation - PreviousPointLocation;
 							const float ProjectionOnSegment = RelativeLocation.Dot(PathDir);
