[MassGameplay] Fixed the way UMassStateTreeSubsystem::DynamicProcessorClass making it possible to use project-specific classes (set via UMassBehaviorSettings or DefaultMass.ini)

diff --git a/Engine/Plugins/AI/MassAI/Source/MassAIBehavior/MassBehaviorSettings.cpp b/Engine/Plugins/AI/MassAI/Source/MassAIBehavior/MassBehaviorSettings.cpp
--- a/Engine/Plugins/AI/MassAI/Source/MassAIBehavior/MassBehaviorSettings.cpp
+++ b/Engine/Plugins/AI/MassAI/Source/MassAIBehavior/MassBehaviorSettings.cpp
@@ -1,6 +1,7 @@
-﻿// Copyright Epic Games, Inc. All Rights Reserved.
+// Copyright Epic Games, Inc. All Rights Reserved.
 
 #include "MassBehaviorSettings.h"
+#include "MassStateTreeProcessors.h"
 
 //----------------------------------------------------------------------//
 // UMassBehaviorSettings
@@ -14,4 +15,6 @@ UMassBehaviorSettings::UMassBehaviorSettings(const FObjectInitializer& ObjectIni
 	MaxActivationsPerLOD[EMassLOD::Medium] = 100;
 	MaxActivationsPerLOD[EMassLOD::Low] = 100;
 	MaxActivationsPerLOD[EMassLOD::Off] = 100;
+
+	DynamicStateTreeProcessorClass = UMassStateTreeProcessor::StaticClass();
 }
diff --git a/Engine/Plugins/AI/MassAI/Source/MassAIBehavior/Private/MassStateTreeSubsystem.cpp b/Engine/Plugins/AI/MassAI/Source/MassAIBehavior/Private/MassStateTreeSubsystem.cpp
--- a/Engine/Plugins/AI/MassAI/Source/MassAIBehavior/Private/MassStateTreeSubsystem.cpp
+++ b/Engine/Plugins/AI/MassAI/Source/MassAIBehavior/Private/MassStateTreeSubsystem.cpp
@@ -10,6 +10,7 @@
 #include "StateTree.h"
 #include "Engine/Engine.h"
 #include "MassStateTreeProcessors.h"
+#include "MassBehaviorSettings.h"
 
 
 namespace UE::Mass::StateTree
@@ -28,10 +29,13 @@ void UMassStateTreeSubsystem::Initialize(FSubsystemCollectionBase& Collection)
 	check(EntitySubsystem);
 	ĘntityManager = EntitySubsystem->GetMutableEntityManager().AsShared();
 
-	if (!DynamicProcessorClass)
-	{
-		DynamicProcessorClass = UMassStateTreeProcessor::StaticClass();
-	}
+	const UMassBehaviorSettings* BehaviorSettings = GetDefault<UMassBehaviorSettings>();
+	check(BehaviorSettings);
+
+	UClass* LoadedClass = BehaviorSettings->DynamicStateTreeProcessorClass.LoadSynchronous();
+	DynamicProcessorClass = LoadedClass
+		? LoadedClass
+		: UMassStateTreeProcessor::StaticClass();
 
 	SimulationSubsystem = Collection.InitializeDependency<UMassSimulationSubsystem>();
 }
diff --git a/Engine/Plugins/AI/MassAI/Source/MassAIBehavior/Public/MassBehaviorSettings.h b/Engine/Plugins/AI/MassAI/Source/MassAIBehavior/Public/MassBehaviorSettings.h
--- a/Engine/Plugins/AI/MassAI/Source/MassAIBehavior/Public/MassBehaviorSettings.h
+++ b/Engine/Plugins/AI/MassAI/Source/MassAIBehavior/Public/MassBehaviorSettings.h
@@ -1,4 +1,4 @@
-﻿// Copyright Epic Games, Inc. All Rights Reserved.
+// Copyright Epic Games, Inc. All Rights Reserved.
 
 #pragma once
 
@@ -6,6 +6,8 @@
 #include "MassLODTypes.h"
 #include "MassBehaviorSettings.generated.h"
 
+class UMassStateTreeProcessor;
+
 UCLASS(config = Mass, defaultconfig, meta = (DisplayName = "Mass Behavior"))
 class MASSAIBEHAVIOR_API UMassBehaviorSettings : public UMassModuleSettings
 {
@@ -16,4 +18,11 @@ public:
 	
 	UPROPERTY(EditAnywhere, Category = "Mass|LOD", config)
 	int32 MaxActivationsPerLOD[EMassLOD::Max];
+
+	/**
+	 * Class to use when creating dynamic processors to handle given StateTree assets.
+	 * Can be also set via DefaultMass.ini file. 
+	 */
+	UPROPERTY(Config, EditAnywhere, Category="Mass|StateTree", NoClear)
+	TSoftClassPtr<UMassStateTreeProcessor> DynamicStateTreeProcessorClass;
 };
diff --git a/Engine/Plugins/AI/MassAI/Source/MassAIBehavior/Public/MassStateTreeSubsystem.h b/Engine/Plugins/AI/MassAI/Source/MassAIBehavior/Public/MassStateTreeSubsystem.h
--- a/Engine/Plugins/AI/MassAI/Source/MassAIBehavior/Public/MassStateTreeSubsystem.h
+++ b/Engine/Plugins/AI/MassAI/Source/MassAIBehavior/Public/MassStateTreeSubsystem.h
@@ -110,8 +110,8 @@ protected:
 
 	/**
 	 * Class to use when creating dynamic processors to handle given StateTree assets.
-	 * Can be overridden via DefaultMass.ini file. 
+	 * Set based on UMassBehaviorSettings.DynamicStateTreeProcessorClass
 	 */
-	UPROPERTY(Config, EditDefaultsOnly, Category=Mass)
+	UPROPERTY(Transient)
 	TSubclassOf<UMassStateTreeProcessor> DynamicProcessorClass;
 };
