Addition of FTransformFragment as a requirement for UMassNavigationObstacleTrait

diff --git a/Engine/Plugins/AI/MassAI/Source/MassNavigation/Private/MassNavigationObstacleTrait.cpp b/Engine/Plugins/AI/MassAI/Source/MassNavigation/Private/MassNavigationObstacleTrait.cpp
--- a/Engine/Plugins/AI/MassAI/Source/MassNavigation/Private/MassNavigationObstacleTrait.cpp
+++ b/Engine/Plugins/AI/MassAI/Source/MassNavigation/Private/MassNavigationObstacleTrait.cpp
@@ -6,6 +6,7 @@
 
 void UMassNavigationObstacleTrait::BuildTemplate(FMassEntityTemplateBuildContext& BuildContext, const UWorld& World) const
 {
+	BuildContext.RequireFragment<FTransformFragment>();
 	BuildContext.RequireFragment<FAgentRadiusFragment>();
 
 	BuildContext.AddFragment<FMassNavigationObstacleGridCellLocationFragment>();
