Seperate concept of 'desired velocity' from 'velocity' in mass agents so that animation can be driven by steering intent without being exposed to noise from current velocity in the case of animation driven root motion

#tests aiproject, animsandbox

diff --git a/Engine/Plugins/AI/MassAI/Source/MassNavigation/Private/MassSteeringProcessors.cpp b/Engine/Plugins/AI/MassAI/Source/MassNavigation/Private/MassSteeringProcessors.cpp
--- a/Engine/Plugins/AI/MassAI/Source/MassNavigation/Private/MassSteeringProcessors.cpp
+++ b/Engine/Plugins/AI/MassAI/Source/MassNavigation/Private/MassSteeringProcessors.cpp
@@ -69,7 +69,7 @@ void UMassSteerToMoveTargetProcessor::ConfigureQueries(const TSharedRef<FMassEnt
 	EntityQuery.AddRequirement<FMassStandingSteeringFragment>(EMassFragmentAccess::ReadWrite);
 	EntityQuery.AddRequirement<FMassGhostLocationFragment>(EMassFragmentAccess::ReadWrite);
 	EntityQuery.AddRequirement<FMassForceFragment>(EMassFragmentAccess::ReadWrite);
-	EntityQuery.AddRequirement<FMassVelocityFragment>(EMassFragmentAccess::ReadWrite);
+	EntityQuery.AddRequirement<FMassDesiredMovementFragment>(EMassFragmentAccess::ReadWrite);	// Note this is read-write because this processor will sometimes zero the desired velocity, but normally it affects it via the FForceFragment
 	EntityQuery.AddConstSharedRequirement<FMassMovementParameters>(EMassFragmentPresence::All);
 	EntityQuery.AddConstSharedRequirement<FMassMovingSteeringParameters>(EMassFragmentPresence::All);
 	EntityQuery.AddConstSharedRequirement<FMassStandingSteeringParameters>(EMassFragmentPresence::All);
@@ -84,7 +84,7 @@ void UMassSteerToMoveTargetProcessor::Execute(FMassEntityManager& EntityManager,
 	{
 		const TConstArrayView<FTransformFragment> TransformList = Context.GetFragmentView<FTransformFragment>();
 		const TArrayView<FMassMoveTargetFragment> MoveTargetList = Context.GetMutableFragmentView<FMassMoveTargetFragment>();
-		const TArrayView<FMassVelocityFragment> VelocityList = Context.GetMutableFragmentView<FMassVelocityFragment>();
+		const TArrayView<FMassDesiredMovementFragment> MovementList = Context.GetMutableFragmentView<FMassDesiredMovementFragment>();
 		const TArrayView<FMassForceFragment> ForceList = Context.GetMutableFragmentView<FMassForceFragment>();
 		const TArrayView<FMassSteeringFragment> SteeringList = Context.GetMutableFragmentView<FMassSteeringFragment>();
 		const TArrayView<FMassStandingSteeringFragment> StandingSteeringList = Context.GetMutableFragmentView<FMassStandingSteeringFragment>();
@@ -104,7 +104,7 @@ void UMassSteerToMoveTargetProcessor::Execute(FMassEntityManager& EntityManager,
 			FMassGhostLocationFragment& Ghost = GhostList[EntityIt];
 			FMassMoveTargetFragment& MoveTarget = MoveTargetList[EntityIt];
 			FMassForceFragment& Force = ForceList[EntityIt];
-			FMassVelocityFragment& Velocity = VelocityList[EntityIt];
+			FMassDesiredMovementFragment& DesiredMovement = MovementList[EntityIt];
 			const FMassEntityHandle Entity = Context.GetEntity(EntityIt);
 
 			const FTransform& Transform = TransformFragment.GetTransform();;
@@ -202,7 +202,11 @@ void UMassSteerToMoveTargetProcessor::Execute(FMassEntityManager& EntityManager,
 
 				// @todo: This current completely overrides steering, we probably should have one processor that resets the steering at the beginning of the frame.
 				Steering.DesiredVelocity = SteerDirection * DesiredSpeed;
-				Force.Value = SteerK * (Steering.DesiredVelocity - Velocity.Value); // Goal force
+
+				// Important: we want steering force to be stable against noisy velocity, so we use
+				// the difference between current desired velocity and target desired velocity.
+				// We dont want to read the actual agent velocity directly since this can create a feedback loop for animated characters
+				Force.Value = SteerK * (Steering.DesiredVelocity - DesiredMovement.DesiredVelocity); // Goal force
 			}
 			else if (MoveTarget.GetCurrentAction() == EMassMovementAction::Stand)
 			{
@@ -293,15 +297,15 @@ void UMassSteerToMoveTargetProcessor::Execute(FMassEntityManager& EntityManager,
 					
 					// @todo: This current completely overrides steering, we probably should have one processor that resets the steering at the beginning of the frame.
 					Steering.DesiredVelocity = SteerDirection * DesiredSpeed;
-					Force.Value = SteerK * (Steering.DesiredVelocity - Velocity.Value); // Goal force
+					Force.Value = SteerK * (Steering.DesiredVelocity - DesiredMovement.DesiredVelocity); // Goal force
 					Force.Value = Force.Value.GetClampedToMaxSize(MovementParams.MaxAcceleration); 
 				}
 				else
 				{
-					// When reached destination, clamp small velocities to zero to avoid tiny drifting.
-					if (Velocity.Value.SquaredLength() < FMath::Square(StandingSteeringParams.LowSpeedThreshold))
+					// When reached destination, clamp small desired velocities to zero to avoid tiny drifting.
+					if (DesiredMovement.DesiredVelocity.SquaredLength() < FMath::Square(StandingSteeringParams.LowSpeedThreshold))
 					{
-						Velocity.Value = FVector::ZeroVector;
+						DesiredMovement.DesiredVelocity = FVector::ZeroVector;
 						Force.Value = FVector::ZeroVector;
 					}
 				}
@@ -314,7 +318,7 @@ void UMassSteerToMoveTargetProcessor::Execute(FMassEntityManager& EntityManager,
 				Steering.Reset();
 				MoveTarget.bSteeringFallingBehind = false;
 				Force.Value = FVector::ZeroVector;
-				Velocity.Value = FVector::ZeroVector;
+				DesiredMovement.DesiredVelocity = FVector::ZeroVector;
 			}
 
 #if WITH_MASSGAMEPLAY_DEBUG && UNSAFE_FOR_MT
@@ -350,9 +354,9 @@ void UMassSteerToMoveTargetProcessor::Execute(FMassEntityManager& EntityManager,
 					UE_VLOG_SEGMENT_THICK(this, LogMassNavigation, Log, MoveTargetCenter, CurrentLocation + ZOffset, LightEntityColor, /*Thickness*/1, TEXT(""));
 				}
 
-				// Display Velocity
-				UE_VLOG_SEGMENT_THICK(this, LogMassNavigation, Log, CurrentLocation + ZOffset, CurrentLocation + ZOffset + Velocity.Value, FColor::Yellow, /*Thickness*/4,
-					TEXT("Velocity %.1f"), Velocity.Value.Length());
+				// Display Desired Velocity
+				UE_VLOG_SEGMENT_THICK(this, LogMassNavigation, Log, CurrentLocation + ZOffset, CurrentLocation + ZOffset + Movement.DesiredVelocity, FColor::Yellow, /*Thickness*/4,
+					TEXT("DesiredVelocity %.1f"), Movement.DesiredVelocity.Length());
 
 				// Display DesiredVelocity
 				UE_VLOG_SEGMENT_THICK(this, LogMassNavigation, Log, CurrentLocation + ZOffset + FVector(0,0,2),
diff --git a/Engine/Plugins/Runtime/MassGameplay/Source/MassCommon/Public/MassCommonTypes.h b/Engine/Plugins/Runtime/MassGameplay/Source/MassCommon/Public/MassCommonTypes.h
--- a/Engine/Plugins/Runtime/MassGameplay/Source/MassCommon/Public/MassCommonTypes.h
+++ b/Engine/Plugins/Runtime/MassGameplay/Source/MassCommon/Public/MassCommonTypes.h
@@ -22,6 +22,7 @@ namespace UE::Mass::ProcessorGroupNames
 	const FName Behavior = FName(TEXT("Behavior"));
 	const FName Tasks = FName(TEXT("Tasks"));
 	const FName Avoidance = FName(TEXT("Avoidance"));
+	const FName ApplyForces = FName(TEXT("ApplyForces"));
 	const FName Movement = FName(TEXT("Movement"));
 }
 
diff --git a/Engine/Plugins/Runtime/MassGameplay/Source/MassMovement/Private/MassMovementProcessors.cpp b/Engine/Plugins/Runtime/MassGameplay/Source/MassMovement/Private/MassMovementProcessors.cpp
--- a/Engine/Plugins/Runtime/MassGameplay/Source/MassMovement/Private/MassMovementProcessors.cpp
+++ b/Engine/Plugins/Runtime/MassGameplay/Source/MassMovement/Private/MassMovementProcessors.cpp
@@ -12,6 +12,51 @@
 #define UNSAFE_FOR_MT 0
 #define MOVEMENT_DEBUGDRAW 0	// Set to 1 to see heading debugdraw
 
+//----------------------------------------------------------------------//
+//  UMassApplyForceProcessor
+//----------------------------------------------------------------------//
+
+UMassApplyForceProcessor::UMassApplyForceProcessor()
+	: EntityQuery(*this)
+{
+	ExecutionFlags = (int32)EProcessorExecutionFlags::AllNetModes;
+	ExecutionOrder.ExecuteInGroup = UE::Mass::ProcessorGroupNames::ApplyForces;
+	ExecutionOrder.ExecuteAfter.Add(UE::Mass::ProcessorGroupNames::Avoidance);
+}
+
+void UMassApplyForceProcessor::ConfigureQueries(const TSharedRef<FMassEntityManager>& EntityManager)
+{
+	EntityQuery.AddRequirement<FMassDesiredMovementFragment>(EMassFragmentAccess::ReadWrite);
+	EntityQuery.AddRequirement<FMassForceFragment>(EMassFragmentAccess::ReadWrite);
+	EntityQuery.AddTagRequirement<FMassOffLODTag>(EMassFragmentPresence::None);
+}
+
+void UMassApplyForceProcessor::Execute(FMassEntityManager& EntityManager, FMassExecutionContext& Context)
+{
+	// Clamp max delta time to avoid force explosion on large time steps (i.e. during initialization).
+	const float DeltaTime = FMath::Min(0.1f, Context.GetDeltaTimeSeconds());
+
+	QUICK_SCOPE_CYCLE_COUNTER(UMassApplyForceProcessor);
+
+	EntityQuery.ForEachEntityChunk(Context, [this, DeltaTime](FMassExecutionContext& Context)
+	{
+		const TArrayView<FMassForceFragment> ForceList = Context.GetMutableFragmentView<FMassForceFragment>();
+		const TArrayView<FMassDesiredMovementFragment> MovementList = Context.GetMutableFragmentView<FMassDesiredMovementFragment>();
+
+		for (FMassExecutionContext::FEntityIterator EntityIt = Context.CreateEntityIterator(); EntityIt; ++EntityIt)
+		{
+			FMassForceFragment& Force = ForceList[EntityIt];
+			FMassDesiredMovementFragment& DesiredMovement = MovementList[EntityIt];
+
+			// Update desired velocity from steering forces.
+			DesiredMovement.DesiredVelocity += Force.Value * DeltaTime;
+
+			// Reset to zero after force is applied. Processors accumulate forces into the force fragment
+			Force.Value = FVector::ZeroVector;
+		}
+	});
+}
+
 //----------------------------------------------------------------------//
 //  UMassApplyMovementProcessor
 //----------------------------------------------------------------------//
@@ -21,16 +66,16 @@ UMassApplyMovementProcessor::UMassApplyMovementProcessor()
 {
 	ExecutionFlags = (int32)EProcessorExecutionFlags::AllNetModes;
 	ExecutionOrder.ExecuteInGroup = UE::Mass::ProcessorGroupNames::Movement;
-	ExecutionOrder.ExecuteAfter.Add(UE::Mass::ProcessorGroupNames::Avoidance);
+	ExecutionOrder.ExecuteAfter.Add(UE::Mass::ProcessorGroupNames::ApplyForces);
 }
 
 void UMassApplyMovementProcessor::ConfigureQueries(const TSharedRef<FMassEntityManager>& EntityManager)
 {
 	EntityQuery.AddRequirement<FMassVelocityFragment>(EMassFragmentAccess::ReadWrite);
 	EntityQuery.AddRequirement<FTransformFragment>(EMassFragmentAccess::ReadWrite);
-	EntityQuery.AddRequirement<FMassForceFragment>(EMassFragmentAccess::ReadWrite);
+	EntityQuery.AddRequirement<FMassDesiredMovementFragment>(EMassFragmentAccess::ReadOnly);
 	EntityQuery.AddTagRequirement<FMassOffLODTag>(EMassFragmentPresence::None);
-	EntityQuery.AddConstSharedRequirement<FMassMovementParameters>(EMassFragmentPresence::All);
+	EntityQuery.AddTagRequirement<FMassCodeDrivenMovementTag>(EMassFragmentPresence::All);
 }
 
 void UMassApplyMovementProcessor::Execute(FMassEntityManager& EntityManager,
@@ -39,24 +84,24 @@ void UMassApplyMovementProcessor::Execute(FMassEntityManager& EntityManager,
 	// Clamp max delta time to avoid force explosion on large time steps (i.e. during initialization).
 	const float DeltaTime = FMath::Min(0.1f, Context.GetDeltaTimeSeconds());
 
-	QUICK_SCOPE_CYCLE_COUNTER(HighRes);
+	QUICK_SCOPE_CYCLE_COUNTER(UMassApplyMovementProcessor);
 
 	EntityQuery.ForEachEntityChunk(Context, [this, DeltaTime](FMassExecutionContext& Context)
 	{
-		const FMassMovementParameters& MovementParams = Context.GetConstSharedFragment<FMassMovementParameters>();
-
 		const TArrayView<FTransformFragment> LocationList = Context.GetMutableFragmentView<FTransformFragment>();
-		const TArrayView<FMassForceFragment> ForceList = Context.GetMutableFragmentView<FMassForceFragment>();
 		const TArrayView<FMassVelocityFragment> VelocityList = Context.GetMutableFragmentView<FMassVelocityFragment>();
+		const TConstArrayView<FMassDesiredMovementFragment> MovementList = Context.GetFragmentView<FMassDesiredMovementFragment>();
 
 		for (FMassExecutionContext::FEntityIterator EntityIt = Context.CreateEntityIterator(); EntityIt; ++EntityIt)
 		{
-			FMassForceFragment& Force = ForceList[EntityIt];
 			FMassVelocityFragment& Velocity = VelocityList[EntityIt];
-			FTransform& CurrentTransform = LocationList[EntityIt].GetMutableTransform();
 
-			// Update velocity from steering forces.
-			Velocity.Value += Force.Value * DeltaTime;
+			// For code driven, we just apply desired velocity as velocity
+			// This is basically the equivalent of the character movement component
+			// Smoothing / acceleration could be done here if we wanted
+			Velocity.Value = MovementList[EntityIt].DesiredVelocity;
+			
+			FTransform& CurrentTransform = LocationList[EntityIt].GetMutableTransform();
 
 #if WITH_MASSGAMEPLAY_DEBUG
 			if (UE::MassMovement::bFreezeMovement)
diff --git a/Engine/Plugins/Runtime/MassGameplay/Source/MassMovement/Private/MassMovementTrait.cpp b/Engine/Plugins/Runtime/MassGameplay/Source/MassMovement/Private/MassMovementTrait.cpp
--- a/Engine/Plugins/Runtime/MassGameplay/Source/MassMovement/Private/MassMovementTrait.cpp
+++ b/Engine/Plugins/Runtime/MassGameplay/Source/MassMovement/Private/MassMovementTrait.cpp
@@ -17,8 +17,14 @@ void UMassMovementTrait::BuildTemplate(FMassEntityTemplateBuildContext& BuildCon
 
 	BuildContext.AddFragment<FMassVelocityFragment>();
 	BuildContext.AddFragment<FMassForceFragment>();
+	BuildContext.AddFragment<FMassDesiredMovementFragment>();
 
 	const FMassMovementParameters MovementValidated = Movement.GetValidated();
 	const FConstSharedStruct MovementFragment = EntityManager.GetOrCreateConstSharedFragment(MovementValidated);
 	BuildContext.AddConstSharedFragment(MovementFragment);
+
+	if (Movement.bIsCodeDrivenMovement)
+	{
+		BuildContext.AddTag<FMassCodeDrivenMovementTag>();
+	}
 }
diff --git a/Engine/Plugins/Runtime/MassGameplay/Source/MassMovement/Public/MassMovementFragments.h b/Engine/Plugins/Runtime/MassGameplay/Source/MassMovement/Public/MassMovementFragments.h
--- a/Engine/Plugins/Runtime/MassGameplay/Source/MassMovement/Public/MassMovementFragments.h
+++ b/Engine/Plugins/Runtime/MassGameplay/Source/MassMovement/Public/MassMovementFragments.h
@@ -11,6 +11,8 @@
 #include "RandomSequence.h"
 #include "MassMovementFragments.generated.h"
 
+// This represents the actual physical velocity of the mass entity in the world
+// For agents with an actor representation, this is the velocity of the movement component
 USTRUCT()
 struct MASSMOVEMENT_API FMassVelocityFragment : public FMassFragment
 {
@@ -23,6 +25,19 @@ struct MASSMOVEMENT_API FMassVelocityFragment : public FMassFragment
 #endif	
 };
 
+// This is the output of all processors that intend to affect movement
+// It is the input to the movement system (e.g. mover, animation etc.)
+USTRUCT()
+struct MASSMOVEMENT_API FMassDesiredMovementFragment : public FMassFragment
+{
+	GENERATED_BODY()
+	
+	FVector DesiredVelocity = FVector::ZeroVector;
+
+	// Todo: Could add a DesiredFacing here
+};
+
+// Accumulator for steering / avoidance forces to apply to the desired velocity 
 USTRUCT()
 struct MASSMOVEMENT_API FMassForceFragment : public FMassFragment
 {
@@ -31,6 +46,20 @@ struct MASSMOVEMENT_API FMassForceFragment : public FMassFragment
 	FVector Value = FVector::ZeroVector;
 };
 
+// The presence of this tag indicates that this mass agent's velocity should be
+// controlled by the FMassMovementCommand.
+//
+// For code driven displacement, we want the desired velocity to affect the velocity
+// directly, which is then applied to the character mover.
+// For e.g. Root Motion driven displacement, we just need to pipe the DesiredVelocity to the
+// animation system and let it do the test.
+USTRUCT()
+struct MASSMOVEMENT_API FMassCodeDrivenMovementTag : public FMassTag
+{
+	GENERATED_BODY()
+};
+
+// Parameters describing how this mass agent should move
 USTRUCT()
 struct MASSMOVEMENT_API FMassMovementParameters : public FMassConstSharedFragment
 {
@@ -97,4 +126,8 @@ struct MASSMOVEMENT_API FMassMovementParameters : public FMassConstSharedFragmen
 	/** List of supported movement styles for this configuration. */
 	UPROPERTY(EditAnywhere, Category = "Movement")
 	TArray<FMassMovementStyleParameters> MovementStyles;
+
+	/** Indicate whether mass AI has direct control over the mass agent's velocity. If true, desired velocity will be written directly to velocity every frame */
+	UPROPERTY(EditAnywhere, Category = "Movement")
+	bool bIsCodeDrivenMovement = true;
 };
diff --git a/Engine/Plugins/Runtime/MassGameplay/Source/MassMovement/Public/Movement/MassMovementProcessors.h b/Engine/Plugins/Runtime/MassGameplay/Source/MassMovement/Public/Movement/MassMovementProcessors.h
--- a/Engine/Plugins/Runtime/MassGameplay/Source/MassMovement/Public/Movement/MassMovementProcessors.h
+++ b/Engine/Plugins/Runtime/MassGameplay/Source/MassMovement/Public/Movement/MassMovementProcessors.h
@@ -6,7 +6,28 @@
 #include "MassMovementProcessors.generated.h"
 
 /**
- * Updates entities position based on force and velocity.
+ * Calculate desired movement based on input forces
+ */
+UCLASS()
+class MASSMOVEMENT_API UMassApplyForceProcessor : public UMassProcessor
+{
+	GENERATED_BODY()
+
+public:
+	UMassApplyForceProcessor();
+
+protected:
+	virtual void ConfigureQueries(const TSharedRef<FMassEntityManager>& EntityManager) override;
+	virtual void Execute(FMassEntityManager& EntityManager, FMassExecutionContext& Context) override;
+
+private:
+	FMassEntityQuery EntityQuery;
+};
+
+
+/**
+ * Updates entities position based on desired velocity.
+ * Only required for agents that have code driven displacement
  * Not applied on Off-LOD entities.
  */
 UCLASS()
