[MassGameplay] Removed no longer necessary uses of UNSAFE_FOR_MT pattern.
The change was made possible once VisualLogger was made thread-safe

diff --git a/Engine/Plugins/AI/MassAI/Source/MassAIBehavior/Private/MassLookAtProcessors.cpp b/Engine/Plugins/AI/MassAI/Source/MassAIBehavior/Private/MassLookAtProcessors.cpp
--- a/Engine/Plugins/AI/MassAI/Source/MassAIBehavior/Private/MassLookAtProcessors.cpp
+++ b/Engine/Plugins/AI/MassAI/Source/MassAIBehavior/Private/MassLookAtProcessors.cpp
@@ -22,8 +22,6 @@
 #include "MassDebugger.h"
 
 
-#define UNSAFE_FOR_MT 1
-
 namespace UE::MassBehavior
 {
 	namespace Tweakables
@@ -227,7 +225,7 @@ void UMassLookAtProcessor::Execute(FMassEntityManager& EntityManager, FMassExecu
 				// Clamp
 				LookAt.Direction = UE::MassBehavior::ClampDirectionToXAxisCone(LookAt.Direction, FMath::DegreesToRadians(AngleThresholdInDegrees));
 
-	#if WITH_MASSGAMEPLAY_DEBUG && UNSAFE_FOR_MT
+	#if WITH_MASSGAMEPLAY_DEBUG
 				if (bDisplayDebug)
 				{
 					const FVector Origin = TransformFragment.GetTransform().GetLocation() + FVector(0.f,0.f,DebugZOffset);
@@ -351,7 +349,7 @@ void UMassLookAtProcessor::UpdateLookAtTrajectory(const FTransform& Transform, c
 		LookAt.Direction.Z = 0.0f;
 		LookAt.MainTargetLocation = AnticipatedPosition;
 
-#if WITH_MASSGAMEPLAY_DEBUG && UNSAFE_FOR_MT
+#if WITH_MASSGAMEPLAY_DEBUG
 		if (bDisplayDebug)
 		{
 			const FVector ZOffset(0.f,0.f,DebugZOffset);
@@ -374,7 +372,7 @@ void UMassLookAtProcessor::UpdateLookAtTrackedEntity(const FMassEntityManager& E
 			LookAt.Direction = Transform.InverseTransformVector(NewGlobalDirection);
 			LookAt.MainTargetLocation = TargetLocation;
 
-#if WITH_MASSGAMEPLAY_DEBUG && UNSAFE_FOR_MT
+#if WITH_MASSGAMEPLAY_DEBUG
 			if (bDisplayDebug)
 			{
 				const FVector ZOffset(0.f,0.f,DebugZOffset);
@@ -401,7 +399,7 @@ bool UMassLookAtProcessor::UpdateGazeTrackedEntity(const FMassEntityManager& Ent
 			LookAt.GazeTargetLocation = TargetLocation;
 			bHasTarget = true;
 			
-#if WITH_MASSGAMEPLAY_DEBUG && UNSAFE_FOR_MT
+#if WITH_MASSGAMEPLAY_DEBUG
 			if (bDisplayDebug)
 			{
 				const FVector ZOffset(0.f,0.f,DebugZOffset);
@@ -491,7 +489,7 @@ void UMassLookAtProcessor::BuildTrajectory(const UZoneGraphSubsystem& ZoneGraphS
 		}
 	}
 
-#if WITH_MASSGAMEPLAY_DEBUG && UNSAFE_FOR_MT
+#if WITH_MASSGAMEPLAY_DEBUG
 	if (bDisplayDebug)
 	{
 		const FVector ZOffset(0,0,35);
@@ -534,5 +532,3 @@ void UMassLookAtProcessor::BuildTrajectory(const UZoneGraphSubsystem& ZoneGraphS
 	}
 #endif // WITH_MASSGAMEPLAY_DEBUG
 }
-
-#undef UNSAFE_FOR_MT
\ No newline at end of file
diff --git a/Engine/Plugins/AI/MassAI/Source/MassNavMeshNavigation/Private/MassNavMeshNavigationProcessors.cpp b/Engine/Plugins/AI/MassAI/Source/MassNavMeshNavigation/Private/MassNavMeshNavigationProcessors.cpp
--- a/Engine/Plugins/AI/MassAI/Source/MassNavMeshNavigation/Private/MassNavMeshNavigationProcessors.cpp
+++ b/Engine/Plugins/AI/MassAI/Source/MassNavMeshNavigation/Private/MassNavMeshNavigationProcessors.cpp
@@ -1,7 +1,6 @@
 // Copyright Epic Games, Inc. All Rights Reserved.
 
 #include "MassNavMeshNavigationProcessors.h"
-
 #include "MassCommonFragments.h"
 #include "Avoidance/MassAvoidanceFragments.h"
 #include "MassCommonTypes.h"
@@ -13,7 +12,6 @@
 #include "MassSignalSubsystem.h"
 #include "MassSimulationLOD.h"
 
-#define UNSAFE_FOR_MT 1
 
 UMassNavMeshPathFollowProcessor::UMassNavMeshPathFollowProcessor()
 	: EntityQuery_Conditional(*this)
@@ -69,7 +67,7 @@ void UMassNavMeshPathFollowProcessor::Execute(FMassEntityManager& EntityManager,
 			const FMassEntityHandle Entity = Context.GetEntity(EntityIt);
 
 			bool bDisplayDebug = false;
-#if WITH_MASSGAMEPLAY_DEBUG && UNSAFE_FOR_MT // this will result in bDisplayDebug == false and disabling of all the vlogs below
+#if WITH_MASSGAMEPLAY_DEBUG // this will result in bDisplayDebug == false and disabling of all the vlogs below
 			bDisplayDebug = UE::Mass::Debug::IsDebuggingEntity(Entity);
 #endif // WITH_MASSGAMEPLAY_DEBUG
 
@@ -122,7 +120,7 @@ void UMassNavMeshPathFollowProcessor::Execute(FMassEntityManager& EntityManager,
 					const float EntityProgress = ShortPath.Points[ClosestSegmentIndex].Distance.Get() + DistanceOnClosestSegment;
 					MoveTarget.EntityDistanceToGoal = FMath::Max(0.f, ShortPath.Points[UpdatePointIndex].Distance.Get() - EntityProgress);
 					
-#if WITH_MASSGAMEPLAY_DEBUG && UNSAFE_FOR_MT
+#if WITH_MASSGAMEPLAY_DEBUG
 					UE_CVLOG_UELOG(bDisplayDebug, this, LogMassNavMeshNavigation, Verbose, TEXT("Entity [%s]   ProgressDistance: %.2f, UpdatePointIndex: %i, EntityDistanceToGoal: %.2f"),
 						*Entity.DebugGetDescription(), ShortPath.MoveTargetProgressDistance, UpdatePointIndex, MoveTarget.EntityDistanceToGoal);
 
@@ -221,7 +219,7 @@ void UMassNavMeshPathFollowProcessor::Execute(FMassEntityManager& EntityManager,
 					EntitiesToSignalPathDone.Add(Entity);
 				}
 
-#if WITH_MASSGAMEPLAY_DEBUG && UNSAFE_FOR_MT
+#if WITH_MASSGAMEPLAY_DEBUG
 				if (bDisplayDebug)
 				{
 					const FColor EntityColor = UE::Mass::Debug::GetEntityDebugColor(Entity);
@@ -284,5 +282,3 @@ void UMassNavMeshPathFollowProcessor::Execute(FMassEntityManager& EntityManager,
 		SignalSubsystem->SignalEntities(UE::Mass::Signals::FollowPointPathDone, EntitiesToSignalPathDone);
 	}
 }
-
-#undef UNSAFE_FOR_MT
\ No newline at end of file
diff --git a/Engine/Plugins/AI/MassAI/Source/MassNavigation/Private/MassNavigationProcessors.cpp b/Engine/Plugins/AI/MassAI/Source/MassNavigation/Private/MassNavigationProcessors.cpp
--- a/Engine/Plugins/AI/MassAI/Source/MassNavigation/Private/MassNavigationProcessors.cpp
+++ b/Engine/Plugins/AI/MassAI/Source/MassNavigation/Private/MassNavigationProcessors.cpp
@@ -14,9 +14,6 @@
 #include "Engine/World.h"
 
 
-#define UNSAFE_FOR_MT 0
-#define MOVEMENT_DEBUGDRAW 0	// Set to 1 to see heading debugdraw
-
 //----------------------------------------------------------------------//
 //  UMassOffLODNavigationProcessor
 //----------------------------------------------------------------------//
@@ -302,5 +299,3 @@ void UMassNavigationObstacleRemoverProcessor::Execute(FMassEntityManager& Entity
 		}
 	});
 }
-
-#undef UNSAFE_FOR_MT
diff --git a/Engine/Plugins/AI/MassAI/Source/MassNavigation/Private/MassSmoothOrientationProcessors.cpp b/Engine/Plugins/AI/MassAI/Source/MassNavigation/Private/MassSmoothOrientationProcessors.cpp
--- a/Engine/Plugins/AI/MassAI/Source/MassNavigation/Private/MassSmoothOrientationProcessors.cpp
+++ b/Engine/Plugins/AI/MassAI/Source/MassNavigation/Private/MassSmoothOrientationProcessors.cpp
@@ -11,8 +11,6 @@
 #include "MassSimulationLOD.h"
 #include "MassNavigationUtils.h"
 
-#define UNSAFE_FOR_MT 0
-#define MOVEMENT_DEBUGDRAW 0	// Set to 1 to see heading debugdraw
 
 //----------------------------------------------------------------------//
 //  UMassSmoothOrientationProcessor
@@ -132,5 +130,3 @@ void UMassSmoothOrientationProcessor::Execute(FMassEntityManager& EntityManager,
 		});
 	}
 }
-
-#undef UNSAFE_FOR_MT
diff --git a/Engine/Plugins/AI/MassAI/Source/MassNavigation/Private/MassSteeringProcessors.cpp b/Engine/Plugins/AI/MassAI/Source/MassNavigation/Private/MassSteeringProcessors.cpp
--- a/Engine/Plugins/AI/MassAI/Source/MassNavigation/Private/MassSteeringProcessors.cpp
+++ b/Engine/Plugins/AI/MassAI/Source/MassNavigation/Private/MassSteeringProcessors.cpp
@@ -11,10 +11,9 @@
 #include "Steering/MassSteeringFragments.h"
 #include "Math/UnrealMathUtility.h"
 #include "MassSimulationLOD.h"
+#if WITH_MASSGAMEPLAY_DEBUG
 #include "MassNavigationDebug.h"
-
-#define UNSAFE_FOR_MT 0
-#define MOVEMENT_DEBUGDRAW 0	// Set to 1 to see heading debugdraw
+#endif // WITH_MASSGAMEPLAY_DEBUG
 
 namespace UE::MassNavigation
 {
@@ -147,7 +146,7 @@ void UMassSteerToMoveTargetProcessor::Execute(FMassEntityManager& EntityManager,
 					SteerDirection *= 1. / DistanceToSteerTarget;
 				}
 
-#if WITH_MASSGAMEPLAY_DEBUG && UNSAFE_FOR_MT
+#if WITH_MASSGAMEPLAY_DEBUG
 				FColor EntityColor = FColor::White;
 				const bool bDisplayDebug = UE::Mass::Debug::IsDebuggingEntity(Entity, &EntityColor);
 				if (bDisplayDebug)
@@ -270,7 +269,7 @@ void UMassSteerToMoveTargetProcessor::Execute(FMassEntityManager& EntityManager,
 				const FVector::FReal Distance = Delta.Size();
 				if (Distance > StandingSteeringParams.DeadZoneRadius)
 				{
-#if WITH_MASSGAMEPLAY_DEBUG && UNSAFE_FOR_MT
+#if WITH_MASSGAMEPLAY_DEBUG
 					FColor EntityColor = FColor::White;
 					const bool bDisplayDebug = UE::Mass::Debug::IsDebuggingEntity(Entity, &EntityColor);
 					if (bDisplayDebug)
@@ -321,7 +320,7 @@ void UMassSteerToMoveTargetProcessor::Execute(FMassEntityManager& EntityManager,
 				DesiredMovement.DesiredVelocity = FVector::ZeroVector;
 			}
 
-#if WITH_MASSGAMEPLAY_DEBUG && UNSAFE_FOR_MT
+#if WITH_MASSGAMEPLAY_DEBUG
 			FColor EntityColor = FColor::White;
 			const bool bDisplayDebug = UE::Mass::Debug::IsDebuggingEntity(Entity, &EntityColor);
 			if (bDisplayDebug)
@@ -355,8 +354,8 @@ void UMassSteerToMoveTargetProcessor::Execute(FMassEntityManager& EntityManager,
 				}
 
 				// Display Desired Velocity
-				UE_VLOG_SEGMENT_THICK(this, LogMassNavigation, Log, CurrentLocation + ZOffset, CurrentLocation + ZOffset + Movement.DesiredVelocity, FColor::Yellow, /*Thickness*/4,
-					TEXT("DesiredVelocity %.1f"), Movement.DesiredVelocity.Length());
+				UE_VLOG_SEGMENT_THICK(this, LogMassNavigation, Log, CurrentLocation + ZOffset, CurrentLocation + ZOffset + DesiredMovement.DesiredVelocity, FColor::Yellow, /*Thickness*/4,
+					TEXT("DesiredVelocity %.1f"), DesiredMovement.DesiredVelocity.Length());
 
 				// Display DesiredVelocity
 				UE_VLOG_SEGMENT_THICK(this, LogMassNavigation, Log, CurrentLocation + ZOffset + FVector(0,0,2),
@@ -372,5 +371,3 @@ void UMassSteerToMoveTargetProcessor::Execute(FMassEntityManager& EntityManager,
 		}
 	});
 }
-
-#undef UNSAFE_FOR_MT
diff --git a/Engine/Plugins/AI/MassAI/Source/MassZoneGraphNavigation/Private/MassZoneGraphNavigationProcessors.cpp b/Engine/Plugins/AI/MassAI/Source/MassZoneGraphNavigation/Private/MassZoneGraphNavigationProcessors.cpp
--- a/Engine/Plugins/AI/MassAI/Source/MassZoneGraphNavigation/Private/MassZoneGraphNavigationProcessors.cpp
+++ b/Engine/Plugins/AI/MassAI/Source/MassZoneGraphNavigation/Private/MassZoneGraphNavigationProcessors.cpp
@@ -17,7 +17,6 @@
 #include "MassDebugger.h"
 #include "MassNavigationDebug.h"
 
-#define UNSAFE_FOR_MT 1
 
 //----------------------------------------------------------------------//
 //  UMassZoneGraphLocationInitializer
@@ -156,7 +155,7 @@ void UMassZoneGraphPathFollowProcessor::Execute(FMassEntityManager& EntityManage
 			const float DeltaTime = bHasVariableTick ? SimVariableTickList[EntityIt].DeltaTime : WorldDeltaTime;
 
 			bool bDisplayDebug = false;
-#if WITH_MASSGAMEPLAY_DEBUG && UNSAFE_FOR_MT // this will result in bDisplayDebug == false and disabling of all the vlogs below
+#if WITH_MASSGAMEPLAY_DEBUG // this will result in bDisplayDebug == false and disabling of all the vlogs below
 			bDisplayDebug = UE::Mass::Debug::IsDebuggingEntity(Entity);
 			if (bDisplayDebug)
 			{
@@ -343,7 +342,7 @@ void UMassZoneGraphPathFollowProcessor::Execute(FMassEntityManager& EntityManage
 					EntitiesToSignalPathDone.Add(Entity);
 				}
 
-#if WITH_MASSGAMEPLAY_DEBUG && UNSAFE_FOR_MT
+#if WITH_MASSGAMEPLAY_DEBUG
 				if (bDisplayDebug)
 				{
 					const FColor EntityColor = UE::Mass::Debug::GetEntityDebugColor(Entity);
@@ -598,5 +597,3 @@ void UMassZoneGraphLaneCacheBoundaryProcessor::Execute(FMassEntityManager& Entit
 		}
 	});
 }
-
-#undef UNSAFE_FOR_MT
diff --git a/Engine/Plugins/Runtime/MassGameplay/Source/MassMovement/Private/MassMovementProcessors.cpp b/Engine/Plugins/Runtime/MassGameplay/Source/MassMovement/Private/MassMovementProcessors.cpp
--- a/Engine/Plugins/Runtime/MassGameplay/Source/MassMovement/Private/MassMovementProcessors.cpp
+++ b/Engine/Plugins/Runtime/MassGameplay/Source/MassMovement/Private/MassMovementProcessors.cpp
@@ -9,8 +9,6 @@
 #include "Math/UnrealMathUtility.h"
 #include "MassSimulationLOD.h"
 
-#define UNSAFE_FOR_MT 0
-#define MOVEMENT_DEBUGDRAW 0	// Set to 1 to see heading debugdraw
 
 //----------------------------------------------------------------------//
 //  UMassApplyForceProcessor
@@ -113,5 +111,3 @@ void UMassApplyMovementProcessor::Execute(FMassEntityManager& EntityManager,
 		}
 	});
 }
-
-#undef UNSAFE_FOR_MT
