[MassLookAt]
- adding subsystem and mass processors to track potential targets in a hashgrid and to store and update active requests
- modified MassActorSubsystem to allow request for the mass entity handle for a given actor for cases where FMassActorFragment/UMassAgentComponent are not used
- updated MassToAnimNextTranslator to consider main target first then fallback on gaze target
- added CancelCommands to MassCommandBuffer for buffers that are queued to flush their commands on the game thread but the EntityManager is no longer available.

diff --git a/Engine/Plugins/AI/MassAI/Source/MassAIBehavior/Private/MassLookAtProcessors.cpp b/Engine/Plugins/AI/MassAI/Source/MassAIBehavior/Private/MassLookAtProcessors.cpp
--- a/Engine/Plugins/AI/MassAI/Source/MassAIBehavior/Private/MassLookAtProcessors.cpp
+++ b/Engine/Plugins/AI/MassAI/Source/MassAIBehavior/Private/MassLookAtProcessors.cpp
@@ -20,6 +20,7 @@
 #include "Engine/World.h"
 #include "MassLODFragments.h"
 #include "MassDebugger.h"
+#include "MassLookAtSubsystem.h"
 
 
 namespace UE::MassBehavior
@@ -27,12 +28,16 @@ namespace UE::MassBehavior
 	namespace Tweakables
 	{
 		float TrajectoryLookAhead = 600.f;
+		bool bConsiderObstacleTargets = false;
 	}
 
 	FAutoConsoleVariableRef CVars[] =
 	{
 		FAutoConsoleVariableRef(TEXT("ai.mass.LookAt.TrajectoryLookAhead"), Tweakables::TrajectoryLookAhead,
 								TEXT("Distance (in cm) further along the look at trajectory (based on current path) to lookat while moving."), ECVF_Cheat),
+
+		FAutoConsoleVariableRef(TEXT("ai.mass.LookAt.ConsiderObstacleTargets"), Tweakables::bConsiderObstacleTargets,
+								TEXT("Indicates if gaze targets can be selected from the entities registered in the obstacle grid."), ECVF_Cheat),
 	};
 
 	
@@ -65,7 +70,7 @@ namespace UE::MassBehavior
 			return 1.0;
 		}
 
-		// @todo: mae configurable
+		// @todo LookAt: mae configurable
 		const float SustainTime = GazeDuration * 0.25f;
 		const float DecayTime = GazeDuration * 0.45f;
 		
@@ -92,7 +97,7 @@ namespace UE::MassBehavior
 UMassLookAtProcessor::UMassLookAtProcessor()
 	: EntityQuery_Conditional(*this)
 {
-	ExecutionFlags = (int32)(EProcessorExecutionFlags::Client | EProcessorExecutionFlags::Standalone);
+	ExecutionFlags = static_cast<int32>(EProcessorExecutionFlags::Client | EProcessorExecutionFlags::Standalone);
 	ExecutionOrder.ExecuteInGroup = UE::Mass::ProcessorGroupNames::Tasks;
 	ExecutionOrder.ExecuteAfter.Add(UE::Mass::ProcessorGroupNames::Representation);
 }
@@ -111,6 +116,7 @@ void UMassLookAtProcessor::ConfigureQueries(const TSharedRef<FMassEntityManager>
 	EntityQuery_Conditional.AddChunkRequirement<FMassVisualizationChunkFragment>(EMassFragmentAccess::ReadOnly, EMassFragmentPresence::Optional);
 	EntityQuery_Conditional.SetChunkFilter(&FMassVisualizationChunkFragment::AreAnyEntitiesVisibleInChunk);
 	EntityQuery_Conditional.AddSubsystemRequirement<UMassNavigationSubsystem>(EMassFragmentAccess::ReadOnly);
+	EntityQuery_Conditional.AddSubsystemRequirement<UMassLookAtSubsystem>(EMassFragmentAccess::ReadOnly);
 	EntityQuery_Conditional.AddSubsystemRequirement<UZoneGraphSubsystem>(EMassFragmentAccess::ReadOnly);
 }
 
@@ -123,6 +129,7 @@ void UMassLookAtProcessor::Execute(FMassEntityManager& EntityManager, FMassExecu
 	EntityQuery_Conditional.ForEachEntityChunk(Context, [this, &EntityManager, CurrentTime](FMassExecutionContext& Context)
 		{
 			const UMassNavigationSubsystem& MassNavSystem = Context.GetSubsystemChecked<UMassNavigationSubsystem>();
+			const UMassLookAtSubsystem& LookAtTargetSystem = Context.GetSubsystemChecked<UMassLookAtSubsystem>();
 			const UZoneGraphSubsystem& ZoneGraphSubsystem = Context.GetSubsystemChecked<UZoneGraphSubsystem>();
 
 			const TArrayView<FMassLookAtFragment> LookAtList = Context.GetMutableFragmentView<FMassLookAtFragment>();
@@ -156,7 +163,7 @@ void UMassLookAtProcessor::Execute(FMassEntityManager& EntityManager, FMassExecu
 					const double TimeSinceUpdate = CurrentTime - LookAt.GazeStartTime;
 					if (TimeSinceUpdate >= LookAt.GazeDuration)
 					{
-						FindNewGazeTarget(MassNavSystem, EntityManager, CurrentTime, TransformFragment.GetTransform(), LookAt);
+						FindNewGazeTarget(MassNavSystem, LookAtTargetSystem, EntityManager, CurrentTime, TransformFragment.GetTransform(), Entity, LookAt);
 					}
 				}
 
@@ -237,9 +244,14 @@ void UMassLookAtProcessor::Execute(FMassEntityManager& EntityManager, FMassExecu
 		});
 }
 
-void UMassLookAtProcessor::FindNewGazeTarget(const UMassNavigationSubsystem& MassNavSystem, const FMassEntityManager& EntityManager, const double CurrentTime, const FTransform& Transform, FMassLookAtFragment& LookAt) const
+void UMassLookAtProcessor::FindNewGazeTarget(const UMassNavigationSubsystem& MassNavSystem
+	, const UMassLookAtSubsystem& LookAtSystem
+	, const FMassEntityManager& EntityManager
+	, const double CurrentTime
+	, const FTransform& Transform
+	, const FMassEntityHandle Entity
+	, FMassLookAtFragment& LookAt) const
 {
-	const FNavigationObstacleHashGrid2D& ObstacleGrid = MassNavSystem.GetObstacleGrid();
 	const FMassEntityHandle LastTrackedEntity = LookAt.GazeTrackedEntity;
 	
 	LookAt.GazeTrackedEntity.Reset();
@@ -255,68 +267,102 @@ void UMassLookAtProcessor::FindNewGazeTarget(const UMassNavigationSubsystem& Mas
 		const FVector QueryOrigin = Transform.TransformPosition(FVector(0.5f*QueryExtent, 0.f, 0.f));
 		const FBox QueryBox = FBox(QueryOrigin - 0.5f*Extent, QueryOrigin + 0.5f*Extent);
 
-		TArray<FNavigationObstacleHashGrid2D::ItemIDType> NearbyEntities;
-		NearbyEntities.Reserve(16);
-		ObstacleGrid.Query(QueryBox, NearbyEntities);
-
-		// We'll pick the first entity that passes, this ensure that it's random one.
-		Algo::RandomShuffle(NearbyEntities);
-
-		const FVector Location = Transform.GetLocation();
-		
-		for (const FNavigationObstacleHashGrid2D::ItemIDType NearbyEntity : NearbyEntities)
+		auto FindTargetFunc = [&EntityManager, &LastTrackedEntity, CosAngleThreshold, &LookAt, Entity]
+		(TStridedView<FMassEntityHandle> InCandidates, const FTransform& InTransform)
 		{
-			// This can happen if we remove entities in the system.
-			if (!EntityManager.IsEntityValid(NearbyEntity.Entity))
+			const FVector Location = InTransform.GetLocation();
+			for (FMassEntityHandle CandidateEntity : InCandidates)
 			{
-				UE_LOG(LogMassBehavior, VeryVerbose, TEXT("Nearby entity is invalid, skipped."));
-				continue;
-			}
+				// This can happen if we remove entities in the system.
+				if (!EntityManager.IsEntityValid(CandidateEntity))
+				{
+					UE_LOG(LogMassBehavior, VeryVerbose, TEXT("Nearby entity is invalid, skipped."));
+					continue;
+				}
 
-			// Do not select same target twice in a row.
-			if (NearbyEntity.Entity == LastTrackedEntity)
-			{
-				continue;
-			}
+				// Do not select self
+				if (CandidateEntity == Entity)
+				{
+					continue;
+				}
 
-			FMassEntityView EntityView(EntityManager, NearbyEntity.Entity);
-			if (!EntityView.HasTag<FMassLookAtTargetTag>())
-			{
-				continue;
-			}
+				// Do not select same target twice in a row.
+				if (CandidateEntity == LastTrackedEntity)
+				{
+					continue;
+				}
 
-			// TargetTag is added through the LookAtTargetTrait and Transform was added with it
-			const FTransformFragment& TargetTransform = EntityView.GetFragmentData<FTransformFragment>();
-			const FVector TargetLocation = TargetTransform.GetTransform().GetLocation();
-			FVector Direction = (TargetLocation - Location).GetSafeNormal();
-			Direction = Transform.InverseTransformVector(Direction);
+				FMassEntityView EntityView(EntityManager, CandidateEntity);
+				if (!EntityView.HasTag<FMassLookAtTargetTag>())
+				{
+					continue;
+				}
 
-			const bool bIsTargetInView = FVector::DotProduct(FVector::ForwardVector, Direction) > CosAngleThreshold;
-			if (bIsTargetInView)
-			{
-				LookAt.GazeDirection = Direction;
-				LookAt.GazeTrackedEntity = NearbyEntity.Entity;
-				LookAt.GazeTargetLocation = TargetLocation;
-				bTargetFound = true;
-				break;
-			}
+				// TargetTag is added through the LookAtTargetTrait and Transform was added with it
+				const FTransformFragment& TargetTransform = EntityView.GetFragmentData<FTransformFragment>();
+				const FVector TargetLocation = TargetTransform.GetTransform().GetLocation();
+				FVector Direction = (TargetLocation - Location).GetSafeNormal();
+				Direction = InTransform.InverseTransformVector(Direction);
 
-			// Allow to pick entities out of view if they are moving towards us.
-			if (const FMassVelocityFragment* Velocity = EntityView.GetFragmentDataPtr<FMassVelocityFragment>())
-			{
-				const FVector MoveDirection = Transform.InverseTransformVector(Velocity->Value.GetSafeNormal());
-				
-				const bool bIsTargetMovingTowards = FVector::DotProduct(MoveDirection, -Direction) > CosAngleThreshold; // Direction negated as it is from the agent to target, and we want target to agent. 
-				if (bIsTargetMovingTowards)
+				const bool bIsTargetInView = FVector::DotProduct(FVector::ForwardVector, Direction) > CosAngleThreshold;
+				if (bIsTargetInView)
 				{
 					LookAt.GazeDirection = Direction;
-					LookAt.GazeTrackedEntity = NearbyEntity.Entity;
+					LookAt.GazeTrackedEntity = CandidateEntity;
 					LookAt.GazeTargetLocation = TargetLocation;
-					bTargetFound = true;
-					break;
+					return true;
+				}
+
+				// Allow to pick entities out of view if they are moving towards us.
+				if (const FMassVelocityFragment* Velocity = EntityView.GetFragmentDataPtr<FMassVelocityFragment>())
+				{
+					const FVector MoveDirection = InTransform.InverseTransformVector(Velocity->Value.GetSafeNormal());
+
+					// Direction negated as it is from the agent to target, and we want target to agent. 
+					const bool bIsTargetMovingTowards = FVector::DotProduct(MoveDirection, -Direction) > CosAngleThreshold;
+					if (bIsTargetMovingTowards)
+					{
+						LookAt.GazeDirection = Direction;
+						LookAt.GazeTrackedEntity = CandidateEntity;
+						LookAt.GazeTargetLocation = TargetLocation;
+						return true;
+					}
 				}
 			}
-			
+
+			return false;
+		};
+
+		// Process from LookAt targets grid first
+		{
+			TArray<UE::Mass::LookAt::FTargetHashGrid2D::ItemIDType> NearbyEntities;
+			NearbyEntities.Reserve(16);
+			if (LookAtSystem.Query(QueryBox, NearbyEntities))
+			{
+				// We'll pick the first entity that passes, this ensure that it's random one.
+				Algo::RandomShuffle(NearbyEntities);
+
+				TStridedView<FMassEntityHandle> NearbyLookAtTargetEntities = MakeStridedView(NearbyEntities);
+				bTargetFound = FindTargetFunc(NearbyLookAtTargetEntities, Transform);
+			}
+		}
+
+		// Process from obstacles grid
+		if (!bTargetFound && UE::MassBehavior::Tweakables::bConsiderObstacleTargets)
+		{
+			TArray<FNavigationObstacleHashGrid2D::ItemIDType> NearbyEntities;
+			NearbyEntities.Reserve(16);
+			const FNavigationObstacleHashGrid2D& Grid = MassNavSystem.GetObstacleGrid();
+			if (Grid.GetItems().Num() > 0)
+			{
+				Grid.Query(QueryBox, NearbyEntities);
+
+				// We'll pick the first entity that passes, this ensure that it's random one.
+				Algo::RandomShuffle(NearbyEntities);
+
+				TStridedView<FMassEntityHandle> NearbyAvoidanceObstacleEntities = MakeStridedView(NearbyEntities, &FNavigationObstacleHashGrid2D::ItemIDType::Entity);
+				bTargetFound = FindTargetFunc(NearbyAvoidanceObstacleEntities, Transform);
+			}
 		}
 	}
 
@@ -328,7 +374,7 @@ void UMassLookAtProcessor::FindNewGazeTarget(const UMassNavigationSubsystem& Mas
 		LookAt.GazeTargetLocation = FVector::ZeroVector;
 	}
 
-	// @todo: This does not currently carry over time. It's intentional, since there might be big gaps between updates.
+	// @todo LookAt: This does not currently carry over time. It's intentional, since there might be big gaps between updates.
 	LookAt.GazeStartTime = CurrentTime;
 	LookAt.GazeDuration = FMath::FRandRange(FMath::Max(Duration - DurationVariation, 0.f), Duration + DurationVariation);
 }
@@ -532,3 +578,198 @@ void UMassLookAtProcessor::BuildTrajectory(const UZoneGraphSubsystem& ZoneGraphS
 	}
 #endif // WITH_MASSGAMEPLAY_DEBUG
 }
+
+//----------------------------------------------------------------------//
+// UMassLookAtTargetProcessor
+//----------------------------------------------------------------------//
+UMassLookAtTargetGridProcessor::UMassLookAtTargetGridProcessor()
+	: AddToGridQuery(*this)
+	, UpdateGridQuery(*this)
+	, RemoveFromGridQuery(*this)
+{
+	ExecutionFlags = static_cast<int32>(EProcessorExecutionFlags::AllNetModes);
+	ExecutionOrder.ExecuteBefore.Add(UE::Mass::ProcessorGroupNames::Tasks);
+}
+
+void UMassLookAtTargetGridProcessor::ConfigureQueries(const TSharedRef<FMassEntityManager>& EntityManager)
+{
+	FMassEntityQuery BaseQuery(EntityManager);
+	BaseQuery.AddRequirement<FTransformFragment>(EMassFragmentAccess::ReadOnly);
+	BaseQuery.AddRequirement<FMassLookAtTargetGridCellLocationFragment>(EMassFragmentAccess::ReadWrite);
+	BaseQuery.AddTagRequirement<FMassLookAtTargetTag>(EMassFragmentPresence::All);
+	BaseQuery.AddSubsystemRequirement<UMassLookAtSubsystem>(EMassFragmentAccess::ReadWrite);
+
+	AddToGridQuery = BaseQuery;
+	AddToGridQuery.AddTagRequirement<FMassOffLODTag>(EMassFragmentPresence::None);
+	AddToGridQuery.AddTagRequirement<FMassInLookAtTargetGridTag>(EMassFragmentPresence::None);
+	AddToGridQuery.RegisterWithProcessor(*this);
+
+	UpdateGridQuery = BaseQuery;
+	UpdateGridQuery.AddTagRequirement<FMassOffLODTag>(EMassFragmentPresence::None);
+	UpdateGridQuery.AddTagRequirement<FMassInLookAtTargetGridTag>(EMassFragmentPresence::All);
+	UpdateGridQuery.RegisterWithProcessor(*this);
+
+	RemoveFromGridQuery = BaseQuery;
+	RemoveFromGridQuery.AddTagRequirement<FMassOffLODTag>(EMassFragmentPresence::All);
+	RemoveFromGridQuery.AddTagRequirement<FMassInLookAtTargetGridTag>(EMassFragmentPresence::All);
+	RemoveFromGridQuery.RegisterWithProcessor(*this);
+}
+
+void UMassLookAtTargetGridProcessor::Execute(FMassEntityManager& EntityManager, FMassExecutionContext& Context)
+{
+	QUICK_SCOPE_CYCLE_COUNTER(LookAtProcessor_Run);
+
+	constexpr float Radius = 50.f;
+
+	AddToGridQuery.ForEachEntityChunk(Context, [](FMassExecutionContext& Context)
+	{
+		UMassLookAtSubsystem& LookAtSubsystem = Context.GetMutableSubsystemChecked<UMassLookAtSubsystem>();
+		TConstArrayView<FTransformFragment> LocationList = Context.GetFragmentView<FTransformFragment>();
+		TArrayView<FMassLookAtTargetGridCellLocationFragment> CellLocationList = Context.GetMutableFragmentView<FMassLookAtTargetGridCellLocationFragment>();
+
+		for (FMassExecutionContext::FEntityIterator EntityIt = Context.CreateEntityIterator(); EntityIt; ++EntityIt)
+		{
+			const FMassEntityHandle Entity = Context.GetEntity(EntityIt);
+			const FVector NewPos = LocationList[EntityIt].GetTransform().GetLocation();
+			const FBox NewBounds(NewPos - FVector(Radius, Radius, 0.f), NewPos + FVector(Radius, Radius, 0.f));
+			CellLocationList[EntityIt].CellLocation = LookAtSubsystem.AddTarget(Entity, NewBounds);
+
+			Context.Defer().AddTag<FMassInLookAtTargetGridTag>(Entity);
+		}
+	});
+
+	UpdateGridQuery.ForEachEntityChunk(Context, [](FMassExecutionContext& Context)
+	{
+		UMassLookAtSubsystem& LookAtSubsystem = Context.GetMutableSubsystemChecked<UMassLookAtSubsystem>();
+		TConstArrayView<FTransformFragment> LocationList = Context.GetFragmentView<FTransformFragment>();
+		TArrayView<FMassLookAtTargetGridCellLocationFragment> CellLocationList = Context.GetMutableFragmentView<FMassLookAtTargetGridCellLocationFragment>();
+		TArray<TTuple<const FMassEntityHandle, UE::Mass::LookAt::FTargetHashGrid2D::FCellLocation&, const FBox&>> AllUpdates;
+
+		for (FMassExecutionContext::FEntityIterator EntityIt = Context.CreateEntityIterator(); EntityIt; ++EntityIt)
+		{
+			const FVector NewPos = LocationList[EntityIt].GetTransform().GetLocation();
+			const FBox NewBounds(NewPos - FVector(Radius, Radius, 0.f), NewPos + FVector(Radius, Radius, 0.f));
+			AllUpdates.Emplace(Context.GetEntity(EntityIt), CellLocationList[EntityIt].CellLocation, NewBounds);
+		}
+
+		LookAtSubsystem.BatchMoveTarget(AllUpdates);
+	});
+
+	RemoveFromGridQuery.ForEachEntityChunk(Context, [](FMassExecutionContext& Context)
+	{
+		UMassLookAtSubsystem& LookAtSubsystem = Context.GetMutableSubsystemChecked<UMassLookAtSubsystem>();
+		TArrayView<FMassLookAtTargetGridCellLocationFragment> CellLocationList = Context.GetMutableFragmentView<FMassLookAtTargetGridCellLocationFragment>();
+
+		for (FMassExecutionContext::FEntityIterator EntityIt = Context.CreateEntityIterator(); EntityIt; ++EntityIt)
+		{
+			const FMassEntityHandle Entity = Context.GetEntity(EntityIt);
+			LookAtSubsystem.RemoveTarget(Entity, CellLocationList[EntityIt].CellLocation);
+			CellLocationList[EntityIt].CellLocation = UE::Mass::LookAt::FTargetHashGrid2D::FCellLocation();
+
+			Context.Defer().RemoveTag<FMassInLookAtTargetGridTag>(Entity);
+		}
+	});
+}
+
+//----------------------------------------------------------------------//
+//  UMassLookAtTargetRemoverProcessor
+//----------------------------------------------------------------------//
+UMassLookAtTargetRemoverProcessor::UMassLookAtTargetRemoverProcessor()
+	: Query(*this)
+{
+	ObservedType = FMassLookAtTargetGridCellLocationFragment::StaticStruct();
+	Operation = EMassObservedOperation::Remove;
+	ExecutionFlags = static_cast<int32>(EProcessorExecutionFlags::AllNetModes);
+}
+
+void UMassLookAtTargetRemoverProcessor::ConfigureQueries(const TSharedRef<FMassEntityManager>& EntityManager)
+{
+	Query.AddRequirement<FMassLookAtTargetGridCellLocationFragment>(EMassFragmentAccess::ReadWrite);
+	Query.AddSubsystemRequirement<UMassLookAtSubsystem>(EMassFragmentAccess::ReadWrite);
+}
+
+void UMassLookAtTargetRemoverProcessor::Execute(FMassEntityManager& EntityManager, FMassExecutionContext& Context)
+{
+	Query.ForEachEntityChunk(Context, [this](FMassExecutionContext& Context)
+	{
+		UMassLookAtSubsystem& LookAtSubsystem = Context.GetMutableSubsystemChecked<UMassLookAtSubsystem>();
+		const TArrayView<FMassLookAtTargetGridCellLocationFragment> CellLocationList = Context.GetMutableFragmentView<FMassLookAtTargetGridCellLocationFragment>();
+
+		for (FMassExecutionContext::FEntityIterator EntityIt = Context.CreateEntityIterator(); EntityIt; ++EntityIt)
+		{
+			LookAtSubsystem.RemoveTarget(Context.GetEntity(EntityIt), CellLocationList[EntityIt].CellLocation);
+		}
+	});
+}
+
+//----------------------------------------------------------------------//
+//  UMassLookAtRequestInitializer
+//----------------------------------------------------------------------//
+UMassLookAtRequestInitializer::UMassLookAtRequestInitializer()
+	: Query(*this)
+{
+	ObservedType = FMassLookAtRequestFragment::StaticStruct();
+	Operation = EMassObservedOperation::Add;
+	ExecutionFlags = static_cast<int32>(EProcessorExecutionFlags::Client | EProcessorExecutionFlags::Standalone);
+}
+
+void UMassLookAtRequestInitializer::ConfigureQueries(const TSharedRef<FMassEntityManager>& EntityManager)
+{
+	Query.AddRequirement<FMassLookAtRequestFragment>(EMassFragmentAccess::ReadOnly);
+	ProcessorRequirements.AddSubsystemRequirement<UMassLookAtSubsystem>(EMassFragmentAccess::ReadWrite);
+}
+
+void UMassLookAtRequestInitializer::Execute(FMassEntityManager& EntityManager, FMassExecutionContext& Context)
+{
+	UMassLookAtSubsystem& LookAtSubsystem = Context.GetMutableSubsystemChecked<UMassLookAtSubsystem>();
+	TArray<UMassLookAtSubsystem::FRequest> Requests;
+
+	Query.ForEachEntityChunk(Context, [&LookAtSubsystem, &Requests](FMassExecutionContext& Context)
+	{
+		const TConstArrayView<FMassLookAtRequestFragment> RequestList = Context.GetFragmentView<FMassLookAtRequestFragment>();
+
+		Requests.Reserve(Requests.Num() + Context.GetNumEntities());
+		for (FMassExecutionContext::FEntityIterator EntityIt = Context.CreateEntityIterator(); EntityIt; ++EntityIt)
+		{
+			Requests.Add({ Context.GetEntity(EntityIt), RequestList[EntityIt] });
+		}
+	});
+
+	LookAtSubsystem.RegisterRequests(Context, MoveTemp(Requests));
+}
+
+//----------------------------------------------------------------------//
+//  UMassLookAtRequestInitializer
+//----------------------------------------------------------------------//
+UMassLookAtRequestDeinitializer::UMassLookAtRequestDeinitializer()
+	: Query(*this)
+{
+	ObservedType = FMassLookAtRequestFragment::StaticStruct();
+	Operation = EMassObservedOperation::Remove;
+	ExecutionFlags = static_cast<int32>(EProcessorExecutionFlags::Client | EProcessorExecutionFlags::Standalone);
+}
+
+void UMassLookAtRequestDeinitializer::ConfigureQueries(const TSharedRef<FMassEntityManager>& EntityManager)
+{
+	Query.AddRequirement<FMassLookAtRequestFragment>(EMassFragmentAccess::ReadOnly);
+	ProcessorRequirements.AddSubsystemRequirement<UMassLookAtSubsystem>(EMassFragmentAccess::ReadWrite);
+}
+
+void UMassLookAtRequestDeinitializer::Execute(FMassEntityManager& EntityManager, FMassExecutionContext& Context)
+{
+	UMassLookAtSubsystem& LookAtSubsystem = Context.GetMutableSubsystemChecked<UMassLookAtSubsystem>();
+	TArray<FMassEntityHandle> Requests;
+
+	Query.ForEachEntityChunk(Context, [&LookAtSubsystem, &Requests](FMassExecutionContext& Context)
+	{
+		const TConstArrayView<FMassLookAtRequestFragment> RequestList = Context.GetFragmentView<FMassLookAtRequestFragment>();
+
+		Requests.Reserve(Requests.Num() + Context.GetNumEntities());
+		for (FMassExecutionContext::FEntityIterator EntityIt = Context.CreateEntityIterator(); EntityIt; ++EntityIt)
+		{
+			Requests.Add(Context.GetEntity(EntityIt));
+		}
+	});
+
+	LookAtSubsystem.UnregisterRequests(Context, Requests);
+}
diff --git a/Engine/Plugins/AI/MassAI/Source/MassAIBehavior/Private/MassLookAtSubsystem.cpp b/Engine/Plugins/AI/MassAI/Source/MassAIBehavior/Private/MassLookAtSubsystem.cpp
new file mode 100644
--- /dev/null
+++ b/Engine/Plugins/AI/MassAI/Source/MassAIBehavior/Private/MassLookAtSubsystem.cpp
@@ -0,0 +1,271 @@
+// Copyright Epic Games, Inc. All Rights Reserved.
+
+#include "MassLookAtSubsystem.h"
+#include "GameFramework/Pawn.h"
+#include "MassActorSubsystem.h"
+#include "MassAIBehaviorTypes.h"
+#include "MassEntitySubsystem.h"
+#include "MassEntityView.h"
+#include "Debugger/StateTreeDebuggerTypes.h"
+
+namespace UE::Mass::LookAt::Private
+{
+
+static FAutoConsoleCommandWithWorldArgsAndOutputDevice CmdSendLookAtPlayerRequestToAll(
+	TEXT("ai.debug.mass.SendLookAtPlayerRequestToAll"),
+	TEXT(""),
+	FConsoleCommandWithWorldArgsAndOutputDeviceDelegate::CreateLambda([](const TArray<FString>& Args, const UWorld* World, FOutputDevice& OutputDevice)
+		{
+			UMassActorSubsystem* ActorSubsystem = World->GetSubsystem<UMassActorSubsystem>();
+			if (ActorSubsystem == nullptr)
+			{
+				OutputDevice.Log(ELogVerbosity::Error, TEXT("%hs failure: unable to find MassActorSubsystem"));
+				return;
+			}
+
+			const APlayerController* PlayerController = World->GetFirstPlayerController();
+			if (PlayerController == nullptr)
+			{
+				OutputDevice.Log(ELogVerbosity::Error, TEXT("%hs failure: unable to find the player controller"));
+				return;
+			}
+
+			const APawn* PlayerPawn = PlayerController->GetPawn();
+			if (PlayerPawn == nullptr)
+			{
+				OutputDevice.Log(ELogVerbosity::Error, TEXT("%hs failure: unable to find the player pawn"));
+				return;
+			}
+
+			UMassEntitySubsystem* EntitySubsystem = World->GetSubsystem<UMassEntitySubsystem>();
+			if (ActorSubsystem == nullptr)
+			{
+				OutputDevice.Log(ELogVerbosity::Error, TEXT("%hs failure: unable to find UMassEntitySubsystem"));
+				return;
+			}
+
+			const FMassEntityHandle PlayerEntity = ActorSubsystem->GetEntityHandleFromActor(PlayerPawn);
+			if (!PlayerEntity.IsSet())
+			{
+				OutputDevice.Log(ELogVerbosity::Error, TEXT("%hs failure: unable to find a MassEntity associated to the player"));
+				return;
+			}
+
+			// Create requests for all entities with a MassLookAtFragment
+			TSharedPtr<FMassCommandBuffer> CommandBuffer = MakeShareable(new FMassCommandBuffer());
+			CommandBuffer->PushCommand<FMassDeferredCreateCommand>([PlayerEntity](FMassEntityManager& InEntityManager)
+				{
+					FMassEntityQuery EntityQuery(InEntityManager.AsShared());
+					EntityQuery.AddRequirement<FMassLookAtFragment>(EMassFragmentAccess::ReadOnly);
+
+					// @todo LookAt: replace this part by FMassEntityQuery::GetMatchingEntityHandles() once available
+					TArray<FMassEntityHandle> Viewers;
+					FMassExecutionContext Context(InEntityManager, 0.0f);
+					EntityQuery.ForEachEntityChunk(Context, [&Viewers](const FMassExecutionContext& Context)
+						{
+							Viewers.Append(Context.GetEntities());
+						});
+
+					FMassArchetypeCompositionDescriptor Composition;
+					Composition.Fragments.Add(*FMassLookAtRequestFragment::StaticStruct());
+
+					TArray<FMassEntityHandle> Requests;
+					const TSharedRef<FMassEntityManager::FEntityCreationContext> CreationContext =
+						InEntityManager.BatchCreateEntities(InEntityManager.CreateArchetype(Composition), Viewers.Num(), Requests);
+
+					check(Viewers.Num() == Requests.Num());
+					for (int Index = 0; Index < Requests.Num(); ++Index)
+					{
+						InEntityManager.SetEntityFragmentValues(Requests[Index]
+							, { FInstancedStruct::Make(FMassLookAtRequestFragment{Viewers[Index], PlayerEntity}) });
+					}
+				});
+
+			ExecuteOnGameThread(UE_SOURCE_LOCATION, [WeakEntityManager = EntitySubsystem->GetMutableEntityManager().AsWeak(), CommandBuffer]()
+				{
+					if (const TSharedPtr<FMassEntityManager> SharedEntityManager = WeakEntityManager.Pin())
+					{
+						SharedEntityManager->FlushCommands(CommandBuffer);
+					}
+					else
+					{
+						CommandBuffer->CancelCommands();
+					}
+				});
+		}
+	));
+
+} // UE::Mass::LookAt::Private
+
+//----------------------------------------------------------------------//
+//  UMassLookAtSubsystem
+//----------------------------------------------------------------------//
+void UMassLookAtSubsystem::Initialize(FSubsystemCollectionBase& Collection)
+{
+	Super::Initialize(Collection);
+
+	OverrideSubsystemTraits<UMassLookAtSubsystem>(Collection);
+}
+
+TStatId UMassLookAtSubsystem::GetStatId() const
+{
+	RETURN_QUICK_DECLARE_CYCLE_STAT(UMassLookAtSubsystem, STATGROUP_Tickables);
+}
+
+void UMassLookAtSubsystem::RegisterRequests(const FMassExecutionContext& InContext, TArray<FRequest>&& InRequests)
+{
+	TArray<int32> DirtyViewers;
+	DirtyViewers.Reserve(InRequests.Num());
+
+	{
+		UE_MT_SCOPED_WRITE_ACCESS(RequestsAccessDetector);
+
+		ActiveRequests.Reserve(ActiveRequests.Num() + (InRequests.Num() - ActiveRequestsFreeList.Num()));
+
+		for (const FRequest& NewRequest : InRequests)
+		{
+			if (!NewRequest.Parameters.TargetEntity.IsSet())
+			{
+				UE_VLOG_UELOG(this, LogMassBehavior, Error, TEXT("Ignoring LookAt request: invalid target entity"));
+				continue;
+			}
+
+			int32 NewRequestIndex = INDEX_NONE;
+			if (ActiveRequestsFreeList.Num())
+			{
+				NewRequestIndex = ActiveRequestsFreeList.Pop();
+				ActiveRequests[NewRequestIndex] = NewRequest;
+			}
+			else
+			{
+				NewRequestIndex = ActiveRequests.Add(NewRequest);
+			}
+
+			int32 ViewerDataIndex = PerViewerRequests.IndexOfByPredicate([ViewerEntity = NewRequest.Parameters.ViewerEntity](const FViewerRequest& ViewerRequest)
+			{
+				return ViewerEntity == ViewerRequest.Viewer;
+			});
+
+			if (ViewerDataIndex == INDEX_NONE)
+			{
+				ViewerDataIndex = PerViewerRequests.Num();
+				PerViewerRequests.AddDefaulted_GetRef().Viewer = NewRequest.Parameters.ViewerEntity;
+			}
+
+			DirtyViewers.Add(ViewerDataIndex);
+			PerViewerRequests[ViewerDataIndex].RequestIndices.Add(NewRequestIndex);
+		}
+	}
+
+	UpdateLookAts(InContext, DirtyViewers);
+}
+
+void UMassLookAtSubsystem::UnregisterRequests(const FMassExecutionContext& InContext, const TConstArrayView<FMassEntityHandle> InRequests)
+{
+	TArray<int32> DirtyViewers;
+	DirtyViewers.Reserve(InRequests.Num());
+
+	{
+		UE_MT_SCOPED_WRITE_ACCESS(RequestsAccessDetector);
+
+		for (const FMassEntityHandle& RemovedRequest : InRequests)
+		{
+			int32 InvalidatedIndex = INDEX_NONE;
+			FMassEntityHandle ViewerEntity;
+
+			for (int32 ActiveRequestIndex = 0; ActiveRequestIndex < ActiveRequests.Num(); ++ActiveRequestIndex)
+			{
+				if (ActiveRequests[ActiveRequestIndex].RequestHandle == RemovedRequest)
+				{
+					InvalidatedIndex = ActiveRequestIndex;
+
+					// Invalidate entry and add to the free list
+					ViewerEntity = ActiveRequests[ActiveRequestIndex].Parameters.ViewerEntity;
+					ActiveRequests[ActiveRequestIndex].RequestHandle = FMassEntityHandle();
+					ActiveRequestsFreeList.Push(ActiveRequestIndex);
+					break;
+				}
+			}
+
+			if (ensureMsgf(InvalidatedIndex != INDEX_NONE, TEXT("Trying to remove a request that was never registered")))
+			{
+				const int32 ViewerDataIndex = PerViewerRequests.IndexOfByPredicate([ViewerEntity](const FViewerRequest& ViewerRequest)
+					{
+						return ViewerEntity == ViewerRequest.Viewer;
+					});
+
+				if (ensureMsgf(ViewerDataIndex != INDEX_NONE, TEXT("Unable to find the per viewer data; looks like a bookkeeping issue")))
+				{
+					DirtyViewers.Add(ViewerDataIndex);
+					PerViewerRequests[ViewerDataIndex].RequestIndices.Remove(InvalidatedIndex);
+				}
+			}
+		}
+	}
+
+	UpdateLookAts(InContext, DirtyViewers);
+}
+
+void UMassLookAtSubsystem::UpdateLookAts(const FMassExecutionContext& Context, const TConstArrayView<int32> DirtyViewers)
+{
+	TArray<TPair<FMassEntityHandle, FMassLookAtRequestFragment>> Updated;
+	for (int32 DirtyViewerIndex : DirtyViewers)
+	{
+		check(PerViewerRequests.IsValidIndex(DirtyViewerIndex));
+
+		int32 SelectedRequestIndex = INDEX_NONE;
+
+		for (const int32 RequestIndex : PerViewerRequests[DirtyViewerIndex].RequestIndices)
+		{
+			check(ActiveRequests.IsValidIndex(RequestIndex));
+
+			SelectedRequestIndex = RequestIndex;
+			break;
+
+			// @todo LookAt: Add and handle priorities
+			//if (SelectedRequestIndex == INDEX_NONE)
+			//{
+			//	SelectedRequestIndex = RequestIndex;
+			//	continue;
+			//}
+
+
+			//if (ActiveRequests[RequestIndex].Parameters.Priority
+			//	> ActiveRequests[SelectedRequestIndex].Parameters.Priority)
+			//{
+			//	SelectedRequestIndex = RequestIndex;
+			//}
+		}
+
+		Updated.Add(
+			{
+				PerViewerRequests[DirtyViewerIndex].Viewer
+				, SelectedRequestIndex != INDEX_NONE ? ActiveRequests[SelectedRequestIndex].Parameters : FMassLookAtRequestFragment{}
+			});
+	}
+
+	Context.Defer().PushCommand<FMassDeferredSetCommand>([Updated = MoveTemp(Updated)](const FMassEntityManager& Manager)
+		{
+			for (const TPair<FMassEntityHandle, FMassLookAtRequestFragment>& Pair : Updated)
+			{
+				if (FMassLookAtFragment* LookAtFragment = Manager.GetFragmentDataPtr<FMassLookAtFragment>(Pair.Key))
+				{
+					if (LookAtFragment->TrackedEntity == Pair.Value.TargetEntity)
+					{
+						continue;
+					}
+
+					// Only update main LookAt information and don't modify gaze related one
+					if (Pair.Value.TargetEntity.IsSet())
+					{
+						LookAtFragment->LookAtMode = EMassLookAtMode::LookAtEntity;
+						LookAtFragment->TrackedEntity = Pair.Value.TargetEntity;
+					}
+					else
+					{
+						LookAtFragment->LookAtMode = EMassLookAtMode::LookForward;
+					}
+				}
+			}
+		});
+}
diff --git a/Engine/Plugins/AI/MassAI/Source/MassAIBehavior/Public/MassLookAtFragments.h b/Engine/Plugins/AI/MassAI/Source/MassAIBehavior/Public/MassLookAtFragments.h
--- a/Engine/Plugins/AI/MassAI/Source/MassAIBehavior/Public/MassLookAtFragments.h
+++ b/Engine/Plugins/AI/MassAI/Source/MassAIBehavior/Public/MassLookAtFragments.h
@@ -5,6 +5,7 @@
 #include "MassEntityHandle.h"
 #include "MassEntityTypes.h"
 #include "MassCommonTypes.h"
+#include "MassLookAtTypes.h"
 #include "ZoneGraphTypes.h"
 #include "Containers/StaticArray.h"
 #include "MassLookAtFragments.generated.h"
@@ -130,6 +131,43 @@ struct MASSAIBEHAVIOR_API FMassLookAtTargetTag : public FMassTag
 	GENERATED_BODY()
 };
 
+/**
+ * Struct that holds all parameters of a look at request
+ */
+USTRUCT()
+struct MASSAIBEHAVIOR_API FMassLookAtRequestFragment : public FMassFragment
+{
+	GENERATED_BODY()
+
+	FMassLookAtRequestFragment() = default;
+	FMassLookAtRequestFragment(const FMassEntityHandle ViewerEntity, const FMassEntityHandle TargetEntity)
+		: ViewerEntity(ViewerEntity)
+		, TargetEntity(TargetEntity)
+	{
+	}
+
+	UPROPERTY(Transient)
+	FMassEntityHandle ViewerEntity;
+
+	UPROPERTY(Transient)
+	FMassEntityHandle TargetEntity;
+};
+
+/** Cell location for LookAt targets */
+USTRUCT()
+struct MASSAIBEHAVIOR_API FMassLookAtTargetGridCellLocationFragment : public FMassFragment
+{
+	GENERATED_BODY()
+
+	UE::Mass::LookAt::FTargetHashGrid2D::FCellLocation CellLocation;
+};
+
+/** Tag to tell if the entity is in the LookAt target grid */
+USTRUCT()
+struct MASSAIBEHAVIOR_API FMassInLookAtTargetGridTag : public FMassTag
+{
+	GENERATED_BODY()
+};
 
 USTRUCT()
 struct MASSAIBEHAVIOR_API FMassLookAtTrajectoryPoint
diff --git a/Engine/Plugins/AI/MassAI/Source/MassAIBehavior/Public/MassLookAtProcessors.h b/Engine/Plugins/AI/MassAI/Source/MassAIBehavior/Public/MassLookAtProcessors.h
--- a/Engine/Plugins/AI/MassAI/Source/MassAIBehavior/Public/MassLookAtProcessors.h
+++ b/Engine/Plugins/AI/MassAI/Source/MassAIBehavior/Public/MassLookAtProcessors.h
@@ -2,6 +2,7 @@
 
 #pragma once
 
+#include "MassObserverProcessor.h"
 #include "MassProcessor.h"
 #if UE_ENABLE_INCLUDE_ORDER_DEPRECATED_IN_5_6
 #include "MassSignalProcessorBase.h"
@@ -9,6 +10,7 @@
 #endif // UE_ENABLE_INCLUDE_ORDER_DEPRECATED_IN_5_6
 #include "MassLookAtProcessors.generated.h"
 
+class UMassLookAtSubsystem;
 class UMassNavigationSubsystem;
 class UZoneGraphSubsystem;
 struct FMassLookAtFragment;
@@ -33,7 +35,13 @@ protected:
 	virtual void Execute(FMassEntityManager& EntityManager, FMassExecutionContext& Context) override;
 
 	/** Selects a nearby target if possible or use a random fixed direction */
-	void FindNewGazeTarget(const UMassNavigationSubsystem& MassNavSystem, const FMassEntityManager& EntityManager, const double CurrentTime, const FTransform& Transform, FMassLookAtFragment& LookAt) const;
+	void FindNewGazeTarget(const UMassNavigationSubsystem& MassNavSystem
+		, const UMassLookAtSubsystem& LookAtSystem
+		, const FMassEntityManager& EntityManager
+		, const double CurrentTime
+		, const FTransform& Transform
+		, FMassEntityHandle Entity
+		, FMassLookAtFragment& LookAt) const;
 
 	/** Updates look direction based on look at trajectory. */
 	void UpdateLookAtTrajectory(const FTransform& Transform, const FMassZoneGraphLaneLocationFragment& ZoneGraphLocation,
@@ -71,3 +79,71 @@ protected:
 
 	FMassEntityQuery EntityQuery_Conditional;
 };
+
+/**
+ * Processor to maintain a list of LookAt targets in a spatial query structure in the subsystem
+ */
+UCLASS()
+class MASSAIBEHAVIOR_API UMassLookAtTargetGridProcessor : public UMassProcessor
+{
+	GENERATED_BODY()
+
+	UMassLookAtTargetGridProcessor();
+
+protected:
+
+	virtual void ConfigureQueries(const TSharedRef<FMassEntityManager>& EntityManager) override;
+	virtual void Execute(FMassEntityManager& EntityManager, FMassExecutionContext& Context) override;
+
+	FMassEntityQuery AddToGridQuery;
+	FMassEntityQuery UpdateGridQuery;
+	FMassEntityQuery RemoveFromGridQuery;
+};
+
+/** Deinitializer processor to remove targets from the hash grid */
+UCLASS()
+class MASSAIBEHAVIOR_API UMassLookAtTargetRemoverProcessor : public UMassObserverProcessor
+{
+	GENERATED_BODY()
+
+	UMassLookAtTargetRemoverProcessor();
+
+protected:
+	virtual void ConfigureQueries(const TSharedRef<FMassEntityManager>& EntityManager) override;
+	virtual void Execute(FMassEntityManager& EntityManager, FMassExecutionContext& Context) override;
+
+	FMassEntityQuery Query;
+};
+
+/** Initializer processing new LookAt requests */
+UCLASS()
+class MASSAIBEHAVIOR_API UMassLookAtRequestInitializer : public UMassObserverProcessor
+{
+	GENERATED_BODY()
+
+	UMassLookAtRequestInitializer();
+
+protected:
+	virtual void ConfigureQueries(const TSharedRef<FMassEntityManager>& EntityManager) override;
+	virtual void Execute(FMassEntityManager& EntityManager, FMassExecutionContext& Context) override;
+
+private:
+	FMassEntityQuery Query;
+};
+
+
+/** Deinitializer processing deleted LookAt requests */
+UCLASS()
+class MASSAIBEHAVIOR_API UMassLookAtRequestDeinitializer : public UMassObserverProcessor
+{
+	GENERATED_BODY()
+
+	UMassLookAtRequestDeinitializer();
+
+protected:
+	virtual void ConfigureQueries(const TSharedRef<FMassEntityManager>& EntityManager) override;
+	virtual void Execute(FMassEntityManager& EntityManager, FMassExecutionContext& Context) override;
+
+private:
+	FMassEntityQuery Query;
+};
\ No newline at end of file
diff --git a/Engine/Plugins/AI/MassAI/Source/MassAIBehavior/Public/MassLookAtSubsystem.h b/Engine/Plugins/AI/MassAI/Source/MassAIBehavior/Public/MassLookAtSubsystem.h
new file mode 100644
--- /dev/null
+++ b/Engine/Plugins/AI/MassAI/Source/MassAIBehavior/Public/MassLookAtSubsystem.h
@@ -0,0 +1,163 @@
+// Copyright Epic Games, Inc. All Rights Reserved.
+
+#pragma once
+
+#include "MassExecutionContext.h"
+#include "MassLookAtFragments.h"
+#include "MassSubsystemBase.h"
+#include "MassLookAtTypes.h"
+#include "MassLookAtSubsystem.generated.h"
+
+/**
+ * Subsystem that keeps track of the LookAt targets
+ */
+UCLASS()
+class MASSAIBEHAVIOR_API UMassLookAtSubsystem : public UMassTickableSubsystemBase
+{
+	GENERATED_BODY()
+
+public:
+
+	/**
+	 * Struct representing a request for a given entity to perform a LookAt action
+	 */
+	struct FRequest
+	{
+		FMassEntityHandle RequestHandle;
+		FMassLookAtRequestFragment Parameters;
+	};
+
+	/**
+	 * Appends the provided requests to the active ones, then updates the LookAt fragments
+	 * of the referenced entity if required.
+	 * @param InContext Context that is used to push deferred commands
+	 * @param InRequests List of new request to register
+	 */
+	void RegisterRequests(const FMassExecutionContext& InContext, TArray<FRequest>&& InRequests);
+
+	/**
+	 * Unregisters the provided requests from the list of active ones, then updates the LookAt fragments
+	 * of the referenced entity if required.
+	 * @param InContext Context that is used to push deferred commands
+	 * @param InRequests List of request to unregister
+	 */
+	void UnregisterRequests(const FMassExecutionContext& InContext, TConstArrayView<FMassEntityHandle> InRequests);
+
+	/**
+	 * Adds new item to the grid for a given entity handle representing a LookAt target.
+	 * @param InEntity - Entity used to identify the item in the grid.
+	 * @param InBounds - Bounding box of the entity.
+	 * @return The cell location where the item was added to, can be used later to move or remove the item.
+	 */
+	[[nodiscard]] UE::Mass::LookAt::FTargetHashGrid2D::FCellLocation AddTarget(const FMassEntityHandle InEntity, const FBox& InBounds)
+	{
+		UE_MT_SCOPED_WRITE_ACCESS(TargetGridAccessDetector);
+		return TargetGrid.Add(InEntity, InBounds);
+	}
+
+	/**
+	 * Moves item based on the entity handle, its previous cell location and its new bounding box.
+	 * @param InEntity - Entity used to identify the item in the grid.
+	 * @param InOldCellLocation - The same cell location the item was previously added or moved with.
+	 * @param InNewBounds - New bounds of the item
+	 * @return The cell location where the item is now added to.
+	 */
+	[[nodiscard]] UE::Mass::LookAt::FTargetHashGrid2D::FCellLocation MoveTarget(const FMassEntityHandle InEntity, const UE::Mass::LookAt::FTargetHashGrid2D::FCellLocation InOldCellLocation, const FBox& InNewBounds)
+	{
+		UE_MT_SCOPED_WRITE_ACCESS(TargetGridAccessDetector);
+		return TargetGrid.Move(InEntity, InOldCellLocation, InNewBounds);
+	}
+
+	/**
+	 * Moves multiple items based on their entity handle, cell location and new bounding box.
+	 * @param InUpdates - List of entity, cell location reference and new bounds to use for each item to update in the grid.
+	 * @note Cell location provided in the tuple is a reference and will be updated in the process.
+	 */
+	void BatchMoveTarget(TArrayView<TTuple<const FMassEntityHandle, UE::Mass::LookAt::FTargetHashGrid2D::FCellLocation&, const FBox&>> InUpdates)
+	{
+		UE_MT_SCOPED_WRITE_ACCESS(TargetGridAccessDetector);
+		TArray<UE::Mass::LookAt::FTargetHashGrid2D::FCellLocation> UpdatedLocations;
+		UpdatedLocations.Reserve(InUpdates.Num());
+		for (TTuple<const FMassEntityHandle, UE::Mass::LookAt::FTargetHashGrid2D::FCellLocation&, const FBox&>& UpdateInfo : InUpdates)
+		{
+			UpdateInfo.Get<1>() = TargetGrid.Move(/*Entity*/UpdateInfo.Get<0>(), /*OldCellLocation*/UpdateInfo.Get<1>(), /*NewBounds*/UpdateInfo.Get<2>());
+		}
+	}
+
+	/**
+	 * Removes item based on the entity handle and the cell location it was added with.
+	 * @param InEntity - Entity used to identify the item in the grid.
+	 * @param InOldCellLocation - The same cell location the item was previously added or moved with.
+	 */
+	void RemoveTarget(const FMassEntityHandle InEntity, const UE::Mass::LookAt::FTargetHashGrid2D::FCellLocation InOldCellLocation)
+	{
+		UE_MT_SCOPED_WRITE_ACCESS(TargetGridAccessDetector);
+		TargetGrid.Remove(InEntity, InOldCellLocation);
+	}
+
+	/**
+	 * Returns entity that potentially touch the bounds. Operates on grid level, can have false positives.
+	 * @param InQueryBox - Query bounding box.
+	 * @param OutEntities - Result of the query, entity handles of potentially overlapping items.
+	 */
+	bool Query(const FBox& InQueryBox, TArray<UE::Mass::LookAt::FTargetHashGrid2D::ItemIDType>& OutEntities) const
+	{
+		UE_MT_SCOPED_READ_ACCESS(TargetGridAccessDetector);
+		TargetGrid.Query(InQueryBox, OutEntities);
+		return OutEntities.Num() > 0;
+	}
+
+	/** @return Number of entities currently registered in the grid. */
+	int32 DebugGetRegisteredTargetCount() const
+	{
+		return TargetGrid.GetItems().Num();
+	}
+
+protected:
+
+	virtual void Initialize(FSubsystemCollectionBase& Collection) override;
+	virtual TStatId GetStatId() const override;
+
+private:
+
+	/**
+	 * Called after registering/unregistering requests to update the LookAtFragment of all viewer entities
+	 * referenced in those requests
+	 */
+	void UpdateLookAts(const FMassExecutionContext& Context, TConstArrayView<int32> DirtyViewers);
+
+	/** Struct to facilitate processing request per viewer entity. */
+	struct FViewerRequest
+	{
+		FMassEntityHandle Viewer;
+		TArray<int32> RequestIndices;
+	};
+
+	/** List of all currently registered (and active) requests as registered by external systems */
+	TArray<FRequest> ActiveRequests;
+
+	/** List of available indices in ActiveRequests (to preserve stable indices when unregistering requests) */
+	TArray<int32> ActiveRequestsFreeList;
+
+	/** Per viewer entity representation of all the active requests */
+	TArray<FViewerRequest> PerViewerRequests;
+
+	/** Multithread access detector to detect threading issues with any list of requests */
+	UE_MT_DECLARE_RW_ACCESS_DETECTOR(RequestsAccessDetector);
+
+	/** Hierarchical 2D hash grid of registered entities representing LookAt targets.*/
+	UE::Mass::LookAt::FTargetHashGrid2D TargetGrid;
+
+	/** Multithread access detector to detect threading issues with the hash grid */
+	UE_MT_DECLARE_RW_ACCESS_DETECTOR(TargetGridAccessDetector);
+};
+
+template<>
+struct TMassExternalSubsystemTraits<UMassLookAtSubsystem>
+{
+	enum
+	{
+		GameThreadOnly = false,
+		ThreadSafeWrite = false, // hash grid not safe
+	};
+};
diff --git a/Engine/Plugins/AI/MassAI/Source/MassAIBehavior/Public/MassLookAtTypes.h b/Engine/Plugins/AI/MassAI/Source/MassAIBehavior/Public/MassLookAtTypes.h
new file mode 100644
--- /dev/null
+++ b/Engine/Plugins/AI/MassAI/Source/MassAIBehavior/Public/MassLookAtTypes.h
@@ -0,0 +1,14 @@
+// Copyright Epic Games, Inc. All Rights Reserved.
+
+#pragma once
+
+#include "HierarchicalHashGrid2D.h"
+#include "MassEntityHandle.h"
+
+namespace UE::Mass::LookAt
+{
+	constexpr int32 HashGridLevelsOfHierarchy = 2;
+	constexpr int32 HashGridRatioBetweenLevels = 4;
+
+	using FTargetHashGrid2D = THierarchicalHashGrid2D<HashGridLevelsOfHierarchy, HashGridRatioBetweenLevels, FMassEntityHandle>;
+}
diff --git a/Engine/Plugins/AI/MassAI/Source/MassAIDebug/Private/GameplayDebuggerCategory_Mass.cpp b/Engine/Plugins/AI/MassAI/Source/MassAIDebug/Private/GameplayDebuggerCategory_Mass.cpp
--- a/Engine/Plugins/AI/MassAI/Source/MassAIDebug/Private/GameplayDebuggerCategory_Mass.cpp
+++ b/Engine/Plugins/AI/MassAI/Source/MassAIDebug/Private/GameplayDebuggerCategory_Mass.cpp
@@ -16,6 +16,7 @@
 #include "MassNavigationFragments.h"
 #include "Steering/MassSteeringFragments.h"
 #include "MassLookAtFragments.h"
+#include "MassLookAtSubsystem.h"
 #include "MassStateTreeFragments.h"
 #include "MassStateTreeExecutionContext.h"
 #include "MassZoneGraphNavigationFragments.h"
@@ -599,11 +600,16 @@ void FGameplayDebuggerCategory_Mass::CollectData(APlayerController* OwnerPC, AAc
 					AddShape(FGameplayDebuggerShape::MakeArrow(MoveBasePos - MoveTarget.Forward * Radius.Radius, MoveBasePos + MoveTarget.Forward * Radius.Radius, 10.0f, 2.0f, FColorList::MediumVioletRed));
 
 					// Look at
-					constexpr FVector::FReal LookArrowLength = 100.;
-					BasePos = EntityLocation + FVector(0,0,EyeHeight);
-
 					if (bHasLookAt)
 					{
+						constexpr FVector::FReal LookArrowLength = 100.;
+						constexpr float TargetArrowHeadSize = 10.f;
+						constexpr float TargetArrowThickness = 1.f;
+						constexpr float LookArrowHeadSize = 20.f;
+						constexpr float LookArrowThickness = 2.f;
+
+						BasePos = EntityLocation + FVector(0,0,EyeHeight);
+
 						const FMassLookAtFragment& LookAt = LookAtList[EntityIndex];
 						const FVector WorldLookDirection = Transform.GetTransform().TransformVector(LookAt.Direction);
 						bool bLookArrowDrawn = false;
@@ -612,9 +618,10 @@ void FGameplayDebuggerCategory_Mass::CollectData(APlayerController* OwnerPC, AAc
 							FVector TargetPosition = LookAt.MainTargetLocation;
 							TargetPosition.Z = BasePos.Z;
 							AddShape(FGameplayDebuggerShape::MakeCircle(TargetPosition, FVector::UpVector, Radius.Radius, FColor::Red));
+							AddShape(FGameplayDebuggerShape::MakeArrow(BasePos, TargetPosition, TargetArrowHeadSize, TargetArrowThickness, FColor::Red));
 
 							const FVector::FReal TargetDistance = FMath::Max(LookArrowLength, FVector::DotProduct(WorldLookDirection, TargetPosition - BasePos));
-							AddShape(FGameplayDebuggerShape::MakeSegment(BasePos, BasePos + WorldLookDirection * TargetDistance, FColorList::LightGrey));
+							AddShape(FGameplayDebuggerShape::MakeArrow(BasePos, BasePos + WorldLookDirection * TargetDistance, LookArrowHeadSize, LookArrowThickness, FColorList::LightGrey));
 							bLookArrowDrawn = true;
 						}
 
@@ -622,12 +629,13 @@ void FGameplayDebuggerCategory_Mass::CollectData(APlayerController* OwnerPC, AAc
 						{
 							FVector TargetPosition = LookAt.GazeTargetLocation;
 							TargetPosition.Z = BasePos.Z;
-							AddShape(FGameplayDebuggerShape::MakeCircle(TargetPosition, FVector::UpVector, Radius.Radius, FColor::Turquoise));
+							AddShape(FGameplayDebuggerShape::MakeCircle(TargetPosition, FVector::UpVector, Radius.Radius, FColorList::Goldenrod));
+							AddShape(FGameplayDebuggerShape::MakeArrow(BasePos, TargetPosition, TargetArrowHeadSize, TargetArrowThickness, FColorList::Goldenrod));
 						}
 
 						if (!bLookArrowDrawn)
 						{
-							AddShape(FGameplayDebuggerShape::MakeArrow(BasePos, BasePos + WorldLookDirection * LookArrowLength, 10.0f, 1.0f, FColor::Turquoise));
+							AddShape(FGameplayDebuggerShape::MakeArrow(BasePos, BasePos + WorldLookDirection * LookArrowLength, LookArrowHeadSize, LookArrowThickness, FColor::Turquoise));
 						}
 					}
 
@@ -830,7 +838,16 @@ void FGameplayDebuggerCategory_Mass::CollectData(APlayerController* OwnerPC, AAc
 			NearEntityDescriptions.RemoveAt(MaxLabels, NearEntityDescriptions.Num() - MaxLabels);
 		}
 	}
-	
+
+	if (const UMassLookAtSubsystem* MassLookAtSubsystem = World->GetSubsystem<UMassLookAtSubsystem>())
+	{
+		AddTextLine(FString::Printf(TEXT("%d LookAt targets"), MassLookAtSubsystem->DebugGetRegisteredTargetCount()));
+	}
+
+	if (const UMassActorSubsystem* ActorSubsystem = World->GetSubsystem<UMassActorSubsystem>())
+	{
+		AddTextLine(FString::Printf(TEXT("%d actor/entity pairs"), ActorSubsystem->DebugGetRegisteredActorCount()));
+	}
 }
 
 void FGameplayDebuggerCategory_Mass::DrawData(APlayerController* OwnerPC, FGameplayDebuggerCanvasContext& CanvasContext)
diff --git a/Engine/Plugins/Runtime/MassGameplay/Source/MassActors/Private/MassActorSubsystem.cpp b/Engine/Plugins/Runtime/MassGameplay/Source/MassActors/Private/MassActorSubsystem.cpp
--- a/Engine/Plugins/Runtime/MassGameplay/Source/MassActors/Private/MassActorSubsystem.cpp
+++ b/Engine/Plugins/Runtime/MassGameplay/Source/MassActors/Private/MassActorSubsystem.cpp
@@ -129,7 +129,13 @@ FMassEntityHandle FMassActorManager::GetEntityHandleFromActor(const TObjectKey<c
 		return FMassEntityManager::InvalidEntity;
 	}
 
-	check(TObjectKey<const AActor>(GetActorFromHandle(*Entity, FMassActorFragment::EActorAccess::IncludeUnreachable)) == Actor);
+	checkCode
+	(
+		const TObjectKey<const AActor> ActorKey(GetActorFromHandle(*Entity, FMassActorFragment::EActorAccess::IncludeUnreachable));
+		checkf(ActorKey == Actor || ActorKey == TObjectKey<const AActor>()
+			, TEXT("Actor associated to the Entity handle should either match or not be set (FMassActorFragment/UMassAgentComponent not used)"));
+	);
+
 	return *Entity;
 }
 
diff --git a/Engine/Plugins/Runtime/MassGameplay/Source/MassActors/Public/MassActorSubsystem.h b/Engine/Plugins/Runtime/MassGameplay/Source/MassActors/Public/MassActorSubsystem.h
--- a/Engine/Plugins/Runtime/MassGameplay/Source/MassActors/Public/MassActorSubsystem.h
+++ b/Engine/Plugins/Runtime/MassGameplay/Source/MassActors/Public/MassActorSubsystem.h
@@ -146,6 +146,13 @@ public:
 	 */
 	void DisconnectActor(const TObjectKey<const AActor> Actor, const FMassEntityHandle Handle);
 
+	/** @return The number of actors with their associated entity registered in the manager */
+	int32 DebugGetRegisteredActorCount() const
+	{
+		UE_MT_SCOPED_READ_ACCESS(ActorHandleMapDetector);
+		return ActorHandleMap.Num();
+	}
+
 protected:
 
 	TMap<TObjectKey<const AActor>, FMassEntityHandle> ActorHandleMap;
@@ -185,6 +192,9 @@ public:
 	 */
 	inline void DisconnectActor(const TObjectKey<const AActor> Actor, const FMassEntityHandle Handle);
 
+	/** @return The number of actors with their associated entity registered in the subsystem */
+	inline int32 DebugGetRegisteredActorCount() const;
+
 protected:
 	// USubsystem BEGIN
 	virtual void Initialize(FSubsystemCollectionBase& Collection) override;
@@ -230,6 +240,10 @@ AActor* UMassActorSubsystem::GetActorFromHandle(const FMassEntityHandle Handle,
 
 void UMassActorSubsystem::DisconnectActor(const TObjectKey<const AActor> Actor, const FMassEntityHandle Handle) 
 { 
-	return ActorManager->DisconnectActor(Actor, Handle); 
+	return ActorManager->DisconnectActor(Actor, Handle);
 }
 
+int32 UMassActorSubsystem::DebugGetRegisteredActorCount() const
+{
+	return ActorManager->DebugGetRegisteredActorCount();
+}
diff --git a/Engine/Source/Runtime/MassEntity/Public/MassCommandBuffer.h b/Engine/Source/Runtime/MassEntity/Public/MassCommandBuffer.h
--- a/Engine/Source/Runtime/MassEntity/Public/MassCommandBuffer.h
+++ b/Engine/Source/Runtime/MassEntity/Public/MassCommandBuffer.h
@@ -170,6 +170,18 @@ public:
 	}
 	bool IsFlushing() const { return bIsFlushing; }
 
+	/**
+	 * Removes any pending command instances
+	 * This could be required for CommandBuffers that are queued to
+	 * flush their commands on the game thread but the EntityManager is no longer available.
+	 * In such scenario we need to cancel commands to avoid an ensure for unprocessed commands
+	 * when the buffer gets destroyed.
+	 */
+	void CancelCommands()
+	{
+		CleanUp();
+	}
+
 private:
 	friend FMassEntityManager;
 
