[MassLookAt]
- A set of Priorities can be defined in MassLookAtSettings (max 16)
- Added FMassLookAtPriority and editor customization to select from the defined priorities (used by the ST task and AnimNotify)
- Added InterpolationSpeed to LookAtFragment, RequestFragment and AnimNotify (not used by the LookAtProcessor but can be forwarded to the animation system with a translator)
- Updated MassLookAtTask to use Request system for the main target (gaze is still only supported by the ST task)
- Update Mass category in GameplayDebugger to append all registered requests and the active one to each entity status vignette

diff --git a/Engine/Plugins/AI/MassAI/MassAI.uplugin b/Engine/Plugins/AI/MassAI/MassAI.uplugin
--- a/Engine/Plugins/AI/MassAI/MassAI.uplugin
+++ b/Engine/Plugins/AI/MassAI/MassAI.uplugin
@@ -37,6 +37,10 @@
 			"Name": "MassAIBehavior",
 			"Type": "Runtime"
 		},
+		{
+			"Name": "MassAIBehaviorEditor",
+			"Type": "Editor"
+		},
 		{
 			"Name": "MassAIReplication",
 			"Type": "Runtime"
diff --git a/Engine/Plugins/AI/MassAI/Source/MassAIBehavior/MassAIBehavior.Build.cs b/Engine/Plugins/AI/MassAI/Source/MassAIBehavior/MassAIBehavior.Build.cs
--- a/Engine/Plugins/AI/MassAI/Source/MassAIBehavior/MassAIBehavior.Build.cs
+++ b/Engine/Plugins/AI/MassAI/Source/MassAIBehavior/MassAIBehavior.Build.cs
@@ -20,6 +20,7 @@ namespace UnrealBuildTool.Rules
 					"Core",
 					"CoreUObject",
 					"Engine",
+					"DeveloperSettings",
 					"GameplayTags",
 					"MassActors",
 					"MassCommon",
diff --git a/Engine/Plugins/AI/MassAI/Source/MassAIBehavior/Private/MassLookAtProcessors.cpp b/Engine/Plugins/AI/MassAI/Source/MassAIBehavior/Private/MassLookAtProcessors.cpp
--- a/Engine/Plugins/AI/MassAI/Source/MassAIBehavior/Private/MassLookAtProcessors.cpp
+++ b/Engine/Plugins/AI/MassAI/Source/MassAIBehavior/Private/MassLookAtProcessors.cpp
@@ -34,13 +34,12 @@ namespace UE::MassBehavior
 	FAutoConsoleVariableRef CVars[] =
 	{
 		FAutoConsoleVariableRef(TEXT("ai.mass.LookAt.TrajectoryLookAhead"), Tweakables::TrajectoryLookAhead,
-								TEXT("Distance (in cm) further along the look at trajectory (based on current path) to lookat while moving."), ECVF_Cheat),
+								TEXT("Distance (in cm) further along the look at trajectory (based on current path) to look at while moving."), ECVF_Cheat),
 
 		FAutoConsoleVariableRef(TEXT("ai.mass.LookAt.ConsiderObstacleTargets"), Tweakables::bConsiderObstacleTargets,
 								TEXT("Indicates if gaze targets can be selected from the entities registered in the obstacle grid."), ECVF_Cheat),
 	};
 
-	
 	// Clamps direction vector to a cone specified by the cone angle along X-axis
 	FVector ClampDirectionToXAxisCone(const FVector Direction, const float ConeAngle)
 	{
@@ -70,7 +69,7 @@ namespace UE::MassBehavior
 			return 1.0;
 		}
 
-		// @todo LookAt: mae configurable
+		// @todo LookAt: make configurable
 		const float SustainTime = GazeDuration * 0.25f;
 		const float DecayTime = GazeDuration * 0.45f;
 		
diff --git a/Engine/Plugins/AI/MassAI/Source/MassAIBehavior/Private/MassLookAtSettings.cpp b/Engine/Plugins/AI/MassAI/Source/MassAIBehavior/Private/MassLookAtSettings.cpp
new file mode 100644
--- /dev/null
+++ b/Engine/Plugins/AI/MassAI/Source/MassAIBehavior/Private/MassLookAtSettings.cpp
@@ -0,0 +1,54 @@
+// Copyright Epic Games, Inc. All Rights Reserved.
+
+#include "MassLookAtSettings.h"
+
+//----------------------------------------------------------------------//
+// UMassLookAtSettings
+//----------------------------------------------------------------------//
+FOnMassLookAtPrioritiesChanged UMassLookAtSettings::OnMassLookAtPrioritiesChanged;
+
+UMassLookAtSettings::UMassLookAtSettings(const FObjectInitializer& ObjectInitializer)
+	: Super(ObjectInitializer)
+{
+	// Setup default values, the config file will override these when the user changes them.
+	for (int32 PriorityIndex = 0; PriorityIndex < static_cast<int32>(EMassLookAtPriorities::MaxPriorities); PriorityIndex++)
+	{
+		FMassLookAtPriorityInfo& Info = Priorities[PriorityIndex];
+		Info.Priority = FMassLookAtPriority(static_cast<uint8>(PriorityIndex));
+	}
+
+	constexpr int32 LowestPriorityIndex = static_cast<int32>(EMassLookAtPriorities::LowestPriority);
+	Priorities[0].Name = FName(TEXT("Highest"));
+	Priorities[LowestPriorityIndex].Name = FName(TEXT("Lowest (default)"));
+}
+
+#if WITH_EDITOR
+void UMassLookAtSettings::PostEditChangeChainProperty(FPropertyChangedChainEvent& PropertyChangedEvent)
+{
+	Super::PostEditChangeChainProperty(PropertyChangedEvent);
+
+	const FProperty* MemberProperty = nullptr;
+	if (PropertyChangedEvent.PropertyChain.GetActiveMemberNode())
+	{
+		MemberProperty = PropertyChangedEvent.PropertyChain.GetActiveMemberNode()->GetValue();
+	}
+
+	if (MemberProperty
+		&& MemberProperty->GetFName() == GET_MEMBER_NAME_CHECKED(UMassLookAtSettings, Priorities))
+	{
+		OnMassLookAtPrioritiesChanged.Broadcast();
+	}
+}
+#endif // WITH_EDITOR
+
+void UMassLookAtSettings::GetValidPriorityInfos(TArray<FMassLookAtPriorityInfo>& OutInfos) const
+{
+	for (int32 i = 0; i < static_cast<int32>(EMassLookAtPriorities::MaxPriorities); i++)
+	{
+		const FMassLookAtPriorityInfo& Info = Priorities[i];
+		if (Info.IsValid())
+		{
+			OutInfos.Add(Info);
+		}
+	}
+}
diff --git a/Engine/Plugins/AI/MassAI/Source/MassAIBehavior/Private/MassLookAtSubsystem.cpp b/Engine/Plugins/AI/MassAI/Source/MassAIBehavior/Private/MassLookAtSubsystem.cpp
--- a/Engine/Plugins/AI/MassAI/Source/MassAIBehavior/Private/MassLookAtSubsystem.cpp
+++ b/Engine/Plugins/AI/MassAI/Source/MassAIBehavior/Private/MassLookAtSubsystem.cpp
@@ -10,76 +10,132 @@
 
 namespace UE::Mass::LookAt::Private
 {
+TArray<FMassEntityHandle> DebugRequests;
 
 static FAutoConsoleCommandWithWorldArgsAndOutputDevice CmdSendLookAtPlayerRequestToAll(
 	TEXT("ai.debug.mass.SendLookAtPlayerRequestToAll"),
-	TEXT(""),
+	TEXT("Creates, or removed, LookAt requests toward the player for all mass entities with a LookAt fragment"
+		"(optional)<0|1> to indicate if the requests must be created (1, default) or deleted (0)."
+		"(optional)<int> to indicate the priority of the request where a lower value represents a higher priority (default is 5)"),
 	FConsoleCommandWithWorldArgsAndOutputDeviceDelegate::CreateLambda([](const TArray<FString>& Args, const UWorld* World, FOutputDevice& OutputDevice)
 		{
+			bool bCreateRequest = true;
+
+			uint8 PriorityLevel = static_cast<uint8>(EMassLookAtPriorities::LowestPriority);
+
+			if (Args.Num() > 2)
+			{
+				OutputDevice.Log(ELogVerbosity::Error, TEXT("Command failed: invalid number of arguments"));
+				return;
+			}
+
+			if (Args.Num() > 0)
+			{
+				if (!LexTryParseString(bCreateRequest, *Args[0]))
+				{
+					OutputDevice.Log(ELogVerbosity::Error, TEXT("Unable to parse the first argument: expecting 0|1 or true|false"));
+					return;
+				}
+
+				if (Args.Num() > 1)
+				{
+					if (!LexTryParseString(PriorityLevel, *Args[1]))
+					{
+						OutputDevice.Log(ELogVerbosity::Error, TEXT("Unable to parse the second argument: expecting an [0-255] integer to represent the priority"));
+						return;
+					}
+
+					// Clamp priority in valid range (highest value being the lowest priority)
+					if (constexpr uint8 LastPriority = static_cast<uint8>(EMassLookAtPriorities::LowestPriority); PriorityLevel > LastPriority)
+					{
+						OutputDevice.Log(ELogVerbosity::Warning, FString::Printf(TEXT("Clamped priority level to the lowest priority %d"), LastPriority));
+						PriorityLevel = FMath::Clamp<uint8>(PriorityLevel, 0, LastPriority);
+					}
+				}
+			}
+
 			UMassActorSubsystem* ActorSubsystem = World->GetSubsystem<UMassActorSubsystem>();
 			if (ActorSubsystem == nullptr)
 			{
-				OutputDevice.Log(ELogVerbosity::Error, TEXT("%hs failure: unable to find MassActorSubsystem"));
+				OutputDevice.Log(ELogVerbosity::Error, TEXT("Command failed: unable to find MassActorSubsystem"));
 				return;
 			}
 
 			const APlayerController* PlayerController = World->GetFirstPlayerController();
 			if (PlayerController == nullptr)
 			{
-				OutputDevice.Log(ELogVerbosity::Error, TEXT("%hs failure: unable to find the player controller"));
+				OutputDevice.Log(ELogVerbosity::Error, TEXT("Command failed: unable to find the player controller"));
 				return;
 			}
 
 			const APawn* PlayerPawn = PlayerController->GetPawn();
 			if (PlayerPawn == nullptr)
 			{
-				OutputDevice.Log(ELogVerbosity::Error, TEXT("%hs failure: unable to find the player pawn"));
+				OutputDevice.Log(ELogVerbosity::Error, TEXT("Command failed: unable to find the player pawn"));
 				return;
 			}
 
 			UMassEntitySubsystem* EntitySubsystem = World->GetSubsystem<UMassEntitySubsystem>();
 			if (ActorSubsystem == nullptr)
 			{
-				OutputDevice.Log(ELogVerbosity::Error, TEXT("%hs failure: unable to find UMassEntitySubsystem"));
+				OutputDevice.Log(ELogVerbosity::Error, TEXT("Command failed: unable to find UMassEntitySubsystem"));
 				return;
 			}
 
 			const FMassEntityHandle PlayerEntity = ActorSubsystem->GetEntityHandleFromActor(PlayerPawn);
 			if (!PlayerEntity.IsSet())
 			{
-				OutputDevice.Log(ELogVerbosity::Error, TEXT("%hs failure: unable to find a MassEntity associated to the player"));
+				OutputDevice.Log(ELogVerbosity::Error, TEXT("Command failed: unable to find a MassEntity associated to the player"));
 				return;
 			}
 
-			// Create requests for all entities with a MassLookAtFragment
 			TSharedPtr<FMassCommandBuffer> CommandBuffer = MakeShareable(new FMassCommandBuffer());
-			CommandBuffer->PushCommand<FMassDeferredCreateCommand>([PlayerEntity](FMassEntityManager& InEntityManager)
-				{
-					FMassEntityQuery EntityQuery(InEntityManager.AsShared());
-					EntityQuery.AddRequirement<FMassLookAtFragment>(EMassFragmentAccess::ReadOnly);
-
-					// @todo LookAt: replace this part by FMassEntityQuery::GetMatchingEntityHandles() once available
-					TArray<FMassEntityHandle> Viewers;
-					FMassExecutionContext Context(InEntityManager, 0.0f);
-					EntityQuery.ForEachEntityChunk(Context, [&Viewers](const FMassExecutionContext& Context)
+			if (bCreateRequest)
+			{
+				// Create requests for all entities with a MassLookAtFragment
+				CommandBuffer->PushCommand<FMassDeferredCreateCommand>([PlayerEntity, PriorityLevel](FMassEntityManager& InEntityManager)
+					{
+						FMassEntityQuery EntityQuery(InEntityManager.AsShared());
+						EntityQuery.AddRequirement<FMassLookAtFragment>(EMassFragmentAccess::ReadOnly);
+
+						// @todo LookAt: replace this part by FMassEntityQuery::GetMatchingEntityHandles() once available
+						TArray<FMassEntityHandle> Viewers;
+						FMassExecutionContext Context(InEntityManager, 0.0f);
+						EntityQuery.ForEachEntityChunk(Context, [&Viewers](const FMassExecutionContext& Context)
+							{
+								Viewers.Append(Context.GetEntities());
+							});
+
+						if (Viewers.Num() > 0)
 						{
-							Viewers.Append(Context.GetEntities());
-						});
-
-					FMassArchetypeCompositionDescriptor Composition;
-					Composition.Fragments.Add(*FMassLookAtRequestFragment::StaticStruct());
-
-					TArray<FMassEntityHandle> Requests;
-					const TSharedRef<FMassEntityManager::FEntityCreationContext> CreationContext =
-						InEntityManager.BatchCreateEntities(InEntityManager.CreateArchetype(Composition), Viewers.Num(), Requests);
-
-					check(Viewers.Num() == Requests.Num());
-					for (int Index = 0; Index < Requests.Num(); ++Index)
+							FMassArchetypeCompositionDescriptor Composition;
+							Composition.Fragments.Add(*FMassLookAtRequestFragment::StaticStruct());
+
+							TArray<FMassEntityHandle> Requests;
+							const TSharedRef<FMassEntityManager::FEntityCreationContext> CreationContext =
+								InEntityManager.BatchCreateEntities(InEntityManager.CreateArchetype(Composition), Viewers.Num(), Requests);
+
+							DebugRequests.Append(Requests);
+							check(Viewers.Num() == Requests.Num());
+							const FMassLookAtPriority Priority(PriorityLevel);
+
+							for (int Index = 0; Index < Requests.Num(); ++Index)
+							{
+								InEntityManager.SetEntityFragmentValues(Requests[Index]
+									, { FInstancedStruct::Make(FMassLookAtRequestFragment{Viewers[Index], Priority, EMassLookAtMode::LookAtEntity, PlayerEntity}) });
+							}
+						}
+					});
+			}
+			else
+			{
+				// Delete all entities created for debug requests
+				CommandBuffer->PushCommand<FMassDeferredCreateCommand>([](FMassEntityManager& InEntityManager)
 					{
-						InEntityManager.SetEntityFragmentValues(Requests[Index]
-							, { FInstancedStruct::Make(FMassLookAtRequestFragment{Viewers[Index], PlayerEntity}) });
-					}
-				});
+						InEntityManager.BatchDestroyEntities(DebugRequests);
+						DebugRequests.Reset();
+					});
+			}
 
 			ExecuteOnGameThread(UE_SOURCE_LOCATION, [WeakEntityManager = EntitySubsystem->GetMutableEntityManager().AsWeak(), CommandBuffer]()
 				{
@@ -120,13 +176,14 @@ void UMassLookAtSubsystem::RegisterRequests(const FMassExecutionContext& InConte
 	{
 		UE_MT_SCOPED_WRITE_ACCESS(RequestsAccessDetector);
 
-		ActiveRequests.Reserve(ActiveRequests.Num() + (InRequests.Num() - ActiveRequestsFreeList.Num()));
+		RegisteredRequests.Reserve(RegisteredRequests.Num() + (InRequests.Num() - ActiveRequestsFreeList.Num()));
 
 		for (const FRequest& NewRequest : InRequests)
 		{
-			if (!NewRequest.Parameters.TargetEntity.IsSet())
+			if (NewRequest.Parameters.LookAtMode == EMassLookAtMode::LookAtEntity
+				&& !NewRequest.Parameters.TargetEntity.IsSet())
 			{
-				UE_VLOG_UELOG(this, LogMassBehavior, Error, TEXT("Ignoring LookAt request: invalid target entity"));
+				UE_VLOG_UELOG(this, LogMassBehavior, Error, TEXT("Ignoring LookAtEntity request: invalid target entity"));
 				continue;
 			}
 
@@ -134,11 +191,11 @@ void UMassLookAtSubsystem::RegisterRequests(const FMassExecutionContext& InConte
 			if (ActiveRequestsFreeList.Num())
 			{
 				NewRequestIndex = ActiveRequestsFreeList.Pop();
-				ActiveRequests[NewRequestIndex] = NewRequest;
+				RegisteredRequests[NewRequestIndex] = NewRequest;
 			}
 			else
 			{
-				NewRequestIndex = ActiveRequests.Add(NewRequest);
+				NewRequestIndex = RegisteredRequests.Add(NewRequest);
 			}
 
 			int32 ViewerDataIndex = PerViewerRequests.IndexOfByPredicate([ViewerEntity = NewRequest.Parameters.ViewerEntity](const FViewerRequest& ViewerRequest)
@@ -173,15 +230,15 @@ void UMassLookAtSubsystem::UnregisterRequests(const FMassExecutionContext& InCon
 			int32 InvalidatedIndex = INDEX_NONE;
 			FMassEntityHandle ViewerEntity;
 
-			for (int32 ActiveRequestIndex = 0; ActiveRequestIndex < ActiveRequests.Num(); ++ActiveRequestIndex)
+			for (int32 ActiveRequestIndex = 0; ActiveRequestIndex < RegisteredRequests.Num(); ++ActiveRequestIndex)
 			{
-				if (ActiveRequests[ActiveRequestIndex].RequestHandle == RemovedRequest)
+				if (RegisteredRequests[ActiveRequestIndex].RequestHandle == RemovedRequest)
 				{
 					InvalidatedIndex = ActiveRequestIndex;
 
 					// Invalidate entry and add to the free list
-					ViewerEntity = ActiveRequests[ActiveRequestIndex].Parameters.ViewerEntity;
-					ActiveRequests[ActiveRequestIndex].RequestHandle = FMassEntityHandle();
+					ViewerEntity = RegisteredRequests[ActiveRequestIndex].Parameters.ViewerEntity;
+					RegisteredRequests[ActiveRequestIndex].RequestHandle = FMassEntityHandle();
 					ActiveRequestsFreeList.Push(ActiveRequestIndex);
 					break;
 				}
@@ -209,7 +266,7 @@ void UMassLookAtSubsystem::UnregisterRequests(const FMassExecutionContext& InCon
 void UMassLookAtSubsystem::UpdateLookAts(const FMassExecutionContext& Context, const TConstArrayView<int32> DirtyViewers)
 {
 	TArray<TPair<FMassEntityHandle, FMassLookAtRequestFragment>> Updated;
-	for (int32 DirtyViewerIndex : DirtyViewers)
+	for (const int32 DirtyViewerIndex : DirtyViewers)
 	{
 		check(PerViewerRequests.IsValidIndex(DirtyViewerIndex));
 
@@ -217,31 +274,31 @@ void UMassLookAtSubsystem::UpdateLookAts(const FMassExecutionContext& Context, c
 
 		for (const int32 RequestIndex : PerViewerRequests[DirtyViewerIndex].RequestIndices)
 		{
-			check(ActiveRequests.IsValidIndex(RequestIndex));
-
-			SelectedRequestIndex = RequestIndex;
-			break;
+			check(RegisteredRequests.IsValidIndex(RequestIndex));
+			RegisteredRequests[RequestIndex].bActive = false;
 
-			// @todo LookAt: Add and handle priorities
-			//if (SelectedRequestIndex == INDEX_NONE)
-			//{
-			//	SelectedRequestIndex = RequestIndex;
-			//	continue;
-			//}
+			if (SelectedRequestIndex == INDEX_NONE)
+			{
+				SelectedRequestIndex = RequestIndex;
+				continue;
+			}
 
+			// Higher priority is represented by the lowest value
+			if (RegisteredRequests[RequestIndex].Parameters.Priority.Get()
+				< RegisteredRequests[SelectedRequestIndex].Parameters.Priority.Get())
+			{
+				SelectedRequestIndex = RequestIndex;
+			}
+		}
 
-			//if (ActiveRequests[RequestIndex].Parameters.Priority
-			//	> ActiveRequests[SelectedRequestIndex].Parameters.Priority)
-			//{
-			//	SelectedRequestIndex = RequestIndex;
-			//}
+		FMassLookAtRequestFragment RequestFragment;
+		if (SelectedRequestIndex != INDEX_NONE)
+		{
+			RegisteredRequests[SelectedRequestIndex].bActive = true;
+			RequestFragment = RegisteredRequests[SelectedRequestIndex].Parameters;
 		}
 
-		Updated.Add(
-			{
-				PerViewerRequests[DirtyViewerIndex].Viewer
-				, SelectedRequestIndex != INDEX_NONE ? ActiveRequests[SelectedRequestIndex].Parameters : FMassLookAtRequestFragment{}
-			});
+		Updated.Add({ PerViewerRequests[DirtyViewerIndex].Viewer, RequestFragment });
 	}
 
 	Context.Defer().PushCommand<FMassDeferredSetCommand>([Updated = MoveTemp(Updated)](const FMassEntityManager& Manager)
@@ -250,22 +307,40 @@ void UMassLookAtSubsystem::UpdateLookAts(const FMassExecutionContext& Context, c
 			{
 				if (FMassLookAtFragment* LookAtFragment = Manager.GetFragmentDataPtr<FMassLookAtFragment>(Pair.Key))
 				{
-					if (LookAtFragment->TrackedEntity == Pair.Value.TargetEntity)
-					{
-						continue;
-					}
-
 					// Only update main LookAt information and don't modify gaze related one
-					if (Pair.Value.TargetEntity.IsSet())
-					{
-						LookAtFragment->LookAtMode = EMassLookAtMode::LookAtEntity;
-						LookAtFragment->TrackedEntity = Pair.Value.TargetEntity;
-					}
-					else
-					{
-						LookAtFragment->LookAtMode = EMassLookAtMode::LookForward;
-					}
+					const FMassLookAtRequestFragment& Request = Pair.Value;
+					LookAtFragment->InterpolationSpeed = Request.InterpolationSpeed;
+					LookAtFragment->LookAtMode = Request.LookAtMode;
+					LookAtFragment->TrackedEntity = Request.TargetEntity;
 				}
 			}
 		});
 }
+
+#if WITH_MASSGAMEPLAY_DEBUG
+FString UMassLookAtSubsystem::DebugGetRequestsString(const FMassEntityHandle InEntity) const
+{
+	UE_MT_SCOPED_READ_ACCESS(RequestsAccessDetector);
+
+	const int32 ViewerDataIndex = PerViewerRequests.IndexOfByPredicate([InEntity](const FViewerRequest& ViewerRequest)
+		{
+			return InEntity == ViewerRequest.Viewer;
+		});
+
+	TStringBuilder<256> Builder;
+	if (ViewerDataIndex != INDEX_NONE)
+	{
+		for (const int32 RequestIndex : PerViewerRequests[ViewerDataIndex].RequestIndices)
+		{
+			if (Builder.Len() > 0)
+			{
+				Builder << TEXT("\n");
+			}
+
+			Builder << (RegisteredRequests[RequestIndex].bActive ? TEXT(">") : TEXT("   "));
+			Builder << LexToString(RegisteredRequests[RequestIndex].Parameters);
+		}
+	}
+	return Builder.ToString();
+}
+#endif // WITH_MASSGAMEPLAY_DEBUG
diff --git a/Engine/Plugins/AI/MassAI/Source/MassAIBehavior/Private/Tasks/MassLookAtTask.cpp b/Engine/Plugins/AI/MassAI/Source/MassAIBehavior/Private/Tasks/MassLookAtTask.cpp
--- a/Engine/Plugins/AI/MassAI/Source/MassAIBehavior/Private/Tasks/MassLookAtTask.cpp
+++ b/Engine/Plugins/AI/MassAI/Source/MassAIBehavior/Private/Tasks/MassLookAtTask.cpp
@@ -19,7 +19,7 @@ bool FMassLookAtTask::Link(FStateTreeLinker& Linker)
 void FMassLookAtTask::GetDependencies(UE::MassBehavior::FStateTreeDependencyBuilder& Builder) const
 {
 	Builder.AddReadWrite<FMassLookAtFragment>();
-	// @todo this is only needed when InstanceData.Duration, but there's no way to tell at this point.
+	// @todo LookAt: this is only needed when InstanceData.Duration, but there's no way to tell at this point.
 	Builder.AddReadWrite<UMassSignalSubsystem>();
 }
 
@@ -32,22 +32,39 @@ EStateTreeRunStatus FMassLookAtTask::EnterState(FStateTreeExecutionContext& Cont
 	const FMassStateTreeExecutionContext& MassStateTreeContext = static_cast<FMassStateTreeExecutionContext&>(Context);
 	FMassLookAtFragment& LookAtFragment = MassStateTreeContext.GetExternalData(LookAtHandle);
 
-	LookAtFragment.Reset();
-	LookAtFragment.LookAtMode = LookAtMode;
-	
-	if (LookAtMode == EMassLookAtMode::LookAtEntity)
-	{
-		if (!InstanceData.TargetEntity.IsSet())
-		{
-			LookAtFragment.LookAtMode = EMassLookAtMode::LookForward;
-			MASSBEHAVIOR_LOG(Error, TEXT("Failed LookAt: invalid target entity"));
-		}
-		else
-		{
-			LookAtFragment.LookAtMode = EMassLookAtMode::LookAtEntity;
-			LookAtFragment.TrackedEntity = InstanceData.TargetEntity;
-		}
-	}
+	MassStateTreeContext.GetMassEntityExecutionContext().Defer()
+		.PushCommand<FMassDeferredCreateCommand>(
+			[&Context, SelfEntity = MassStateTreeContext.GetEntity(), &InstanceData, Mode = LookAtMode, Priority = Priority, InterpolationSpeed = InterpolationSpeed]
+			(FMassEntityManager& InEntityManager)
+			{
+				FMassArchetypeCompositionDescriptor Composition;
+				Composition.Fragments.Add(*FMassLookAtRequestFragment::StaticStruct());
+
+				constexpr int32 NumRequestsToCreate = 1;
+				TArray<FMassEntityHandle> Requests;
+				const TSharedRef<FMassEntityManager::FEntityCreationContext> CreationContext =
+					InEntityManager.BatchCreateEntities(InEntityManager.CreateArchetype(Composition), NumRequestsToCreate, Requests);
+
+				if (Requests.Num() == NumRequestsToCreate)
+				{
+					InstanceData.Request = Requests.Top();
+					FMassLookAtRequestFragment RequestFragment{ SelfEntity, Priority, Mode, InstanceData.TargetEntity };
+					RequestFragment.InterpolationSpeed = InterpolationSpeed;
+
+					// When using 'LookAtEntity' we validate the target entity to use it,
+					// or we use the default 'LookForward' as fallback.
+					if (Mode == EMassLookAtMode::LookAtEntity)
+					{
+						if (!InstanceData.TargetEntity.IsSet())
+						{
+							MASSBEHAVIOR_LOG(Error, TEXT("Failed LookAt: invalid target entity"));
+							RequestFragment = { SelfEntity, Priority, EMassLookAtMode::LookForward };
+						}
+					}
+
+					InEntityManager.SetEntityFragmentValues(InstanceData.Request, { FInstancedStruct::Make(RequestFragment) });
+				}
+			});
 
 	LookAtFragment.RandomGazeMode = RandomGazeMode;
 	LookAtFragment.RandomGazeYawVariation = RandomGazeYawVariation;
@@ -68,9 +85,23 @@ EStateTreeRunStatus FMassLookAtTask::EnterState(FStateTreeExecutionContext& Cont
 void FMassLookAtTask::ExitState(FStateTreeExecutionContext& Context, const FStateTreeTransitionResult& Transition) const
 {
 	const FMassStateTreeExecutionContext& MassStateTreeContext = static_cast<FMassStateTreeExecutionContext&>(Context);
+
+	// Reset gaze directly since it is used only by the task and not by the requests
 	FMassLookAtFragment& LookAtFragment = MassStateTreeContext.GetExternalData(LookAtHandle);
-	
-	LookAtFragment.Reset();
+	LookAtFragment.ResetGaze();
+
+	// Destroy any associated request for the main LookAt
+	FInstanceDataType& InstanceData = Context.GetInstanceData(*this);
+	if (InstanceData.Request.IsSet())
+	{
+		MassStateTreeContext.GetMassEntityExecutionContext().Defer()
+			.PushCommand<FMassDeferredCreateCommand>([Request = InstanceData.Request](FMassEntityManager& InEntityManager)
+			{
+				InEntityManager.DestroyEntity(Request);
+			});
+
+		InstanceData.Request.Reset();
+	}
 }
 
 EStateTreeRunStatus FMassLookAtTask::Tick(FStateTreeExecutionContext& Context, const float DeltaTime) const
diff --git a/Engine/Plugins/AI/MassAI/Source/MassAIBehavior/Public/MassAIBehaviorTypes.h b/Engine/Plugins/AI/MassAI/Source/MassAIBehavior/Public/MassAIBehaviorTypes.h
--- a/Engine/Plugins/AI/MassAI/Source/MassAIBehavior/Public/MassAIBehaviorTypes.h
+++ b/Engine/Plugins/AI/MassAI/Source/MassAIBehavior/Public/MassAIBehaviorTypes.h
@@ -28,10 +28,10 @@ namespace UE::Mass::ProcessorGroupNames
  *		MASSBEHAVIOR_CLOG(bDisplayDebug, Log, TEXT("Starting action: %s"), *StaticEnum<ESomeActionEnum>()->GetValueAsString(SomeActionEnumValue));
  */
 #if WITH_MASSGAMEPLAY_DEBUG
-#define MASSBEHAVIOR_LOG(Verbosity, Format, ...) UE_VLOG_UELOG(static_cast<FMassStateTreeExecutionContext&>(Context).GetOwner(), LogMassBehavior, Verbosity, \
-	TEXT("Entity [%s][%s] ") Format, *static_cast<FMassStateTreeExecutionContext&>(Context).GetEntity().DebugGetDescription(), *StaticStruct()->GetName(), ##__VA_ARGS__)
-#define MASSBEHAVIOR_CLOG(Condition, Verbosity, Format, ...) UE_CVLOG_UELOG((Condition), static_cast<FMassStateTreeExecutionContext&>(Context).GetOwner(), LogMassBehavior, Verbosity, \
-	TEXT("Entity [%s][%s] ") Format, *static_cast<FMassStateTreeExecutionContext&>(Context).GetEntity().DebugGetDescription(), *StaticStruct()->GetName(), ##__VA_ARGS__)
+#define MASSBEHAVIOR_LOG(Verbosity, Format, ...) UE_VLOG_UELOG(static_cast<const FMassStateTreeExecutionContext&>(Context).GetOwner(), LogMassBehavior, Verbosity, \
+	TEXT("Entity [%s][%s] ") Format, *static_cast<const FMassStateTreeExecutionContext&>(Context).GetEntity().DebugGetDescription(), *StaticStruct()->GetName(), ##__VA_ARGS__)
+#define MASSBEHAVIOR_CLOG(Condition, Verbosity, Format, ...) UE_CVLOG_UELOG((Condition), static_cast<const FMassStateTreeExecutionContext&>(Context).GetOwner(), LogMassBehavior, Verbosity, \
+	TEXT("Entity [%s][%s] ") Format, *static_cast<const FMassStateTreeExecutionContext&>(Context).GetEntity().DebugGetDescription(), *StaticStruct()->GetName(), ##__VA_ARGS__)
 #else
 #define MASSBEHAVIOR_LOG(Verbosity, Format, ...)
 #define MASSBEHAVIOR_CLOG(Condition, Verbosity, Format, ...)
diff --git a/Engine/Plugins/AI/MassAI/Source/MassAIBehavior/Public/MassLookAtFragments.h b/Engine/Plugins/AI/MassAI/Source/MassAIBehavior/Public/MassLookAtFragments.h
--- a/Engine/Plugins/AI/MassAI/Source/MassAIBehavior/Public/MassLookAtFragments.h
+++ b/Engine/Plugins/AI/MassAI/Source/MassAIBehavior/Public/MassLookAtFragments.h
@@ -46,25 +46,33 @@ struct MASSAIBEHAVIOR_API FMassLookAtFragment : public FMassFragment
 		: bRandomGazeEntities(false)
 	{
 	}
-	
-	void Reset()
+
+	/** Resets gaze related properties in the fragment to their default values */
+	void ResetGaze()
 	{
-		MainTargetLocation = FVector::ZeroVector;
 		GazeTargetLocation = FVector::ZeroVector;
-		Direction = FVector::ForwardVector;
 		GazeDirection = FVector::ForwardVector;
-		TrackedEntity.Reset();
 		GazeTrackedEntity.Reset();
 		GazeStartTime = 0.0f;
 		GazeDuration = 0.0f;
-		LastSeenActionID = 0;
-		LookAtMode = EMassLookAtMode::LookForward;
 		RandomGazeMode = EMassLookAtGazeMode::None;
 		RandomGazeYawVariation = 0;
 		RandomGazePitchVariation = 0;
 		bRandomGazeEntities = false;
 	}
 
+	/** Resets all properties in the fragment to their default values */
+	void ResetAll()
+	{
+		ResetGaze();
+
+		MainTargetLocation = FVector::ZeroVector;
+		Direction = FVector::ForwardVector;
+		TrackedEntity.Reset();
+		LastSeenActionID = 0;
+		LookAtMode = EMassLookAtMode::LookForward;
+	}
+
 	/** Current look at main target location. */
 	UPROPERTY(Transient)
 	FVector MainTargetLocation = FVector::ZeroVector;
@@ -105,6 +113,10 @@ struct MASSAIBEHAVIOR_API FMassLookAtFragment : public FMassFragment
 	UPROPERTY(Transient)
 	EMassLookAtMode LookAtMode = EMassLookAtMode::LookForward;
 
+	/** Primary look at interpolation speed (not used by the LookAt processor but can be forwarded to the animation system). */
+	UPROPERTY(Transient)
+	EMassLookAtInterpolationSpeed InterpolationSpeed = EMassLookAtInterpolationSpeed::Regular;
+
 	/** Gaze look at mode. */
 	UPROPERTY(Transient)
 	EMassLookAtGazeMode RandomGazeMode = EMassLookAtGazeMode::None;
@@ -140,17 +152,48 @@ struct MASSAIBEHAVIOR_API FMassLookAtRequestFragment : public FMassFragment
 	GENERATED_BODY()
 
 	FMassLookAtRequestFragment() = default;
-	FMassLookAtRequestFragment(const FMassEntityHandle ViewerEntity, const FMassEntityHandle TargetEntity)
+	FMassLookAtRequestFragment(const FMassEntityHandle ViewerEntity, const FMassLookAtPriority Priority, const EMassLookAtMode Mode)
 		: ViewerEntity(ViewerEntity)
-		, TargetEntity(TargetEntity)
+		, Priority(Priority)
+		, LookAtMode(Mode)
 	{
 	}
 
+	FMassLookAtRequestFragment(const FMassEntityHandle ViewerEntity, const FMassLookAtPriority Priority, const EMassLookAtMode Mode, const FMassEntityHandle Target)
+		: FMassLookAtRequestFragment(ViewerEntity, Priority, Mode)
+	{
+		TargetEntity = Target;
+	}
+
 	UPROPERTY(Transient)
 	FMassEntityHandle ViewerEntity;
 
 	UPROPERTY(Transient)
 	FMassEntityHandle TargetEntity;
+
+	UPROPERTY(Transient)
+	FMassLookAtPriority Priority;
+
+	UPROPERTY(Transient)
+	EMassLookAtMode LookAtMode = EMassLookAtMode::LookForward;
+
+	UPROPERTY(Transient)
+	EMassLookAtInterpolationSpeed InterpolationSpeed = EMassLookAtInterpolationSpeed::Regular;
+
+	friend FString LexToString(const FMassLookAtRequestFragment& Fragment)
+	{
+		if (Fragment.LookAtMode == EMassLookAtMode::LookAtEntity)
+		{
+			return FString::Printf(TEXT("Priority = %d Mode = %s Target = [%s]")
+			, Fragment.Priority.Get()
+			, *UEnum::GetValueAsString(Fragment.LookAtMode)
+			, *LexToString(Fragment.TargetEntity));
+		}
+
+		return FString::Printf(TEXT("Priority=%d Mode=%s")
+			, Fragment.Priority.Get()
+			, *UEnum::GetValueAsString(Fragment.LookAtMode));
+	}
 };
 
 /** Cell location for LookAt targets */
diff --git a/Engine/Plugins/AI/MassAI/Source/MassAIBehavior/Public/MassLookAtSettings.h b/Engine/Plugins/AI/MassAI/Source/MassAIBehavior/Public/MassLookAtSettings.h
new file mode 100644
--- /dev/null
+++ b/Engine/Plugins/AI/MassAI/Source/MassAIBehavior/Public/MassLookAtSettings.h
@@ -0,0 +1,33 @@
+// Copyright Epic Games, Inc. All Rights Reserved.
+#pragma once
+
+#include "Delegates/DelegateCombinations.h"
+#include "Engine/DeveloperSettings.h"
+#include "MassLookAtTypes.h"
+#include "MassLookAtSettings.generated.h"
+
+DECLARE_MULTICAST_DELEGATE(FOnMassLookAtPrioritiesChanged);
+
+/**
+ * Implements the settings for the MassLookAt module.
+ */
+UCLASS(MinimalAPI, config = Plugins, defaultconfig, DisplayName = "Mass LookAt")
+class UMassLookAtSettings : public UDeveloperSettings
+{
+	GENERATED_BODY()
+public:
+	explicit UMassLookAtSettings(const FObjectInitializer& ObjectInitializer = FObjectInitializer::Get());
+
+#if WITH_EDITOR
+	virtual void PostEditChangeChainProperty(FPropertyChangedChainEvent& PropertyChangedEvent) override;
+#endif
+
+	MASSAIBEHAVIOR_API void GetValidPriorityInfos(TArray<FMassLookAtPriorityInfo>& OutInfos) const;
+
+	static MASSAIBEHAVIOR_API FOnMassLookAtPrioritiesChanged OnMassLookAtPrioritiesChanged;
+
+protected:
+
+	UPROPERTY(EditAnywhere, config, Category = LookAt)
+	FMassLookAtPriorityInfo Priorities[static_cast<int32>(EMassLookAtPriorities::MaxPriorities)];
+};
\ No newline at end of file
diff --git a/Engine/Plugins/AI/MassAI/Source/MassAIBehavior/Public/MassLookAtSubsystem.h b/Engine/Plugins/AI/MassAI/Source/MassAIBehavior/Public/MassLookAtSubsystem.h
--- a/Engine/Plugins/AI/MassAI/Source/MassAIBehavior/Public/MassLookAtSubsystem.h
+++ b/Engine/Plugins/AI/MassAI/Source/MassAIBehavior/Public/MassLookAtSubsystem.h
@@ -25,6 +25,7 @@ public:
 	{
 		FMassEntityHandle RequestHandle;
 		FMassLookAtRequestFragment Parameters;
+		bool bActive = false;
 	};
 
 	/**
@@ -113,6 +114,11 @@ public:
 		return TargetGrid.GetItems().Num();
 	}
 
+#if WITH_MASSGAMEPLAY_DEBUG
+	/** @return String detailing all requests registered for a given entity. */
+	FString DebugGetRequestsString(FMassEntityHandle InEntity) const;
+#endif
+
 protected:
 
 	virtual void Initialize(FSubsystemCollectionBase& Collection) override;
@@ -133,8 +139,8 @@ private:
 		TArray<int32> RequestIndices;
 	};
 
-	/** List of all currently registered (and active) requests as registered by external systems */
-	TArray<FRequest> ActiveRequests;
+	/** List of all currently registered requests as registered by external systems */
+	TArray<FRequest> RegisteredRequests;
 
 	/** List of available indices in ActiveRequests (to preserve stable indices when unregistering requests) */
 	TArray<int32> ActiveRequestsFreeList;
diff --git a/Engine/Plugins/AI/MassAI/Source/MassAIBehavior/Public/MassLookAtTypes.h b/Engine/Plugins/AI/MassAI/Source/MassAIBehavior/Public/MassLookAtTypes.h
--- a/Engine/Plugins/AI/MassAI/Source/MassAIBehavior/Public/MassLookAtTypes.h
+++ b/Engine/Plugins/AI/MassAI/Source/MassAIBehavior/Public/MassLookAtTypes.h
@@ -4,6 +4,8 @@
 
 #include "HierarchicalHashGrid2D.h"
 #include "MassEntityHandle.h"
+#include "UObject/ObjectMacros.h"
+#include "MassLookAtTypes.generated.h"
 
 namespace UE::Mass::LookAt
 {
@@ -12,3 +14,100 @@ namespace UE::Mass::LookAt
 
 	using FTargetHashGrid2D = THierarchicalHashGrid2D<HashGridLevelsOfHierarchy, HashGridRatioBetweenLevels, FMassEntityHandle>;
 }
+
+/**
+ * Enum representing the different interpolation speeds that can be used when assigning new LookAt targets
+ */
+UENUM()
+enum class EMassLookAtInterpolationSpeed : uint8
+{
+	Instant,
+	Fast,
+	Regular,
+	Slow
+};
+
+/**
+ * Enum used to define the number of configurable priorities exposed by the MassLookAtSettings
+ */
+UENUM()
+enum class EMassLookAtPriorities : uint8
+{
+	MaxPriorities = 16,
+	MaxPriorityIndex = MaxPriorities - 1,
+	LowestPriority = MaxPriorityIndex
+};
+
+/**
+ * Struct used as a priority selector exposed to the Editor
+ */
+USTRUCT()
+struct FMassLookAtPriority
+{
+	GENERATED_BODY()
+
+	FMassLookAtPriority() = default;
+	
+	explicit FMassLookAtPriority(const uint8 InBit)
+		: Value(InBit)
+	{
+		check(InBit <= static_cast<uint8>(EMassLookAtPriorities::MaxPriorities));
+	}
+
+	void Set(const uint8 InBit)
+	{
+		check(InBit <= static_cast<uint8>(EMassLookAtPriorities::MaxPriorities));
+		Value = InBit;
+	}
+
+	uint8 Get() const
+	{
+		return Value;
+	}
+
+	void Reset()
+	{
+		Value = NoneValue;
+	}
+
+	bool IsValid() const
+	{
+		return Value != NoneValue;
+	}
+
+	bool operator==(const FMassLookAtPriority& RHS) const
+	{
+		return Value == RHS.Value;
+	}
+
+	bool operator!=(const FMassLookAtPriority& RHS) const
+	{
+		return Value != RHS.Value;
+	}
+
+private:
+	static constexpr uint8 NoneValue = 0xFF;
+
+	UPROPERTY(Category = LookAt, EditAnywhere)
+	uint8 Value = NoneValue;
+};
+
+/**
+ * Struct used to represent configurable priorities in MassLookAtSettings
+ */
+USTRUCT()
+struct FMassLookAtPriorityInfo
+{
+	GENERATED_BODY()
+
+	bool IsValid() const
+	{
+		return !Name.IsNone() && Priority.IsValid();
+	}
+
+	UPROPERTY(Category = LookAt, EditAnywhere)
+	FName Name;
+
+	UPROPERTY(Category = LookAt, EditAnywhere)
+	FMassLookAtPriority Priority;
+};
\ No newline at end of file
diff --git a/Engine/Plugins/AI/MassAI/Source/MassAIBehavior/Public/Tasks/MassLookAtTask.h b/Engine/Plugins/AI/MassAI/Source/MassAIBehavior/Public/Tasks/MassLookAtTask.h
--- a/Engine/Plugins/AI/MassAI/Source/MassAIBehavior/Public/Tasks/MassLookAtTask.h
+++ b/Engine/Plugins/AI/MassAI/Source/MassAIBehavior/Public/Tasks/MassLookAtTask.h
@@ -32,6 +32,10 @@ struct MASSAIBEHAVIOR_API FMassLookAtTaskInstanceData
 	/** Accumulated time used to stop task if duration is set */
 	UPROPERTY()
 	float Time = 0.f;
+
+	/** Request created for the main LookAt */
+	UPROPERTY()
+	FMassEntityHandle Request;
 };
 
 USTRUCT(meta = (DisplayName = "Mass LookAt Task"))
@@ -52,9 +56,17 @@ protected:
 	TStateTreeExternalDataHandle<UMassSignalSubsystem> MassSignalSubsystemHandle;
 	TStateTreeExternalDataHandle<FMassLookAtFragment> LookAtHandle;
 
+	/** Look At Priority */
+	UPROPERTY(EditAnywhere, Category = Parameter)
+	FMassLookAtPriority Priority{static_cast<uint8>(EMassLookAtPriorities::LowestPriority)};
+
 	/** Look At Mode */
 	UPROPERTY(EditAnywhere, Category = Parameter)
-	EMassLookAtMode LookAtMode = EMassLookAtMode::LookForward; 
+	EMassLookAtMode LookAtMode = EMassLookAtMode::LookForward;
+
+	/** Look at interpolation speed (not used by the LookAt processor but can be forwarded to the animation system). */
+	UPROPERTY(EditAnywhere, Category = Parameter)
+	EMassLookAtInterpolationSpeed InterpolationSpeed = EMassLookAtInterpolationSpeed::Regular;
 
 	/** Random gaze Mode */
 	UPROPERTY(EditAnywhere, Category = Parameter)
diff --git a/Engine/Plugins/AI/MassAI/Source/MassAIBehaviorEditor/MassAIBehaviorEditor.Build.cs b/Engine/Plugins/AI/MassAI/Source/MassAIBehaviorEditor/MassAIBehaviorEditor.Build.cs
new file mode 100644
--- /dev/null
+++ b/Engine/Plugins/AI/MassAI/Source/MassAIBehaviorEditor/MassAIBehaviorEditor.Build.cs
@@ -0,0 +1,37 @@
+// Copyright Epic Games, Inc. All Rights Reserved.
+
+namespace UnrealBuildTool.Rules
+{
+	public class MassAIBehaviorEditor : ModuleRules
+	{
+		public MassAIBehaviorEditor(ReadOnlyTargetRules Target) : base(Target)
+		{
+			PublicIncludePaths.AddRange(
+			new string[] {
+			}
+			);
+
+			PublicDependencyModuleNames.AddRange(
+			new string[] {
+			}
+			);
+
+			PrivateDependencyModuleNames.AddRange(
+			new string[] {
+				"AIModule",
+				"AppFramework",
+				"Core",
+				"CoreUObject",
+				"Engine",
+				"InputCore",
+				"PropertyEditor",
+				"MassAIBehavior",
+				"MassEntity",
+				"Slate",
+				"SlateCore",
+			}
+			);
+		}
+
+	}
+}
diff --git a/Engine/Plugins/AI/MassAI/Source/MassAIBehaviorEditor/Private/MassAIBehaviorEditorModule.cpp b/Engine/Plugins/AI/MassAI/Source/MassAIBehaviorEditor/Private/MassAIBehaviorEditorModule.cpp
new file mode 100644
--- /dev/null
+++ b/Engine/Plugins/AI/MassAI/Source/MassAIBehaviorEditor/Private/MassAIBehaviorEditorModule.cpp
@@ -0,0 +1,32 @@
+// Copyright Epic Games, Inc. All Rights Reserved.
+
+#include "MassAIBehaviorEditorModule.h"
+#include "MassLookAtPriorityDetails.h"
+#include "MassLookAtPriorityInfoDetails.h"
+#include "Modules/ModuleManager.h"
+#include "PropertyEditorModule.h"
+
+#define LOCTEXT_NAMESPACE "MassAIBehaviorEditor"
+
+IMPLEMENT_MODULE(FMassAIBehaviorEditorModule, MassAIBehaviorEditor)
+
+void FMassAIBehaviorEditorModule::StartupModule()
+{
+	// Register the details customizer
+	FPropertyEditorModule& PropertyModule = FModuleManager::LoadModuleChecked<FPropertyEditorModule>("PropertyEditor");
+
+	PropertyModule.RegisterCustomPropertyTypeLayout(TEXT("MassLookAtPriority"), FOnGetPropertyTypeCustomizationInstance::CreateStatic(&FMassLookAtPriorityDetails::MakeInstance));
+	PropertyModule.RegisterCustomPropertyTypeLayout(TEXT("MassLookAtPriorityInfo"), FOnGetPropertyTypeCustomizationInstance::CreateStatic(&FMassLookAtPriorityInfoDetails::MakeInstance));
+}
+
+void FMassAIBehaviorEditorModule::ShutdownModule()
+{
+	if (FModuleManager::Get().IsModuleLoaded(TEXT("PropertyEditor")))
+	{
+		FPropertyEditorModule& PropertyModule = FModuleManager::GetModuleChecked<FPropertyEditorModule>(TEXT("PropertyEditor"));
+		PropertyModule.UnregisterCustomPropertyTypeLayout(TEXT("MassLookAtPriority"));
+		PropertyModule.UnregisterCustomPropertyTypeLayout(TEXT("MassLookAtPriorityInfo"));
+	}
+}
+
+#undef LOCTEXT_NAMESPACE
diff --git a/Engine/Plugins/AI/MassAI/Source/MassAIBehaviorEditor/Private/MassAIBehaviorEditorModule.h b/Engine/Plugins/AI/MassAI/Source/MassAIBehaviorEditor/Private/MassAIBehaviorEditorModule.h
new file mode 100644
--- /dev/null
+++ b/Engine/Plugins/AI/MassAI/Source/MassAIBehaviorEditor/Private/MassAIBehaviorEditorModule.h
@@ -0,0 +1,14 @@
+// Copyright Epic Games, Inc. All Rights Reserved.
+
+#pragma  once
+
+#include "Modules/ModuleInterface.h"
+
+class MASSAIBEHAVIOREDITOR_API FMassAIBehaviorEditorModule : public IModuleInterface
+{
+public:
+	//~ Begin IModuleInterface
+	virtual void StartupModule() override;
+	virtual void ShutdownModule() override;
+	//~ End IModuleInterface
+};
diff --git a/Engine/Plugins/AI/MassAI/Source/MassAIBehaviorEditor/Private/MassLookAtPriorityDetails.cpp b/Engine/Plugins/AI/MassAI/Source/MassAIBehaviorEditor/Private/MassLookAtPriorityDetails.cpp
new file mode 100644
--- /dev/null
+++ b/Engine/Plugins/AI/MassAI/Source/MassAIBehaviorEditor/Private/MassLookAtPriorityDetails.cpp
@@ -0,0 +1,136 @@
+// Copyright Epic Games, Inc. All Rights Reserved.
+
+#include "MassLookAtPriorityDetails.h"
+
+#include "DetailWidgetRow.h"
+#include "DetailLayoutBuilder.h"
+#include "ISettingsModule.h"
+#include "MassLookAtSettings.h"
+#include "Framework/MultiBox/MultiBoxBuilder.h"
+#include "Modules/ModuleManager.h"
+#include "Widgets/Input/SComboButton.h"
+
+#define LOCTEXT_NAMESPACE "MassLookAtPriorityDetails"
+
+TSharedRef<IPropertyTypeCustomization> FMassLookAtPriorityDetails::MakeInstance()
+{
+	return MakeShareable(new FMassLookAtPriorityDetails);
+}
+
+void FMassLookAtPriorityDetails::CustomizeHeader(const TSharedRef<IPropertyHandle> StructPropertyHandle, FDetailWidgetRow& HeaderRow, IPropertyTypeCustomizationUtils& StructCustomizationUtils)
+{
+	// Can't use GET_MEMBER_NAME_CHECKED() since property is private
+	constexpr TCHAR PriorityValueName[](TEXT("Value"));
+	PriorityValueProperty = StructPropertyHandle->GetChildHandle(PriorityValueName);
+	checkf(PriorityValueProperty, TEXT("Unable to find property called '%s' in FMassLookAtPriority."
+		" Make sure this code matches the property name and that the property is exposed to the editor."), PriorityValueName);
+
+	// Build cache for current priorities and register for modifications
+	CachePriorityInfos();
+	UMassLookAtSettings::OnMassLookAtPrioritiesChanged.AddSP(this, &FMassLookAtPriorityDetails::CachePriorityInfos);
+
+	HeaderRow
+	.NameContent()
+	[
+		StructPropertyHandle->CreatePropertyNameWidget()
+	]
+	.ValueContent()
+	.MinDesiredWidth(200.0f)
+	[
+		SNew(SComboButton)
+		.OnGetMenuContent(this, &FMassLookAtPriorityDetails::OnGetComboContent)
+		.ContentPadding(FMargin(2.0f, 0.0f))
+		.ButtonContent()
+		[
+			SNew(SHorizontalBox)
+			// Description
+			+ SHorizontalBox::Slot()
+			.VAlign(VAlign_Center)
+			.Padding(FMargin(6.0f, 2.0f))
+			[
+				SNew(STextBlock)
+				.Text(this, &FMassLookAtPriorityDetails::GetDescription)
+				.Font(IDetailLayoutBuilder::GetDetailFontBold())
+			]
+		]
+	];
+}
+
+void FMassLookAtPriorityDetails::CustomizeChildren(TSharedRef<IPropertyHandle> StructPropertyHandle, IDetailChildrenBuilder& StructBuilder, IPropertyTypeCustomizationUtils& StructCustomizationUtils)
+{
+}
+
+void FMassLookAtPriorityDetails::CachePriorityInfos()
+{
+	PriorityInfos.Reset();
+	if (const UMassLookAtSettings* Settings = GetDefault<UMassLookAtSettings>())
+	{
+		Settings->GetValidPriorityInfos(PriorityInfos);
+	}
+}
+
+FText FMassLookAtPriorityDetails::GetDescription() const
+{
+	uint8 Priority = 0;
+	const FPropertyAccess::Result Result = PriorityValueProperty->GetValue(Priority);
+	if (Result == FPropertyAccess::Success)
+	{
+		FText Name = LOCTEXT("NameEmpty", "(not set)");
+		for (const FMassLookAtPriorityInfo& Info : PriorityInfos)
+		{
+			if (Info.Priority.Get() == Priority)
+			{
+				Name = FText::FromName(Info.Name);
+				break;
+			}
+		}
+		return Name;
+	}
+
+	if (Result == FPropertyAccess::MultipleValues)
+	{
+		return LOCTEXT("MultipleSelected", "(Multiple Selected)");
+	}
+
+	return FText::GetEmpty();
+}
+
+TSharedRef<SWidget> FMassLookAtPriorityDetails::OnGetComboContent() const
+{
+	FMenuBuilder MenuBuilder(true, nullptr);
+
+	FUIAction EditPrioritiesItemAction(FExecuteAction::CreateSPLambda(this, []
+		{
+			// Go to settings to edit priorities
+			const UMassLookAtSettings* Settings = GetDefault<UMassLookAtSettings>();
+			FModuleManager::LoadModuleChecked<ISettingsModule>("Settings").ShowViewer(Settings->GetContainerName(), Settings->GetCategoryName(), Settings->GetSectionName());
+		}));
+
+	MenuBuilder.AddMenuEntry(LOCTEXT("EditPriorities", "Edit Priorities..."), TAttribute<FText>(), FSlateIcon(), EditPrioritiesItemAction);
+	MenuBuilder.AddMenuSeparator();
+
+	for (const FMassLookAtPriorityInfo& Info : PriorityInfos)
+	{
+		const uint8 PriorityValue = Info.Priority.Get();
+		FUIAction BitItemAction(FExecuteAction::CreateLambda([PriorityValueProperty = PriorityValueProperty, PriorityValue]
+		{
+			PriorityValueProperty->SetValue(PriorityValue);
+		}));
+
+		TSharedRef<SWidget> Widgets = SNew(SHorizontalBox)
+			// Description
+			+ SHorizontalBox::Slot()
+			.VAlign(VAlign_Center)
+			.Padding(FMargin(6.0f, 2.0f))
+			[
+				SNew(STextBlock)
+				.Text(FText::FromName(Info.Name))
+			];
+
+		MenuBuilder.AddMenuEntry(BitItemAction, Widgets);
+	}
+
+	return MenuBuilder.MakeWidget();
+}
+
+#undef LOCTEXT_NAMESPACE
\ No newline at end of file
diff --git a/Engine/Plugins/AI/MassAI/Source/MassAIBehaviorEditor/Private/MassLookAtPriorityDetails.h b/Engine/Plugins/AI/MassAI/Source/MassAIBehaviorEditor/Private/MassLookAtPriorityDetails.h
new file mode 100644
--- /dev/null
+++ b/Engine/Plugins/AI/MassAI/Source/MassAIBehaviorEditor/Private/MassLookAtPriorityDetails.h
@@ -0,0 +1,38 @@
+// Copyright Epic Games, Inc. All Rights Reserved.
+#pragma once
+
+#include "IPropertyTypeCustomization.h"
+
+struct FMassLookAtPriorityInfo;
+class IPropertyHandle;
+class SWidget;
+
+/**
+ * Type customization for FMassLookAtPriority.
+ */
+class FMassLookAtPriorityDetails final : public IPropertyTypeCustomization
+{
+public:
+	/** Makes a new instance of this detail layout class for a specific detail view requesting it */
+	static TSharedRef<IPropertyTypeCustomization> MakeInstance();
+
+	//~ Begin IPropertyTypeCustomization interface
+	virtual void CustomizeHeader(TSharedRef<IPropertyHandle> StructPropertyHandle, FDetailWidgetRow& HeaderRow, IPropertyTypeCustomizationUtils& StructCustomizationUtils) override;
+	virtual void CustomizeChildren(TSharedRef<IPropertyHandle> StructPropertyHandle, IDetailChildrenBuilder& StructBuilder, IPropertyTypeCustomizationUtils& StructCustomizationUtils) override;
+	//~ End IPropertyTypeCustomization interface
+
+private:
+
+	/** @return Description for the current priority value */
+	FText GetDescription() const;
+
+	/** Builds the priority drop down */
+	TSharedRef<SWidget> OnGetComboContent() const;
+
+	/** Cache priority values from MassLookAtSetting into 'PriorityInfos' */
+	void CachePriorityInfos();
+
+	TArray<FMassLookAtPriorityInfo> PriorityInfos;
+	TSharedPtr<IPropertyHandle> PriorityValueProperty;
+	TSharedPtr<IPropertyHandle> StructProperty;
+};
\ No newline at end of file
diff --git a/Engine/Plugins/AI/MassAI/Source/MassAIBehaviorEditor/Private/MassLookAtPriorityInfoDetails.cpp b/Engine/Plugins/AI/MassAI/Source/MassAIBehaviorEditor/Private/MassLookAtPriorityInfoDetails.cpp
new file mode 100644
--- /dev/null
+++ b/Engine/Plugins/AI/MassAI/Source/MassAIBehaviorEditor/Private/MassLookAtPriorityInfoDetails.cpp
@@ -0,0 +1,99 @@
+// Copyright Epic Games, Inc. All Rights Reserved.
+
+#include "MassLookAtPriorityInfoDetails.h"
+#include "MassLookAtTypes.h"
+#include "DetailWidgetRow.h"
+#include "DetailLayoutBuilder.h"
+#include "Widgets/Text/STextBlock.h"
+
+#define LOCTEXT_NAMESPACE "MassLookAtPriorityInfoDetails"
+
+FMassLookAtPriorityInfoDetails::~FMassLookAtPriorityInfoDetails()
+{
+}
+
+TSharedRef<IPropertyTypeCustomization> FMassLookAtPriorityInfoDetails::MakeInstance()
+{
+	return MakeShareable(new FMassLookAtPriorityInfoDetails);
+}
+
+void FMassLookAtPriorityInfoDetails::CustomizeHeader(const TSharedRef<IPropertyHandle> StructPropertyHandle, FDetailWidgetRow& HeaderRow, IPropertyTypeCustomizationUtils& StructCustomizationUtils)
+{
+	NameProperty = StructPropertyHandle->GetChildHandle(GET_MEMBER_NAME_CHECKED(FMassLookAtPriorityInfo, Name));
+	PriorityProperty = StructPropertyHandle->GetChildHandle(GET_MEMBER_NAME_CHECKED(FMassLookAtPriorityInfo, Priority));
+
+	HeaderRow
+	.WholeRowContent()
+	[
+		SNew(SHorizontalBox)
+		// Priority value
+		+ SHorizontalBox::Slot()
+		.FillWidth(0.1f)
+		.MaxWidth(30.0f)
+		.VAlign(VAlign_Center)
+		.Padding(FMargin(6.0f, 2.0f))
+		[
+			SNew(STextBlock)
+			.Text(this, &FMassLookAtPriorityInfoDetails::GetPriorityDescription)
+			.Font(IDetailLayoutBuilder::GetDetailFont())
+		]
+		// Description
+		+ SHorizontalBox::Slot()
+		.FillWidth(1.0f)
+		.VAlign(VAlign_Center)
+		.Padding(FMargin(6.0f, 2.0f))
+		[
+			NameProperty->CreatePropertyValueWidget()
+		]
+		+ SHorizontalBox::Slot()
+		.Padding(FMargin(12.0f, 0.0f))
+		.HAlign(HAlign_Right)
+		[
+			StructPropertyHandle->CreateDefaultPropertyButtonWidgets()
+		]
+	];
+}
+
+void FMassLookAtPriorityInfoDetails::CustomizeChildren(TSharedRef<IPropertyHandle> StructPropertyHandle, IDetailChildrenBuilder& StructBuilder, IPropertyTypeCustomizationUtils& StructCustomizationUtils)
+{
+	// Nothing required; pure virtual in IPropertyTypeCustomization
+}
+
+FText FMassLookAtPriorityInfoDetails::GetPriorityDescription() const
+{
+	if (!PriorityProperty.IsValid())
+	{
+		return FText::GetEmpty();
+	}
+
+	FMassLookAtPriority Value;
+	bool bValueSet = false;
+
+	TArray<void*> RawData;
+	PriorityProperty->AccessRawData(RawData);
+	for (void* Data : RawData)
+	{
+		if (Data)
+		{
+			FMassLookAtPriority CurValue = *static_cast<FMassLookAtPriority*>(Data);
+			if (!bValueSet)
+			{
+				bValueSet = true;
+				Value = CurValue;
+			}
+			else if (CurValue != Value)
+			{
+				// Multiple values
+				return FText::GetEmpty();
+			}
+		}
+	}
+
+	if (Value.IsValid())
+	{
+		return FText::AsNumber(Value.Get());
+	}
+	return FText::GetEmpty();
+}
+
+#undef LOCTEXT_NAMESPACE
\ No newline at end of file
diff --git a/Engine/Plugins/AI/MassAI/Source/MassAIBehaviorEditor/Private/MassLookAtPriorityInfoDetails.h b/Engine/Plugins/AI/MassAI/Source/MassAIBehaviorEditor/Private/MassLookAtPriorityInfoDetails.h
new file mode 100644
--- /dev/null
+++ b/Engine/Plugins/AI/MassAI/Source/MassAIBehaviorEditor/Private/MassLookAtPriorityInfoDetails.h
@@ -0,0 +1,31 @@
+// Copyright Epic Games, Inc. All Rights Reserved.
+#pragma once
+
+#include "IPropertyTypeCustomization.h"
+#include "Templates/SharedPointer.h"
+
+class IPropertyHandle;
+
+/**
+ * Type customization for FMassLookAtPriorityInfo.
+ */
+class FMassLookAtPriorityInfoDetails : public IPropertyTypeCustomization
+{
+public:
+	virtual ~FMassLookAtPriorityInfoDetails() override;
+
+	/** Makes a new instance of this detail layout class for a specific detail view requesting it */
+	static TSharedRef<IPropertyTypeCustomization> MakeInstance();
+
+	//~ Begin IPropertyTypeCustomization interface
+	virtual void CustomizeHeader(TSharedRef<IPropertyHandle> StructPropertyHandle, FDetailWidgetRow& HeaderRow, IPropertyTypeCustomizationUtils& StructCustomizationUtils) override;
+	virtual void CustomizeChildren(TSharedRef<IPropertyHandle> StructPropertyHandle, IDetailChildrenBuilder& StructBuilder, IPropertyTypeCustomizationUtils& StructCustomizationUtils) override;
+	//~ End IPropertyTypeCustomization interface
+
+private:
+
+	FText GetPriorityDescription() const;
+
+	TSharedPtr<IPropertyHandle> PriorityProperty;
+	TSharedPtr<IPropertyHandle> NameProperty;
+};
\ No newline at end of file
diff --git a/Engine/Plugins/AI/MassAI/Source/MassAIDebug/Private/GameplayDebuggerCategory_Mass.cpp b/Engine/Plugins/AI/MassAI/Source/MassAIDebug/Private/GameplayDebuggerCategory_Mass.cpp
--- a/Engine/Plugins/AI/MassAI/Source/MassAIDebug/Private/GameplayDebuggerCategory_Mass.cpp
+++ b/Engine/Plugins/AI/MassAI/Source/MassAIDebug/Private/GameplayDebuggerCategory_Mass.cpp
@@ -278,7 +278,7 @@ void FGameplayDebuggerCategory_Mass::CollectData(APlayerController* OwnerPC, AAc
 		AddTextLine(FString::Printf(TEXT("Source: {yellow}%s{white}"), bDebugLocalEntityManager ? TEXT("LOCAL") : TEXT("REMOTE")));
 	}
 
-	UWorld* World = GetDataWorld(OwnerPC, DebugActor);
+	const UWorld* World = GetDataWorld(OwnerPC, DebugActor);
 	check(World);
 
 	UMassEntitySubsystem* EntitySubsystem = UWorld::GetSubsystem<UMassEntitySubsystem>(World);
@@ -288,7 +288,8 @@ void FGameplayDebuggerCategory_Mass::CollectData(APlayerController* OwnerPC, AAc
 		return;
 	}
 	FMassEntityManager& EntityManager = EntitySubsystem->GetMutableEntityManager();
-	UMassDebuggerSubsystem* Debugger = World->GetSubsystem<UMassDebuggerSubsystem>();
+	const UMassLookAtSubsystem* LookAtSubsystem = World->GetSubsystem<UMassLookAtSubsystem>();
+	const UMassDebuggerSubsystem* Debugger = World->GetSubsystem<UMassDebuggerSubsystem>();
 
 	const UMassAgentComponent* AgentComp = nullptr;
 	
@@ -521,7 +522,7 @@ void FGameplayDebuggerCategory_Mass::CollectData(APlayerController* OwnerPC, AAc
 		{
 			FMassExecutionContext Context(EntityManager, 0.0f);
 		
-			EntityQuery.ForEachEntityChunk(Context, [this, MassStateTreeSubsystem, SmartObjectSubsystem, ViewLocation, ViewDirection, CurrentTime](FMassExecutionContext& Context)
+			EntityQuery.ForEachEntityChunk(Context, [this, MassStateTreeSubsystem, SmartObjectSubsystem, LookAtSubsystem, ViewLocation, ViewDirection, CurrentTime](FMassExecutionContext& Context)
 			{
 				FMassEntityManager& EntityManager = Context.GetEntityManagerChecked();
 
@@ -554,6 +555,7 @@ void FGameplayDebuggerCategory_Mass::CollectData(APlayerController* OwnerPC, AAc
 
 				for (int32 EntityIndex = 0; EntityIndex < NumEntities; ++EntityIndex)
 				{
+					FMassEntityHandle Entity = Context.GetEntity(EntityIndex);
 					const FTransformFragment& Transform = TransformList[EntityIndex];
 					const FVector EntityLocation = Transform.GetTransform().GetLocation();
 					
@@ -721,7 +723,6 @@ void FGameplayDebuggerCategory_Mass::CollectData(APlayerController* OwnerPC, AAc
 						FString Status;
 
 						// Entity name
-						FMassEntityHandle Entity = Context.GetEntity(EntityIndex);
 						Status += TEXT("{orange}");
 						Status += Entity.DebugGetDescription();
 						Status += TEXT(" {white}LOD ");
@@ -773,8 +774,13 @@ void FGameplayDebuggerCategory_Mass::CollectData(APlayerController* OwnerPC, AAc
 							const double RemainingTime = LookAt.GazeDuration - (CurrentTime - LookAt.GazeStartTime);
 							Status += FString::Printf(TEXT("{turquoise}%s/%s {lightgrey}%.1f\n"),
 								*UEnum::GetDisplayValueAsText(LookAt.LookAtMode).ToString(), *UEnum::GetDisplayValueAsText(LookAt.RandomGazeMode).ToString(), RemainingTime);
+
+							if (LookAtSubsystem != nullptr)
+							{
+								Status += FString::Printf(TEXT("{lightgrey}%s\n"), *LookAtSubsystem->DebugGetRequestsString(Entity));
+							}
 						}
-						
+
 						if (!Status.IsEmpty())
 						{
 							BasePos += FVector(0,0,50);
