[MassLookAt] validate entity for deferred commands

diff --git a/Engine/Plugins/AI/MassAI/Source/MassAIBehavior/Private/MassLookAtSubsystem.cpp b/Engine/Plugins/AI/MassAI/Source/MassAIBehavior/Private/MassLookAtSubsystem.cpp
--- a/Engine/Plugins/AI/MassAI/Source/MassAIBehavior/Private/MassLookAtSubsystem.cpp
+++ b/Engine/Plugins/AI/MassAI/Source/MassAIBehavior/Private/MassLookAtSubsystem.cpp
@@ -305,13 +305,16 @@ void UMassLookAtSubsystem::UpdateLookAts(const FMassExecutionContext& Context, c
 		{
 			for (const TPair<FMassEntityHandle, FMassLookAtRequestFragment>& Pair : Updated)
 			{
-				if (FMassLookAtFragment* LookAtFragment = Manager.GetFragmentDataPtr<FMassLookAtFragment>(Pair.Key))
+				if (Manager.IsEntityValid(Pair.Key))
 				{
-					// Only update main LookAt information and don't modify gaze related one
-					const FMassLookAtRequestFragment& Request = Pair.Value;
-					LookAtFragment->InterpolationSpeed = Request.InterpolationSpeed;
-					LookAtFragment->LookAtMode = Request.LookAtMode;
-					LookAtFragment->TrackedEntity = Request.TargetEntity;
+					if (FMassLookAtFragment* LookAtFragment = Manager.GetFragmentDataPtr<FMassLookAtFragment>(Pair.Key))
+					{
+						// Only update main LookAt information and don't modify gaze related one
+						const FMassLookAtRequestFragment& Request = Pair.Value;
+						LookAtFragment->InterpolationSpeed = Request.InterpolationSpeed;
+						LookAtFragment->LookAtMode = Request.LookAtMode;
+						LookAtFragment->TrackedEntity = Request.TargetEntity;
+					}
 				}
 			}
 		});
