[LookAt] added variable for interpolation speed to anim module and exposed the enumeration

diff --git a/Engine/Plugins/AI/MassAI/Source/MassAIBehavior/Public/MassLookAtTypes.h b/Engine/Plugins/AI/MassAI/Source/MassAIBehavior/Public/MassLookAtTypes.h
--- a/Engine/Plugins/AI/MassAI/Source/MassAIBehavior/Public/MassLookAtTypes.h
+++ b/Engine/Plugins/AI/MassAI/Source/MassAIBehavior/Public/MassLookAtTypes.h
@@ -18,7 +18,7 @@ namespace UE::Mass::LookAt
 /**
  * Enum representing the different interpolation speeds that can be used when assigning new LookAt targets
  */
-UENUM()
+UENUM(BlueprintType)
 enum class EMassLookAtInterpolationSpeed : uint8
 {
 	Instant,
