[LookAt] reenabled avoidance obstacles as potential targets

diff --git a/Engine/Plugins/AI/MassAI/Source/MassAIBehavior/Private/MassLookAtProcessors.cpp b/Engine/Plugins/AI/MassAI/Source/MassAIBehavior/Private/MassLookAtProcessors.cpp
--- a/Engine/Plugins/AI/MassAI/Source/MassAIBehavior/Private/MassLookAtProcessors.cpp
+++ b/Engine/Plugins/AI/MassAI/Source/MassAIBehavior/Private/MassLookAtProcessors.cpp
@@ -28,7 +28,7 @@ namespace UE::MassBehavior
 	namespace Tweakables
 	{
 		float TrajectoryLookAhead = 600.f;
-		bool bConsiderObstacleTargets = false;
+		bool bConsiderObstacleTargets = true;
 	}
 
 	FAutoConsoleVariableRef CVars[] =
