MassNavigation uses new mass debug logging API, which adds an Actor LogOwner where available allowing debug to work with the RewindDebugger

diff --git a/Engine/Plugins/AI/MassAI/Source/MassNavigation/Private/MassAvoidanceProcessors.cpp b/Engine/Plugins/AI/MassAI/Source/MassNavigation/Private/MassAvoidanceProcessors.cpp
--- a/Engine/Plugins/AI/MassAI/Source/MassNavigation/Private/MassAvoidanceProcessors.cpp
+++ b/Engine/Plugins/AI/MassAI/Source/MassNavigation/Private/MassAvoidanceProcessors.cpp
@@ -10,6 +10,7 @@
 #include "Logging/LogMacros.h"
 #include "MassSimulationLOD.h"
 #include "MassCommonFragments.h"
+#include "MassDebugLogging.h"
 #include "MassMovementFragments.h"
 #include "MassNavigationSubsystem.h"
 #include "MassNavigationFragments.h"
@@ -18,8 +19,6 @@
 #include "MassDebugger.h"
 #include "MassNavigationDebug.h"
 
-#define UNSAFE_FOR_MT 1
-
 DEFINE_LOG_CATEGORY(LogAvoidance);
 DEFINE_LOG_CATEGORY(LogAvoidanceVelocities);
 DEFINE_LOG_CATEGORY(LogAvoidanceAgents);
@@ -210,7 +209,7 @@ namespace UE::MassAvoidance
 	static void DebugDrawVelocity(const FDebugContext& Context, const FVector& Start, const FVector& End, const FColor& Color)
 	{
 		// Different arrow than DebugDrawArrow()
-		if (!DebugIsSelected(Context.Entity))
+		if (!Context.ShouldLogEntity())
 		{
 			return;
 		}
@@ -223,10 +222,11 @@ namespace UE::MassAvoidance
 		const FVector Left = Perp - (Pointyness * UnitV);
 		const FVector Right = -Perp - (Pointyness * UnitV);
 		const FVector::FReal HeadSize = 0.08 * Line.Size();
-		UE_VLOG_SEGMENT_THICK(Context.LogOwner, Context.Category, Log, Start, End, Color, (int16)Thickness, TEXT(""));
-		UE_VLOG_SEGMENT_THICK(Context.LogOwner, Context.Category, Log, End, End + HeadSize * Left, Color, (int16)Thickness, TEXT(""));
-		UE_VLOG_SEGMENT_THICK(Context.LogOwner, Context.Category, Log, End, End + HeadSize * Right, Color, (int16)Thickness, TEXT(""));
-		UE_VLOG_SEGMENT_THICK(Context.LogOwner, Context.Category, Log, End + HeadSize * Left, End + HeadSize * Right, Color, (int16)Thickness, TEXT(""));
+		const UObject* LogOwner = Context.GetLogOwner(); 
+		UE_VLOG_SEGMENT_THICK(LogOwner, Context.Category, Log, Start, End, Color, (int16)Thickness, TEXT(""));
+		UE_VLOG_SEGMENT_THICK(LogOwner, Context.Category, Log, End, End + HeadSize * Left, Color, (int16)Thickness, TEXT(""));
+		UE_VLOG_SEGMENT_THICK(LogOwner, Context.Category, Log, End, End + HeadSize * Right, Color, (int16)Thickness, TEXT(""));
+		UE_VLOG_SEGMENT_THICK(LogOwner, Context.Category, Log, End + HeadSize * Left, End + HeadSize * Right, Color, (int16)Thickness, TEXT(""));
 
 		if (UseDrawDebugHelper() && Context.World)
 		{
@@ -277,6 +277,10 @@ void UMassMovingAvoidanceProcessor::ConfigureQueries(const TSharedRef<FMassEntit
 	EntityQuery.AddTagRequirement<FMassOffLODTag>(EMassFragmentPresence::None);
 	EntityQuery.AddConstSharedRequirement<FMassMovingAvoidanceParameters>(EMassFragmentPresence::All);
 	EntityQuery.AddConstSharedRequirement<FMassMovementParameters>(EMassFragmentPresence::All);
+
+#if WITH_MASSGAMEPLAY_DEBUG
+	EntityQuery.DebugEnableEntityOwnerLogging();
+#endif
 }
 
 void UMassMovingAvoidanceProcessor::InitializeInternal(UObject& Owner, const TSharedRef<FMassEntityManager>& EntityManager)
@@ -397,14 +401,14 @@ void UMassMovingAvoidanceProcessor::Execute(FMassEntityManager& EntityManager, F
 			const FVector::FReal NearStartScaling = FMath::Lerp<FVector::FReal>(MovingAvoidanceParams.StartOfPathAvoidanceScale, 1., NearStartFade);
 			const FVector::FReal NearEndScaling = FMath::Lerp<FVector::FReal>(MovingAvoidanceParams.EndOfPathAvoidanceScale, 1., NearEndFade);
 			
-#if WITH_MASSGAMEPLAY_DEBUG && UNSAFE_FOR_MT
-			const UE::MassAvoidance::FDebugContext BaseDebugContext(this, LogAvoidance, World, Entity);
-			const UE::MassAvoidance::FDebugContext VelocitiesDebugContext(this, LogAvoidanceVelocities, World, Entity);
-			const UE::MassAvoidance::FDebugContext ObstacleDebugContext(this, LogAvoidanceObstacles, World, Entity);
-			const UE::MassAvoidance::FDebugContext AgentDebugContext(this, LogAvoidanceAgents, World, Entity);
+#if WITH_MASSGAMEPLAY_DEBUG
+			const UE::MassAvoidance::FDebugContext BaseDebugContext(Context, this, LogAvoidance, World, Entity, EntityIt);
+			const UE::MassAvoidance::FDebugContext VelocitiesDebugContext(Context, this, LogAvoidanceVelocities, World, Entity, EntityIt);
+			const UE::MassAvoidance::FDebugContext ObstacleDebugContext(Context, this, LogAvoidanceObstacles, World, Entity, EntityIt);
+			const UE::MassAvoidance::FDebugContext AgentDebugContext(Context, this, LogAvoidanceAgents, World, Entity, EntityIt);
 
 			FColor EntityColor = FColor::White;
-			if (UE::Mass::Debug::IsDebuggingEntity(Entity, &EntityColor))
+			if (BaseDebugContext.ShouldLogEntity(&EntityColor))
 			{
 				// Draw agent
 				const FString Text = FString::Printf(TEXT("%i"), Entity.Index);
@@ -447,7 +451,7 @@ void UMassMovingAvoidanceProcessor::Execute(FMassEntityManager& EntityManager, F
 				const FVector DesiredAcceleration = UE::MassNavigation::ClampVector(SteeringForce, MaxSteerAccel);
 				const FVector DesiredVelocity = UE::MassNavigation::ClampVector(AgentVelocity + DesiredAcceleration * DeltaTime, MaximumSpeed);
 
-#if WITH_MASSGAMEPLAY_DEBUG && UNSAFE_FOR_MT
+#if WITH_MASSGAMEPLAY_DEBUG
 				// Draw desired velocity (yellow)
 				UE::MassAvoidance::DebugDrawVelocity(VelocitiesDebugContext, AgentLocation + UE::MassAvoidance::DebugInputForceHeight,
 					AgentLocation + UE::MassAvoidance::DebugInputForceHeight + DesiredVelocity, UE::MassAvoidance::DesiredVelocityColor);
@@ -527,8 +531,8 @@ void UMassMovingAvoidanceProcessor::Execute(FMassEntityManager& EntityManager, F
 					Contact.bBehindTheEdge = bBehindTheEdge;
 					Contacts.Add(Contact);
 
-#if WITH_MASSGAMEPLAY_DEBUG && UNSAFE_FOR_MT
-					if (UE::MassAvoidance::Tweakables::bEnableDetailedDebug)
+#if WITH_MASSGAMEPLAY_DEBUG
+					if (UE::MassAvoidance::Tweakables::bEnableDetailedDebug && ObstacleDebugContext.ShouldLogEntity())
 					{
 						if (bAlongTheEdge)
 						{
@@ -556,8 +560,8 @@ void UMassMovingAvoidanceProcessor::Execute(FMassEntityManager& EntityManager, F
 					{
 						// Avoid edges
 						
-#if WITH_MASSGAMEPLAY_DEBUG && UNSAFE_FOR_MT
-						if (UE::MassAvoidance::Tweakables::bEnableDetailedDebug)
+#if WITH_MASSGAMEPLAY_DEBUG
+						if (UE::MassAvoidance::Tweakables::bEnableDetailedDebug && ObstacleDebugContext.ShouldLogEntity())
 						{
 							// Draw environment predictive avoidance distance
 							FVector ZOffset = FVector(0., 0., 7.);
@@ -593,7 +597,7 @@ void UMassMovingAvoidanceProcessor::Execute(FMassEntityManager& EntityManager, F
 
 						SteeringForce += AvoidForce;
 
-#if WITH_MASSGAMEPLAY_DEBUG && UNSAFE_FOR_MT
+#if WITH_MASSGAMEPLAY_DEBUG
 						if (!AvoidForce.IsNearlyZero())
 						{
 							if (UE::MassAvoidance::Tweakables::bEnableDetailedDebug)
@@ -619,7 +623,7 @@ void UMassMovingAvoidanceProcessor::Execute(FMassEntityManager& EntityManager, F
 					
 				} // edge loop
 
-#if WITH_MASSGAMEPLAY_DEBUG && UNSAFE_FOR_MT
+#if WITH_MASSGAMEPLAY_DEBUG
 				// Draw total steering force to avoid obstacles
 				const FVector EnvironmentAvoidSteeringForce = SteeringForce - OldSteeringForce;
 				UE::MassAvoidance::DebugDrawSummedForce(ObstacleDebugContext,
@@ -660,8 +664,8 @@ void UMassMovingAvoidanceProcessor::Execute(FMassEntityManager& EntityManager, F
 						
 						SteeringForce += SeparationForce;
 
-	#if WITH_MASSGAMEPLAY_DEBUG && UNSAFE_FOR_MT
-						UE_VLOG(ObstacleDebugContext.LogOwner, ObstacleDebugContext.Category, Log, TEXT("EnvironmentSeparationAgentRadius: %f, EnvironmentSeparationDistanc: %f, Distance: %f, SeparationPenalty: %f, SeparationMag: %f, SeparationForce: %s"),
+	#if WITH_MASSGAMEPLAY_DEBUG
+						UE_VLOG(ObstacleDebugContext.GetLogOwner(), ObstacleDebugContext.Category, Log, TEXT("EnvironmentSeparationAgentRadius: %f, EnvironmentSeparationDistanc: %f, Distance: %f, SeparationPenalty: %f, SeparationMag: %f, SeparationForce: %s"),
 							EnvironmentSeparationAgentRadius,
 							MovingAvoidanceParams.EnvironmentSeparationDistance,
 							Contact.Distance,
@@ -683,7 +687,7 @@ void UMassMovingAvoidanceProcessor::Execute(FMassEntityManager& EntityManager, F
 					}
 				}
 				
-#if WITH_MASSGAMEPLAY_DEBUG && UNSAFE_FOR_MT
+#if WITH_MASSGAMEPLAY_DEBUG
 				// Draw total steering force to separate from close edges
 				const FVector TotalSeparationForce = SteeringForce - SteeringForceBeforeSeparation;
 				UE::MassAvoidance::DebugDrawSummedForce(ObstacleDebugContext,
@@ -692,7 +696,7 @@ void UMassMovingAvoidanceProcessor::Execute(FMassEntityManager& EntityManager, F
 					UE::MassAvoidance::ObstacleSeparationForceColor);
 
 				// Display close obstacle edges
-				if (UE::MassAvoidance::DebugIsSelected(Entity))
+				if (ObstacleDebugContext.ShouldLogEntity())
 				{
 					for (const FNavigationAvoidanceEdge& Edge : NavEdges.AvoidanceEdges)
 					{
@@ -929,7 +933,7 @@ void UMassMovingAvoidanceProcessor::Execute(FMassEntityManager& EntityManager, F
 
 				SteeringForce += AvoidForce;
 
-#if WITH_MASSGAMEPLAY_DEBUG && UNSAFE_FOR_MT
+#if WITH_MASSGAMEPLAY_DEBUG
 				// Display close agent
 				UE::MassAvoidance::DebugDrawCylinder(AgentDebugContext, Collider.Location,
 					Collider.Location + UE::MassAvoidance::DebugLowCylinderOffset, Collider.Radius, UE::MassAvoidance::AgentsColor);
@@ -985,7 +989,7 @@ void UMassMovingAvoidanceProcessor::Execute(FMassEntityManager& EntityManager, F
 
 			Force.Value = UE::MassNavigation::ClampVector(SteeringForce, MaxSteerAccel); // Assume unit mass
 
-#if WITH_MASSGAMEPLAY_DEBUG && UNSAFE_FOR_MT
+#if WITH_MASSGAMEPLAY_DEBUG
 			const FVector AgentAvoidSteeringForce = SteeringForce - OldSteeringForce;
 
 			// Draw total steering force to separate agents
@@ -1033,6 +1037,10 @@ void UMassStandingAvoidanceProcessor::ConfigureQueries(const TSharedRef<FMassEnt
 	EntityQuery.AddTagRequirement<FMassLowLODTag>(EMassFragmentPresence::None);
 	EntityQuery.AddTagRequirement<FMassOffLODTag>(EMassFragmentPresence::None);
 	EntityQuery.AddConstSharedRequirement<FMassStandingAvoidanceParameters>(EMassFragmentPresence::All);
+
+#if WITH_MASSGAMEPLAY_DEBUG
+	EntityQuery.DebugEnableEntityOwnerLogging();
+#endif
 }
 
 void UMassStandingAvoidanceProcessor::InitializeInternal(UObject& Owner, const TSharedRef<FMassEntityManager>& EntityManager)
@@ -1168,8 +1176,8 @@ void UMassStandingAvoidanceProcessor::Execute(FMassEntityManager& EntityManager,
 
 			const FVector::FReal GhostRadius = AgentRadius * StandingParams.GhostSeparationRadiusScale;
 
-#if WITH_MASSGAMEPLAY_DEBUG && UNSAFE_FOR_MT			
-			const UE::MassAvoidance::FDebugContext AgentDebugContext(this, LogAvoidanceAgents, World, Entity);
+#if WITH_MASSGAMEPLAY_DEBUG			
+			const UE::MassAvoidance::FDebugContext AgentDebugContext(Context, this, LogAvoidanceAgents, World, Entity, EntityIt);
 
 			// Draw entity
 			UE::MassAvoidance::DebugDrawCylinder(AgentDebugContext, AgentLocation, AgentLocation + UE::MassAvoidance::DebugAgentHeightOffset,
@@ -1191,7 +1199,7 @@ void UMassStandingAvoidanceProcessor::Execute(FMassEntityManager& EntityManager,
 				const FVector::FReal OtherRadius = OtherEntityView.GetFragmentData<FAgentRadiusFragment>().Radius;
 				const FVector::FReal TotalRadius = GhostRadius + OtherRadius;
 
-#if WITH_MASSGAMEPLAY_DEBUG && UNSAFE_FOR_MT
+#if WITH_MASSGAMEPLAY_DEBUG
 				// Draw other agent
 				UE::MassAvoidance::DebugDrawCylinder(AgentDebugContext, OtherAgent.Location, OtherAgent.Location + UE::MassAvoidance::DebugAgentHeightOffset,
 					OtherRadius, UE::MassAvoidance::AgentsColor);
@@ -1226,7 +1234,7 @@ void UMassStandingAvoidanceProcessor::Execute(FMassEntityManager& EntityManager,
 
 					GhostSteeringForce += SeparationForce;
 
-#if WITH_MASSGAMEPLAY_DEBUG && UNSAFE_FOR_MT					
+#if WITH_MASSGAMEPLAY_DEBUG					
 					// Draw agent avoid force
 					UE::MassAvoidance::DebugDrawForce(AgentDebugContext,
 						Ghost.Location + UE::MassAvoidance::DebugAgentHeightOffset,
@@ -1258,7 +1266,7 @@ void UMassStandingAvoidanceProcessor::Execute(FMassEntityManager& EntityManager,
 
 					GhostSteeringForce += SeparationForce;
 
-#if WITH_MASSGAMEPLAY_DEBUG && UNSAFE_FOR_MT
+#if WITH_MASSGAMEPLAY_DEBUG
 					// Draw agent avoid force
 					UE::MassAvoidance::DebugDrawForce(AgentDebugContext,
 						OtherBasePosition + UE::MassAvoidance::DebugAgentHeightOffset,
@@ -1271,7 +1279,7 @@ void UMassStandingAvoidanceProcessor::Execute(FMassEntityManager& EntityManager,
 			GhostSteeringForce.Z = 0.;
 			GhostSteeringForce = UE::MassNavigation::ClampVector(GhostSteeringForce, StandingParams.GhostMaxAcceleration); // Assume unit mass
 
-#if WITH_MASSGAMEPLAY_DEBUG && UNSAFE_FOR_MT
+#if WITH_MASSGAMEPLAY_DEBUG
 			// Draw total steering force
 			UE::MassAvoidance::DebugDrawSummedForce(AgentDebugContext,
 				AgentLocation + UE::MassAvoidance::DebugAgentHeightOffset,
@@ -1299,4 +1307,3 @@ void UMassStandingAvoidanceProcessor::Execute(FMassEntityManager& EntityManager,
 	
 }
 
-#undef UNSAFE_FOR_MT
diff --git a/Engine/Plugins/AI/MassAI/Source/MassNavigation/Private/MassNavigationDebug.cpp b/Engine/Plugins/AI/MassAI/Source/MassNavigation/Private/MassNavigationDebug.cpp
--- a/Engine/Plugins/AI/MassAI/Source/MassNavigation/Private/MassNavigationDebug.cpp
+++ b/Engine/Plugins/AI/MassAI/Source/MassNavigation/Private/MassNavigationDebug.cpp
@@ -13,11 +13,13 @@ namespace UE::MassNavigation::Debug
 	namespace Tweakables
 	{
 		bool bUseDrawDebugHelpers = false;
+		bool bLogEverythingWhenRecording = false;
 	} // Tweakables
 
 	FAutoConsoleVariableRef Vars[] = 
 	{
-		FAutoConsoleVariableRef(TEXT("ai.mass.debug.UseDrawDebugHelpers"), Tweakables::bUseDrawDebugHelpers, TEXT("Use debug draw helpers in addition to visual logs."), ECVF_Cheat)
+		FAutoConsoleVariableRef(TEXT("ai.mass.debug.UseDrawDebugHelpers"), Tweakables::bUseDrawDebugHelpers, TEXT("Use debug draw helpers in addition to visual logs."), ECVF_Cheat),
+		FAutoConsoleVariableRef(TEXT("ai.mass.debug.LogEverythingWhenRecording"), Tweakables::bLogEverythingWhenRecording, TEXT("If true, will log all debug draw events regardless of debug entity selection if the visual log recorder is activated"), ECVF_Cheat)
 	};
 	
 	bool UseDrawDebugHelper()
@@ -25,24 +27,19 @@ namespace UE::MassNavigation::Debug
 		return Tweakables::bUseDrawDebugHelpers;
 	}
 
-	bool DebugIsSelected(const FMassEntityHandle Entity)
+	bool ShouldLogEverythingWhenRecording()
 	{
-#if WITH_MASSENTITY_DEBUG	
-		FColor Color;
-		return UE::Mass::Debug::IsDebuggingEntity(Entity, &Color);
-#else
-		return true;
-#endif		
+		return Tweakables::bLogEverythingWhenRecording;
 	}
 
 	void DebugDrawLine(const FDebugContext& Context, const FVector& Start, const FVector& End, const FColor& Color, const float Thickness /*= 0.f*/, const bool bPersistent /*= false*/, const FString& Text /*= FString()*/)
 	{
-		if (!DebugIsSelected(Context.Entity))
+		if (!Context.ShouldLogEntity())
 		{
 			return;
 		}
 
-		UE_VLOG_SEGMENT_THICK(Context.LogOwner, Context.Category, Log, Start, End, Color, (int16)Thickness, TEXT("%s"), *Text);
+		UE_VLOG_SEGMENT_THICK(Context.GetLogOwner(), Context.Category, Log, Start, End, Color, (int16)Thickness, TEXT("%s"), *Text);
 
 		if (UseDrawDebugHelper() && Context.World)
 		{
@@ -52,20 +49,22 @@ namespace UE::MassNavigation::Debug
 
 	void DebugDrawArrow(const FDebugContext& Context, const FVector& Start, const FVector& End, const FColor& Color, const float HeadSize /*= 8.f*/, const float Thickness /*= 1.5f*/)
 	{
-		if (!DebugIsSelected(Context.Entity))
+		if (!Context.ShouldLogEntity())
 		{
 			return;
 		}
 
+		const UObject* LogOwner = Context.GetLogOwner();
+
 		constexpr FVector::FReal Pointyness = 1.8;
 		const FVector Line = End - Start;
 		const FVector UnitV = Line.GetSafeNormal();
 		const FVector Perp = FVector::CrossProduct(UnitV, FVector::UpVector);
 		const FVector Left = Perp - (Pointyness*UnitV);
 		const FVector Right = -Perp - (Pointyness*UnitV);
-		UE_VLOG_SEGMENT_THICK(Context.LogOwner, Context.Category, Log, Start, End, Color, (int16)Thickness, TEXT(""));
-		UE_VLOG_SEGMENT_THICK(Context.LogOwner, Context.Category, Log, End, End + HeadSize * Left, Color, (int16)Thickness, TEXT(""));
-		UE_VLOG_SEGMENT_THICK(Context.LogOwner, Context.Category, Log, End, End + HeadSize * Right, Color, (int16)Thickness, TEXT(""));
+		UE_VLOG_SEGMENT_THICK(LogOwner, Context.Category, Log, Start, End, Color, (int16)Thickness, TEXT(""));
+		UE_VLOG_SEGMENT_THICK(LogOwner, Context.Category, Log, End, End + HeadSize * Left, Color, (int16)Thickness, TEXT(""));
+		UE_VLOG_SEGMENT_THICK(LogOwner, Context.Category, Log, End, End + HeadSize * Right, Color, (int16)Thickness, TEXT(""));
 
 		if (UseDrawDebugHelper() && Context.World)
 		{
@@ -77,13 +76,13 @@ namespace UE::MassNavigation::Debug
 
 	void DebugDrawSphere(const FDebugContext& Context, const FVector& Center, const FVector::FReal InRadius, const FColor& Color)
 	{
-		if (!DebugIsSelected(Context.Entity))
+		if (!Context.ShouldLogEntity())
 		{
 			return;
 		}
 
 		const float Radius = static_cast<float>(InRadius);
-		UE_VLOG_LOCATION(Context.LogOwner, Context.Category, Log, Center, Radius, Color, TEXT(""));
+		UE_VLOG_LOCATION(Context.GetLogOwner(), Context.Category, Log, Center, Radius, Color, TEXT(""));
 
 		if (UseDrawDebugHelper() && Context.World)
 		{
@@ -93,12 +92,12 @@ namespace UE::MassNavigation::Debug
 
 	void DebugDrawBox(const FDebugContext& Context, const FBox& Box, const FColor& Color)
 	{
-		if (!DebugIsSelected(Context.Entity))
+		if (!Context.ShouldLogEntity())
 		{
 			return;
 		}
 
-		UE_VLOG_BOX(Context.LogOwner, Context.Category, Log, Box, Color, TEXT(""));
+		UE_VLOG_BOX(Context.GetLogOwner(), Context.Category, Log, Box, Color, TEXT(""));
 		
 		if (UseDrawDebugHelper() && Context.World)
 		{
@@ -108,13 +107,13 @@ namespace UE::MassNavigation::Debug
 	
 	void DebugDrawCylinder(const FDebugContext& Context, const FVector& Bottom, const FVector& Top, const FVector::FReal InRadius, const FColor& Color, const FString& Text /*= FString()*/)
 	{
-		if (!DebugIsSelected(Context.Entity))
+		if (!Context.ShouldLogEntity())
 		{
 			return;
 		}
 
 		const float Radius = static_cast<float>(InRadius);
-		UE_VLOG_CYLINDER(Context.LogOwner, Context.Category, Log, Bottom, Top, Radius, Color, TEXT("%s"), *Text);
+		UE_VLOG_CYLINDER(Context.GetLogOwner(), Context.Category, Log, Bottom, Top, Radius, Color, TEXT("%s"), *Text);
 
 		if (UseDrawDebugHelper() && Context.World)
 		{
@@ -124,13 +123,13 @@ namespace UE::MassNavigation::Debug
 	
 	void DebugDrawCircle(const FDebugContext& Context, const FVector& Bottom, const FVector::FReal InRadius, const FColor& Color, const FString& Text /*= FString()*/)
 	{
-		if (!DebugIsSelected(Context.Entity))
+		if (!Context.ShouldLogEntity())
 		{
 			return;
 		}
 
 		const float Radius = static_cast<float>(InRadius);
-		UE_VLOG_CIRCLE_THICK(Context.LogOwner, Context.Category, Log, Bottom, FVector(0,0,1), Radius, Color, /*Thickness*/2, TEXT("%s"), *Text);
+		UE_VLOG_CIRCLE_THICK(Context.GetLogOwner(), Context.Category, Log, Bottom, FVector(0,0,1), Radius, Color, /*Thickness*/2, TEXT("%s"), *Text);
 
 		if (UseDrawDebugHelper() && Context.World)
 		{
diff --git a/Engine/Plugins/AI/MassAI/Source/MassNavigation/Private/MassSteeringProcessors.cpp b/Engine/Plugins/AI/MassAI/Source/MassNavigation/Private/MassSteeringProcessors.cpp
--- a/Engine/Plugins/AI/MassAI/Source/MassNavigation/Private/MassSteeringProcessors.cpp
+++ b/Engine/Plugins/AI/MassAI/Source/MassNavigation/Private/MassSteeringProcessors.cpp
@@ -3,6 +3,7 @@
 #include "Steering/MassSteeringProcessors.h"
 #include "MassCommonUtils.h"
 #include "MassCommandBuffer.h"
+#include "MassDebugLogging.h"
 #include "MassCommonFragments.h"
 #include "MassDebugger.h"
 #include "MassExecutionContext.h"
@@ -72,6 +73,10 @@ void UMassSteerToMoveTargetProcessor::ConfigureQueries(const TSharedRef<FMassEnt
 	EntityQuery.AddConstSharedRequirement<FMassMovementParameters>(EMassFragmentPresence::All);
 	EntityQuery.AddConstSharedRequirement<FMassMovingSteeringParameters>(EMassFragmentPresence::All);
 	EntityQuery.AddConstSharedRequirement<FMassStandingSteeringParameters>(EMassFragmentPresence::All);
+	
+#if WITH_MASSGAMEPLAY_DEBUG
+	EntityQuery.DebugEnableEntityOwnerLogging();
+#endif
 
 	// No need for Off LOD to do steering, applying move target directly
 	EntityQuery.AddTagRequirement<FMassOffLODTag>(EMassFragmentPresence::None);
@@ -79,7 +84,10 @@ void UMassSteerToMoveTargetProcessor::ConfigureQueries(const TSharedRef<FMassEnt
 
 void UMassSteerToMoveTargetProcessor::Execute(FMassEntityManager& EntityManager, FMassExecutionContext& Context)
 {
-	EntityQuery.ForEachEntityChunk(Context, [this](FMassExecutionContext& Context)
+	const UWorld* World = EntityManager.GetWorld();
+	ensure(World);
+	
+	EntityQuery.ForEachEntityChunk(Context, [this, World](FMassExecutionContext& Context)
 	{
 		const TConstArrayView<FTransformFragment> TransformList = Context.GetFragmentView<FTransformFragment>();
 		const TArrayView<FMassMoveTargetFragment> MoveTargetList = Context.GetMutableFragmentView<FMassMoveTargetFragment>();
@@ -114,6 +122,12 @@ void UMassSteerToMoveTargetProcessor::Execute(FMassEntityManager& EntityManager,
 
 			const FVector::FReal LookAheadDistance = FMath::Max(1.0f, MovingSteeringParams.LookAheadTime * MoveTarget.DesiredSpeed.Get());
 
+#if WITH_MASSGAMEPLAY_DEBUG
+			UE::MassNavigation::Debug::FDebugContext NavigationDebugContext(Context, this, LogMassNavigation, World, Entity, EntityIt); 
+			const bool bDisplayDebug = NavigationDebugContext.ShouldLogEntity();
+			const UObject* LogOwner = NavigationDebugContext.GetLogOwner();
+#endif
+
 			if (MoveTarget.GetCurrentAction() == EMassMovementAction::Move)
 			{
 				// Tune down avoidance and speed when arriving at goal.
@@ -147,13 +161,11 @@ void UMassSteerToMoveTargetProcessor::Execute(FMassEntityManager& EntityManager,
 				}
 
 #if WITH_MASSGAMEPLAY_DEBUG
-				FColor EntityColor = FColor::White;
-				const bool bDisplayDebug = UE::Mass::Debug::IsDebuggingEntity(Entity, &EntityColor);
 				if (bDisplayDebug)
 				{
 					// Display SteerDirection
 					const FVector ZOffset(0,0,25);
-					UE_VLOG_SEGMENT_THICK(this, LogMassNavigation, Log, CurrentLocation + ZOffset, CurrentLocation + ZOffset + 100.*SteerDirection,
+					UE_VLOG_SEGMENT_THICK(LogOwner, LogMassNavigation, Log, CurrentLocation + ZOffset, CurrentLocation + ZOffset + 100.*SteerDirection,
 						FColor::Red, /*Thickness*/2, TEXT("SteerDirection"));
 				}
 #endif //WITH_MASSGAMEPLAY_DEBUG
@@ -270,11 +282,9 @@ void UMassSteerToMoveTargetProcessor::Execute(FMassEntityManager& EntityManager,
 				if (Distance > StandingSteeringParams.DeadZoneRadius)
 				{
 #if WITH_MASSGAMEPLAY_DEBUG
-					FColor EntityColor = FColor::White;
-					const bool bDisplayDebug = UE::Mass::Debug::IsDebuggingEntity(Entity, &EntityColor);
 					if (bDisplayDebug)
 					{
-						UE_VLOG_UELOG(this, LogMassNavigation, Verbose, TEXT("Standing steering: out of deadzone (Distance: %.2f)"), Distance);
+						UE_VLOG_UELOG(LogOwner, LogMassNavigation, Verbose, TEXT("Standing steering: out of deadzone (Distance: %.2f)"), Distance);
 					}
 #endif
 					
@@ -322,48 +332,46 @@ void UMassSteerToMoveTargetProcessor::Execute(FMassEntityManager& EntityManager,
 
 #if WITH_MASSGAMEPLAY_DEBUG
 			FColor EntityColor = FColor::White;
-			const bool bDisplayDebug = UE::Mass::Debug::IsDebuggingEntity(Entity, &EntityColor);
 			if (bDisplayDebug)
 			{
 				const FVector ZOffset(0,0,25);
 
-				const FColor DarkEntityColor = UE::MassNavigation::Debug::MixColors(EntityColor, FColor::Black);
 				const FColor LightEntityColor = UE::MassNavigation::Debug::MixColors(EntityColor, FColor::White);
 				
 				const FVector MoveTargetCenter = MoveTarget.Center + ZOffset;
 
 				// Display MoveTarget location
-				UE_VLOG_CIRCLE_THICK(this, LogMassNavigation, Log, MoveTargetCenter, FVector::UpVector, 5.f, EntityColor, /*Thickness*/2,
+				UE_VLOG_CIRCLE_THICK(LogOwner, LogMassNavigation, Log, MoveTargetCenter, FVector::UpVector, 5.f, EntityColor, /*Thickness*/2,
 					TEXT("MoveTarget\ngoal: %s"), *UEnum::GetDisplayValueAsText(MoveTarget.IntentAtGoal).ToString());
 
 				// Display ghost location
-				UE_VLOG_CIRCLE_THICK(this, LogMassNavigation, Log, Ghost.Location, FVector::UpVector, 5.f, FColor::Silver, /*Thickness*/2, TEXT("Ghost"));
+				UE_VLOG_CIRCLE_THICK(LogOwner, LogMassNavigation, Log, Ghost.Location, FVector::UpVector, 5.f, FColor::Silver, /*Thickness*/2, TEXT("Ghost"));
 
 				// Display MoveTarget orientation
-				UE_VLOG_SEGMENT_THICK(this, LogMassNavigation, Log, MoveTargetCenter, MoveTargetCenter + MoveTarget.Forward * 100, EntityColor, /*Thickness*/1,
+				UE_VLOG_SEGMENT_THICK(LogOwner, LogMassNavigation, Log, MoveTargetCenter, MoveTargetCenter + MoveTarget.Forward * 100, EntityColor, /*Thickness*/1,
 					TEXT("MoveTarget\nforward"));
 
 				// Display MoveTarget - current location relation
 				if (MoveTarget.DesiredSpeed.Get() > 0.f && FVector::Dist2D(CurrentLocation, MoveTarget.Center) > LookAheadDistance * 1.5f)
 				{
-					UE_VLOG_SEGMENT_THICK(this, LogMassNavigation, Log, MoveTargetCenter, CurrentLocation + ZOffset, FColor::Red, /*Thickness*/1, TEXT("LOST"));
+					UE_VLOG_SEGMENT_THICK(LogOwner, LogMassNavigation, Log, MoveTargetCenter, CurrentLocation + ZOffset, FColor::Red, /*Thickness*/1, TEXT("LOST"));
 				}
 				else
 				{
-					UE_VLOG_SEGMENT_THICK(this, LogMassNavigation, Log, MoveTargetCenter, CurrentLocation + ZOffset, LightEntityColor, /*Thickness*/1, TEXT(""));
+					UE_VLOG_SEGMENT_THICK(LogOwner, LogMassNavigation, Log, MoveTargetCenter, CurrentLocation + ZOffset, LightEntityColor, /*Thickness*/1, TEXT(""));
 				}
 
 				// Display Desired Velocity
-				UE_VLOG_SEGMENT_THICK(this, LogMassNavigation, Log, CurrentLocation + ZOffset, CurrentLocation + ZOffset + DesiredMovement.DesiredVelocity, FColor::Yellow, /*Thickness*/4,
+				UE_VLOG_SEGMENT_THICK(LogOwner, LogMassNavigation, Log, CurrentLocation + ZOffset, CurrentLocation + ZOffset + DesiredMovement.DesiredVelocity, FColor::Yellow, /*Thickness*/4,
 					TEXT("DesiredVelocity %.1f"), DesiredMovement.DesiredVelocity.Length());
 
 				// Display DesiredVelocity
-				UE_VLOG_SEGMENT_THICK(this, LogMassNavigation, Log, CurrentLocation + ZOffset + FVector(0,0,2),
+				UE_VLOG_SEGMENT_THICK(LogOwner, LogMassNavigation, Log, CurrentLocation + ZOffset + FVector(0,0,2),
 					CurrentLocation + Steering.DesiredVelocity + ZOffset + FVector(0,0,2), FColor::Orange, /*Thickness*/4,
 					TEXT("DesiredVelocity %.1f"), Steering.DesiredVelocity.Length());
 
 				// Display Force
-				UE_VLOG_SEGMENT_THICK(this, LogMassNavigation, Log, CurrentLocation + ZOffset, CurrentLocation + Force.Value + ZOffset, FColor::Emerald, /*Thickness*/4,
+				UE_VLOG_SEGMENT_THICK(LogOwner, LogMassNavigation, Log, CurrentLocation + ZOffset, CurrentLocation + Force.Value + ZOffset, FColor::Emerald, /*Thickness*/4,
 					TEXT("Steering Force %.1f"), Force.Value.Length());
 			}
 #endif // WITH_MASSGAMEPLAY_DEBUG
diff --git a/Engine/Plugins/AI/MassAI/Source/MassNavigation/Public/MassNavigationDebug.h b/Engine/Plugins/AI/MassAI/Source/MassNavigation/Public/MassNavigationDebug.h
--- a/Engine/Plugins/AI/MassAI/Source/MassNavigation/Public/MassNavigationDebug.h
+++ b/Engine/Plugins/AI/MassAI/Source/MassNavigation/Public/MassNavigationDebug.h
@@ -3,25 +3,49 @@
 #pragma once
 
 #include "MassCommonTypes.h"
-#include "MassEntityHandle.h"
 
 #if WITH_MASSGAMEPLAY_DEBUG
+#include "MassDebugLogging.h"
+#include "MassEntityHandle.h"
+
+struct FMassExecutionContext;
 
 namespace UE::MassNavigation::Debug
 {
+	/** If true, will log all debug draw events regardless of debug entity selection if the visual log recorder is activated */
+	MASSNAVIGATION_API bool ShouldLogEverythingWhenRecording();
+	
 	struct MASSNAVIGATION_API FDebugContext
 	{
-		FDebugContext(const UObject* InLogOwner, const FLogCategoryBase& InCategory, const UWorld* InWorld, const FMassEntityHandle InEntity)
-			: LogOwner(InLogOwner)
+		FDebugContext(const FMassExecutionContext& Context, const UObject* InFallbackLogOwner, const FLogCategoryBase& InCategory, const UWorld* InWorld, const FMassEntityHandle InEntity, int32 InEntityIndex)
+			: LogContext(Context, ShouldLogEverythingWhenRecording())
+			, FallbackLogOwner(InFallbackLogOwner)
 			, Category(InCategory)
 			, World(InWorld)
 			, Entity(InEntity)
+			, EntityIndex(InEntityIndex)
 		{}
 
-		const UObject* LogOwner;
+		FORCEINLINE const UObject* GetLogOwner() const
+		{
+			return LogContext.GetLogOwner(EntityIndex, FallbackLogOwner);
+		}
+		
+		FORCEINLINE bool ShouldLogEntity(FColor* OutEntityColor = nullptr) const
+		{
+			return LogContext.ShouldLogEntity(EntityIndex, OutEntityColor);
+		}
+
+		const Mass::Debug::FLoggingContext LogContext;
+		
+		/** If entity has no debug owner set, use this instead. Usually the mass processor */
+		const UObject* FallbackLogOwner;
 		const FLogCategoryBase& Category;
 		const UWorld* World;
 		const FMassEntityHandle Entity;
+		
+		/** EntityIndex is index in current chunk */
+		const int32 EntityIndex;
 	};
 	
 	MASSNAVIGATION_API bool UseDrawDebugHelper();
diff --git a/Engine/Plugins/AI/MassAI/Source/MassZoneGraphNavigation/Private/MassZoneGraphNavigationProcessors.cpp b/Engine/Plugins/AI/MassAI/Source/MassZoneGraphNavigation/Private/MassZoneGraphNavigationProcessors.cpp
--- a/Engine/Plugins/AI/MassAI/Source/MassZoneGraphNavigation/Private/MassZoneGraphNavigationProcessors.cpp
+++ b/Engine/Plugins/AI/MassAI/Source/MassZoneGraphNavigation/Private/MassZoneGraphNavigationProcessors.cpp
@@ -6,6 +6,7 @@
 #include "MassNavigationFragments.h"
 #include "MassNavigationUtils.h"
 #include "Avoidance/MassAvoidanceFragments.h"
+#include "Avoidance/MassAvoidanceProcessors.h"
 #include "MassCommonFragments.h"
 #include "MassExecutionContext.h"
 #include "MassSignalSubsystem.h"
@@ -17,6 +18,7 @@
 #include "MassDebugger.h"
 #include "MassNavigationDebug.h"
 
+#define WITH_ZONEGRAPH_DEBUG_DETAILED 0
 
 //----------------------------------------------------------------------//
 //  UMassZoneGraphLocationInitializer
@@ -121,6 +123,10 @@ void UMassZoneGraphPathFollowProcessor::ConfigureQueries(const TSharedRef<FMassE
 	EntityQuery_Conditional.SetChunkFilter(&FMassSimulationVariableTickChunkFragment::ShouldTickChunkThisFrame);
 
 	EntityQuery_Conditional.AddSubsystemRequirement<UZoneGraphSubsystem>(EMassFragmentAccess::ReadOnly);
+
+#if WITH_MASSGAMEPLAY_DEBUG
+	EntityQuery_Conditional.DebugEnableEntityOwnerLogging();
+#endif
 }
 
 void UMassZoneGraphPathFollowProcessor::Execute(FMassEntityManager& EntityManager, FMassExecutionContext& Context)
@@ -129,11 +135,14 @@ void UMassZoneGraphPathFollowProcessor::Execute(FMassEntityManager& EntityManage
 	{
 		return;
 	}
+
+	const UWorld* World = EntityManager.GetWorld();
+	ensure(World);
 	
 	TArray<FMassEntityHandle> EntitiesToSignalPathDone;
 	TArray<FMassEntityHandle> EntitiesToSignalLaneChanged;
 
-	EntityQuery_Conditional.ForEachEntityChunk(Context, [this, &EntitiesToSignalPathDone, &EntitiesToSignalLaneChanged](FMassExecutionContext& Context)
+	EntityQuery_Conditional.ForEachEntityChunk(Context, [this, &EntitiesToSignalPathDone, &EntitiesToSignalLaneChanged, World](FMassExecutionContext& Context)
 	{
 		const UZoneGraphSubsystem& ZoneGraphSubsystem = Context.GetSubsystemChecked<UZoneGraphSubsystem>();
 
@@ -156,10 +165,12 @@ void UMassZoneGraphPathFollowProcessor::Execute(FMassEntityManager& EntityManage
 
 			bool bDisplayDebug = false;
 #if WITH_MASSGAMEPLAY_DEBUG // this will result in bDisplayDebug == false and disabling of all the vlogs below
-			bDisplayDebug = UE::Mass::Debug::IsDebuggingEntity(Entity);
+			const UE::MassNavigation::Debug::FDebugContext NavigationDebugContext(Context, this, LogMassNavigation, World, Entity, EntityIt);
+			bDisplayDebug = NavigationDebugContext.ShouldLogEntity();
+			const UObject* LogOwner = NavigationDebugContext.GetLogOwner();
 			if (bDisplayDebug)
 			{
-				UE_VLOG(this, LogMassNavigation, Log, TEXT("Entity [%s] Updating path following"), *Entity.DebugGetDescription());
+				UE_VLOG(LogOwner, LogMassNavigation, Log, TEXT("Entity [%s] Updating path following"), *Entity.DebugGetDescription());
 			}
 #endif // WITH_MASSGAMEPLAY_DEBUG
 
@@ -197,7 +208,7 @@ void UMassZoneGraphPathFollowProcessor::Execute(FMassEntityManager& EntityManage
 						MoveTarget.DistanceToGoal = ShortPath.Points[LastPointIndex].Distance.Get();
 						MoveTarget.bOffBoundaries = ShortPath.Points[0].bOffLane;
 
-						UE_CVLOG(bDisplayDebug,this, LogMassNavigation, Verbose, TEXT("Entity [%s] before start of lane %s at distance %.1f. Distance to goal: %.1f. Off Boundaries: %s"),
+						UE_CVLOG(bDisplayDebug, LogOwner, LogMassNavigation, Verbose, TEXT("Entity [%s] before start of lane %s at distance %.1f. Distance to goal: %.1f. Off Boundaries: %s"),
 							*Entity.DebugGetDescription(),
 							*LaneLocation.LaneHandle.ToString(),
 							LaneLocation.DistanceAlongLane,
@@ -229,7 +240,7 @@ void UMassZoneGraphPathFollowProcessor::Execute(FMassEntityManager& EntityManage
 						MoveTarget.DistanceToGoal = ShortPath.Points[LastPointIndex].Distance.Get() - FMath::Lerp(CurrPoint.Distance.Get(), NextPoint.Distance.Get(), T);
 						MoveTarget.bOffBoundaries = CurrPoint.bOffLane || NextPoint.bOffLane;
 
-						UE_CVLOG(bDisplayDebug, this, LogMassNavigation, Verbose, TEXT("Entity [%s] along lane %s at distance %.1f. Distance to goal: %.1f. Off Boundaries: %s"),
+						UE_CVLOG(bDisplayDebug, LogOwner, LogMassNavigation, Verbose, TEXT("Entity [%s] along lane %s at distance %.1f. Distance to goal: %.1f. Off Boundaries: %s"),
 							*Entity.DebugGetDescription(),
 							*LaneLocation.LaneHandle.ToString(),
 							LaneLocation.DistanceAlongLane,
@@ -246,12 +257,12 @@ void UMassZoneGraphPathFollowProcessor::Execute(FMassEntityManager& EntityManage
 						MoveTarget.DistanceToGoal = 0.0f;
 						MoveTarget.bOffBoundaries = ShortPath.Points[LastPointIndex].bOffLane;
 
-						UE_CVLOG(bDisplayDebug, this, LogMassNavigation, Log, TEXT("Entity [%s] Finished path follow on lane %s at distance %f. Off Boundaries: %s"),
+						UE_CVLOG(bDisplayDebug, LogOwner, LogMassNavigation, Log, TEXT("Entity [%s] Finished path follow on lane %s at distance %f. Off Boundaries: %s"),
 							*Entity.DebugGetDescription(), *LaneLocation.LaneHandle.ToString(), LaneLocation.DistanceAlongLane, *LexToString((bool)MoveTarget.bOffBoundaries));
 
 						if (bDisplayDebug)
 						{
-							UE_VLOG(this, LogMassNavigation, Log, TEXT("Entity [%s] End of path."), *Entity.DebugGetDescription());
+							UE_VLOG(LogOwner, LogMassNavigation, Log, TEXT("Entity [%s] End of path."), *Entity.DebugGetDescription());
 						}
 
 						// Check to see if need advance to next lane.
@@ -265,7 +276,7 @@ void UMassZoneGraphPathFollowProcessor::Execute(FMassEntityManager& EntityManage
 									float NewLaneLength = 0.f;
 									UE::ZoneGraph::Query::GetLaneLength(*ZoneGraphStorage, ShortPath.NextLaneHandle, NewLaneLength);
 
-									UE_CVLOG(bDisplayDebug, this, LogMassNavigation, Log, TEXT("Entity [%s] Switching to OUTGOING lane %s -> %s, new distance %f."),
+									UE_CVLOG(bDisplayDebug, LogOwner, LogMassNavigation, Log, TEXT("Entity [%s] Switching to OUTGOING lane %s -> %s, new distance %f."),
 										*Entity.DebugGetDescription(), *LaneLocation.LaneHandle.ToString(), *ShortPath.NextLaneHandle.ToString(), 0.f);
 
 									// update lane location
@@ -278,7 +289,7 @@ void UMassZoneGraphPathFollowProcessor::Execute(FMassEntityManager& EntityManage
 									float NewLaneLength = 0.f;
 									UE::ZoneGraph::Query::GetLaneLength(*ZoneGraphStorage, ShortPath.NextLaneHandle, NewLaneLength);
 
-									UE_CVLOG(bDisplayDebug, this, LogMassNavigation, Log, TEXT("Entity [%s] Switching to INCOMING lane %s -> %s, new distance %f."),
+									UE_CVLOG(bDisplayDebug, LogOwner, LogMassNavigation, Log, TEXT("Entity [%s] Switching to INCOMING lane %s -> %s, new distance %f."),
 										*Entity.DebugGetDescription(), *LaneLocation.LaneHandle.ToString(), *ShortPath.NextLaneHandle.ToString(), NewLaneLength);
 
 									// update lane location
@@ -295,7 +306,7 @@ void UMassZoneGraphPathFollowProcessor::Execute(FMassEntityManager& EntityManage
 										float NewLaneLength = 0.f;
 										UE::ZoneGraph::Query::GetLaneLength(*ZoneGraphStorage, ShortPath.NextLaneHandle, NewLaneLength);
 
-										UE_CVLOG(bDisplayDebug, this, LogMassNavigation, Log, TEXT("Entity [%s] Switching to ADJACENT lane %s -> %s, new distance %f."),
+										UE_CVLOG(bDisplayDebug, LogOwner, LogMassNavigation, Log, TEXT("Entity [%s] Switching to ADJACENT lane %s -> %s, new distance %f."),
 											*Entity.DebugGetDescription(), *LaneLocation.LaneHandle.ToString(), *ShortPath.NextLaneHandle.ToString(), NewLocation.DistanceAlongLane);
 
 										// update lane location
@@ -307,7 +318,7 @@ void UMassZoneGraphPathFollowProcessor::Execute(FMassEntityManager& EntityManage
 									}
 									else
 									{
-										UE_CVLOG(bDisplayDebug, this, LogMassNavigation, Error, TEXT("Entity [%s] Failed to switch to ADJACENT lane %s -> %s."),
+										UE_CVLOG(bDisplayDebug, LogOwner, LogMassNavigation, Error, TEXT("Entity [%s] Failed to switch to ADJACENT lane %s -> %s."),
 											*Entity.DebugGetDescription(), *LaneLocation.LaneHandle.ToString(), *ShortPath.NextLaneHandle.ToString());
 									}
 								}
@@ -321,13 +332,13 @@ void UMassZoneGraphPathFollowProcessor::Execute(FMassEntityManager& EntityManage
 							}
 							else
 							{
-								UE_CVLOG(bDisplayDebug, this, LogMassNavigation, Error, TEXT("Entity [%s] Could not find ZoneGraph storage for lane %s."),
+								UE_CVLOG(bDisplayDebug, LogOwner, LogMassNavigation, Error, TEXT("Entity [%s] Could not find ZoneGraph storage for lane %s."),
 									*Entity.DebugGetDescription(), *LaneLocation.LaneHandle.ToString());
 							}
 						}
 						else
 						{
-							UE_CVLOG(bDisplayDebug, this, LogMassNavigation, Log, TEXT("Entity [%s] Next lane not defined."), *Entity.DebugGetDescription());
+							UE_CVLOG(bDisplayDebug, LogOwner, LogMassNavigation, Log, TEXT("Entity [%s] Next lane not defined."), *Entity.DebugGetDescription());
 						}
 
 						ShortPath.bDone = true;
@@ -356,7 +367,7 @@ void UMassZoneGraphPathFollowProcessor::Execute(FMassEntityManager& EntityManage
 						const FMassZoneGraphPathPoint& NextPoint = ShortPath.Points[PointIndex + 1];
 
 						// Path
-						UE_VLOG_SEGMENT_THICK(this, LogMassNavigation, Display, CurrPoint.Position + ZOffset, NextPoint.Position + ZOffset, EntityColor, /*Thickness*/3, TEXT(""));
+						UE_VLOG_SEGMENT_THICK(LogOwner, LogMassNavigation, Display, CurrPoint.Position + ZOffset, NextPoint.Position + ZOffset, EntityColor, /*Thickness*/3, TEXT(""));
 					}
 					
 					for (uint8 PointIndex = 0; PointIndex < ShortPath.NumPoints; PointIndex++)
@@ -364,14 +375,14 @@ void UMassZoneGraphPathFollowProcessor::Execute(FMassEntityManager& EntityManage
 						const FMassZoneGraphPathPoint& CurrPoint = ShortPath.Points[PointIndex];
 						const FVector CurrBase = CurrPoint.Position + ZOffset;
 						// Lane tangents
-						UE_VLOG_SEGMENT_THICK(this, LogMassNavigation, Display, CurrBase, CurrBase + CurrPoint.Tangent.GetVector() * 100.0f, LightEntityColor, /*Thickness*/1, TEXT(""));
+						UE_VLOG_SEGMENT_THICK(LogOwner, LogMassNavigation, Display, CurrBase, CurrBase + CurrPoint.Tangent.GetVector() * 100.0f, LightEntityColor, /*Thickness*/1, TEXT(""));
 					}
 
 					if (ShortPath.NumPoints > 0 && ShortPath.NextLaneHandle.IsValid())
 					{
 						const FMassZoneGraphPathPoint& LastPoint = ShortPath.Points[ShortPath.NumPoints - 1];
 						const FVector CurrBase = LastPoint.Position + ZOffset;
-						UE_VLOG_SEGMENT_THICK(this, LogMassNavigation, Display, CurrBase, CurrBase + FVector(0,0,100), FColor::Red, /*Thickness*/3, TEXT("Next: %s"), *ShortPath.NextLaneHandle.ToString());
+						UE_VLOG_SEGMENT_THICK(LogOwner, LogMassNavigation, Display, CurrBase, CurrBase + FVector(0,0,100), FColor::Red, /*Thickness*/3, TEXT("Next: %s"), *ShortPath.NextLaneHandle.ToString());
 					}
 				}
 #endif // WITH_MASSGAMEPLAY_DEBUG
@@ -412,6 +423,10 @@ void UMassZoneGraphLaneCacheBoundaryProcessor::ConfigureQueries(const TSharedRef
 	EntityQuery.AddRequirement<FMassLaneCacheBoundaryFragment>(EMassFragmentAccess::ReadWrite);
 	EntityQuery.AddRequirement<FMassNavigationEdgesFragment>(EMassFragmentAccess::ReadWrite);	// output edges
 	EntityQuery.AddTagRequirement<FMassOffLODTag>(EMassFragmentPresence::None);
+
+#if WITH_MASSGAMEPLAY_DEBUG
+	EntityQuery.DebugEnableEntityOwnerLogging();
+#endif
 }
 
 void UMassZoneGraphLaneCacheBoundaryProcessor::InitializeInternal(UObject& Owner, const TSharedRef<FMassEntityManager>& EntityManager)
@@ -455,9 +470,9 @@ void UMassZoneGraphLaneCacheBoundaryProcessor::Execute(FMassEntityManager& Entit
 			const FVector::FReal DeltaDistSquared = FVector::DistSquared(MovementTarget.Center, LaneCacheBoundary.LastUpdatePosition);
 			const FVector::FReal UpdateDistanceThresholdSquared = FMath::Square(50.);
 
-#if WITH_MASSGAMEPLAY_DEBUG && 0
-			const FDebugContext ObstacleDebugContext(this, LogAvoidanceObstacles, World, Entity);
-			if (DebugIsSelected(Entity))
+#if WITH_MASSGAMEPLAY_DEBUG && WITH_ZONEGRAPH_DEBUG_DETAILED
+			const UE::MassNavigation::Debug::FDebugContext ObstacleDebugContext(Context, this, LogAvoidanceObstacles, World, Entity, EntityIt);
+			if (ObstacleDebugContext.ShouldLogEntity())
 			{
 				DebugDrawSphere(ObstacleDebugContext, LaneCacheBoundary.LastUpdatePosition, /*Radius=*/10.f, FColor(128,128,128));
 				DebugDrawSphere(ObstacleDebugContext, MovementTarget.Center, /*Radius=*/10.f, FColor(255,255,255));
@@ -482,11 +497,8 @@ void UMassZoneGraphLaneCacheBoundaryProcessor::Execute(FMassEntityManager& Entit
 			}
 			
 
-#if WITH_MASSGAMEPLAY_DEBUG && 0
-			if (DebugIsSelected(Entity))
-			{
-				DebugDrawSphere(ObstacleDebugContext, MovementTarget.Center, /*Radius=*/100.f, FColor(128,128,128));
-			}
+#if WITH_MASSGAMEPLAY_DEBUG && WITH_ZONEGRAPH_DEBUG_DETAILED
+			DebugDrawSphere(ObstacleDebugContext, MovementTarget.Center, /*Radius=*/100.f, FColor(128,128,128));
 #endif // WITH_MASSGAMEPLAY_DEBUG
 
 			const float HalfWidth = 0.5f * CachedLane.LaneWidth.Get();
@@ -546,8 +558,8 @@ void UMassZoneGraphLaneCacheBoundaryProcessor::Execute(FMassEntityManager& Entit
 			int32 NumRightPositions = NumPoints;
 
 
-#if 0 && WITH_MASSGAMEPLAY_DEBUG // Detailed debug disabled
-			if (DebugIsSelected(Entity))
+#if WITH_MASSGAMEPLAY_DEBUG && WITH_ZONEGRAPH_DEBUG_DETAILED
+			if (ObstacleDebugContext.ShouldLogEntity())
 			{
 				float Radius = 2.f;
 				for (int32 Index = 0; Index < NumPoints; Index++)
diff --git a/Engine/Source/Runtime/MassEntity/Public/MassDebugLogging.h b/Engine/Source/Runtime/MassEntity/Public/MassDebugLogging.h
--- a/Engine/Source/Runtime/MassEntity/Public/MassDebugLogging.h
+++ b/Engine/Source/Runtime/MassEntity/Public/MassDebugLogging.h
@@ -22,7 +22,7 @@ struct MASSENTITY_API FMassDebugLogFragment : public FMassFragment
 	GENERATED_BODY()
 
 	UPROPERTY()
-	TWeakObjectPtr<UObject> LogOwner = nullptr;
+	TWeakObjectPtr<const UObject> LogOwner = nullptr;
 };
 
 namespace UE::Mass::Debug
@@ -30,19 +30,20 @@ namespace UE::Mass::Debug
 	struct FLoggingContext
 	{
 #if WITH_MASSENTITY_DEBUG
-		explicit FLoggingContext(const FMassExecutionContext& Context, bool bInCheckVisualLoggerForRecording = true)
-			: DebugFragmentsView(Context.GetFragmentView<FMassDebugLogFragment>())
-			, EntityListView(Context.GetEntities())
-			, bCheckVisualLoggerForRecording(bInCheckVisualLoggerForRecording)
+		explicit FLoggingContext(const FMassExecutionContext& InContext, bool bInLogEverythingWhenRecording = true)
+			: DebugFragmentsView(InContext.GetFragmentView<FMassDebugLogFragment>())
+			, EntityListView(InContext.GetEntities())
+			, bLogEverythingWhenRecording(bInLogEverythingWhenRecording)
 		{	
 		}
 
 		FORCEINLINE bool ShouldLogEntity(int32 EntityIndex, FColor* OutEntityColor = nullptr) const
 		{
 #if ENABLE_VISUAL_LOG
-			if (DebugFragmentsView.IsValidIndex(EntityIndex)
-				&& DebugFragmentsView[EntityIndex].LogOwner != nullptr 
-				&& (bCheckVisualLoggerForRecording == false || FVisualLogger::IsRecording()))
+			if (bLogEverythingWhenRecording
+				&& DebugFragmentsView.IsValidIndex(EntityIndex)
+				&& DebugFragmentsView[EntityIndex].LogOwner != nullptr
+				&& FVisualLogger::IsRecording())
 			{
 				return true;
 			}
@@ -52,7 +53,7 @@ namespace UE::Mass::Debug
 				&& IsDebuggingEntity(EntityListView[EntityIndex], OutEntityColor);
 		}
 
-		FORCEINLINE UObject* GetLogOwner(int32 EntityIndex, UObject* FallbackOwner = nullptr) const
+		FORCEINLINE const UObject* GetLogOwner(int32 EntityIndex, const UObject* FallbackOwner = nullptr) const
 		{
 			return DebugFragmentsView.IsValidIndex(EntityIndex)
 				? DebugFragmentsView[EntityIndex].LogOwner.Get()
@@ -62,7 +63,11 @@ namespace UE::Mass::Debug
 	private:
 		const TConstArrayView<FMassDebugLogFragment> DebugFragmentsView;
 		const TConstArrayView<FMassEntityHandle> EntityListView;
-		const bool bCheckVisualLoggerForRecording = true;
+
+		/** If true, ShouldLogEntity will return true when the visual logger is recording
+		 *  If false, ShouldLogEntity will rely only on the MassDebugger
+		 */
+		const bool bLogEverythingWhenRecording = true;
 
 #else
 
