Buildfix in MassZoneGraphNavigationProcessors on enginetesteditorcooked

diff --git a/Engine/Plugins/AI/MassAI/Source/MassZoneGraphNavigation/Private/MassZoneGraphNavigationProcessors.cpp b/Engine/Plugins/AI/MassAI/Source/MassZoneGraphNavigation/Private/MassZoneGraphNavigationProcessors.cpp
--- a/Engine/Plugins/AI/MassAI/Source/MassZoneGraphNavigation/Private/MassZoneGraphNavigationProcessors.cpp
+++ b/Engine/Plugins/AI/MassAI/Source/MassZoneGraphNavigation/Private/MassZoneGraphNavigationProcessors.cpp
@@ -164,10 +164,11 @@ void UMassZoneGraphPathFollowProcessor::Execute(FMassEntityManager& EntityManage
 			const float DeltaTime = bHasVariableTick ? SimVariableTickList[EntityIt].DeltaTime : WorldDeltaTime;
 
 			bool bDisplayDebug = false;
+			const UObject* LogOwner = nullptr;
 #if WITH_MASSGAMEPLAY_DEBUG // this will result in bDisplayDebug == false and disabling of all the vlogs below
 			const UE::MassNavigation::Debug::FDebugContext NavigationDebugContext(Context, this, LogMassNavigation, World, Entity, EntityIt);
 			bDisplayDebug = NavigationDebugContext.ShouldLogEntity();
-			const UObject* LogOwner = NavigationDebugContext.GetLogOwner();
+			LogOwner = NavigationDebugContext.GetLogOwner();
 			if (bDisplayDebug)
 			{
 				UE_VLOG(LogOwner, LogMassNavigation, Log, TEXT("Entity [%s] Updating path following"), *Entity.DebugGetDescription());
