Apply desired max speed override when initiating a navmesh pathfollow.
Includes various fixes to previous merges

diff --git a/Engine/Plugins/AI/MassAI/Source/MassAIBehavior/Private/Tasks/MassNavMeshPathFollowTask.cpp b/Engine/Plugins/AI/MassAI/Source/MassAIBehavior/Private/Tasks/MassNavMeshPathFollowTask.cpp
--- a/Engine/Plugins/AI/MassAI/Source/MassAIBehavior/Private/Tasks/MassNavMeshPathFollowTask.cpp
+++ b/Engine/Plugins/AI/MassAI/Source/MassAIBehavior/Private/Tasks/MassNavMeshPathFollowTask.cpp
@@ -17,6 +17,7 @@ bool FMassNavMeshPathFollowTask::Link(FStateTreeLinker& Linker)
 	Linker.LinkExternalData(TransformHandle);
 	Linker.LinkExternalData(MoveTargetHandle);
 	Linker.LinkExternalData(AgentRadiusHandle);
+	Linker.LinkExternalData(DesiredMovementHandle);
 	Linker.LinkExternalData(MovementParamsHandle);
 	Linker.LinkExternalData(CachedPathHandle);
 	Linker.LinkExternalData(ShortPathHandle);
@@ -65,44 +66,54 @@ bool FMassNavMeshPathFollowTask::RequestPath(FStateTreeExecutionContext& Context
 		Result = Query.NavData->FindPath(NavAgentProperties, Query);
 	}
 
-	if (Result.IsSuccessful() && Result.Path.Get()->GetPathPoints().Num() > 1)
+	if (Result.IsSuccessful())
 	{
-		// A path was found
-		if (bDisplayDebug)
+		// @todo: Investigate single point paths but for now, only move if we have more than one point.
+		if (Result.Path.Get()->GetPathPoints().Num() > 1)
 		{
-			MASSBEHAVIOR_LOG(Verbose, TEXT("path found"));
+			// A path was found
+			if (bDisplayDebug)
+			{
+				MASSBEHAVIOR_LOG(Verbose, TEXT("path found"));
+			}
+
+			FMassNavMeshCachedPathFragment& CachedPathFragment = Context.GetExternalData(CachedPathHandle);
+			CachedPathFragment.NavPath = Result.Path;
+
+			// Build corridor
+			CachedPathFragment.Corridor = MakeShared<FNavCorridor>();
+			const FSharedConstNavQueryFilter NavQueryFilter = Query.QueryFilter ? Query.QueryFilter : NavData->GetDefaultQueryFilter();
+
+			FNavCorridorParams CorridorParams;
+			CorridorParams.SetFromWidth(InstanceData.CorridorWidth);
+			CorridorParams.PathOffsetFromBoundaries = InstanceData.OffsetFromBoundaries;
+
+			CachedPathFragment.Corridor->BuildFromPath(*CachedPathFragment.NavPath, NavQueryFilter, CorridorParams);
+
+			// Update short path
+			FMassNavMeshShortPathFragment& ShortPathFragment = Context.GetExternalData(ShortPathHandle);
+			ShortPathFragment.RequestShortPath(CachedPathFragment.Corridor, /*NavCorridorStartIndex*/0, /*NumLeadingPoints*/0, InstanceData.EndDistanceThreshold);
+
+			CachedPathFragment.NavPathNextStartIndex = (uint16)FMath::Max(ShortPathFragment.NumPoints - FMassNavMeshShortPathFragment::NumPointsBeyondUpdate - FMassNavMeshCachedPathFragment::NumLeadingPoints, 0);
+			
+			// Update MoveTarget
+			FMassMoveTargetFragment& MoveTarget = Context.GetExternalData(MoveTargetHandle);
+			const FMassMovementParameters& MovementParams = Context.GetExternalData(MovementParamsHandle);
+			float DesiredSpeed = FMath::Min(
+				MovementParams.GenerateDesiredSpeed(InstanceData.MovementStyle, MassContext.GetEntity().Index) * InstanceData.SpeedScale,
+				MovementParams.MaxSpeed);
+
+			// Apply DesiredMaxSpeedOverride
+			const FMassDesiredMovementFragment& DesiredMovementFragment = Context.GetExternalData(DesiredMovementHandle);
+			DesiredSpeed = FMath::Min(DesiredSpeed, DesiredMovementFragment.DesiredMaxSpeedOverride);
+
+			MoveTarget.DesiredSpeed.Set(DesiredSpeed);
+			
+			MoveTarget.CreateNewAction(EMassMovementAction::Move, *Context.GetWorld());
+
+			return true;
 		}
-
-		FMassNavMeshCachedPathFragment& CachedPathFragment = Context.GetExternalData(CachedPathHandle);
-		CachedPathFragment.NavPath = Result.Path;
-
-		// Build corridor
-		CachedPathFragment.Corridor = MakeShared<FNavCorridor>();
-		const FSharedConstNavQueryFilter NavQueryFilter = Query.QueryFilter ? Query.QueryFilter : NavData->GetDefaultQueryFilter();
-
-		FNavCorridorParams CorridorParams;
-		CorridorParams.SetFromWidth(InstanceData.CorridorWidth);
-		CorridorParams.PathOffsetFromBoundaries = InstanceData.OffsetFromBoundaries;
-
-		CachedPathFragment.Corridor->BuildFromPath(*CachedPathFragment.NavPath, NavQueryFilter, CorridorParams);
-
-		// Update short path
-		FMassNavMeshShortPathFragment& ShortPathFragment = Context.GetExternalData(ShortPathHandle);
-		ShortPathFragment.RequestShortPath(CachedPathFragment.Corridor, /*NavCorridorStartIndex*/0, /*NumLeadingPoints*/0, InstanceData.EndDistanceThreshold);
-
-		CachedPathFragment.NavPathNextStartIndex = (uint16)FMath::Max(ShortPathFragment.NumPoints - FMassNavMeshShortPathFragment::NumPointsBeyondUpdate - FMassNavMeshCachedPathFragment::NumLeadingPoints, 0);
-		
-		// Update MoveTarget
-		FMassMoveTargetFragment& MoveTarget = Context.GetExternalData(MoveTargetHandle);
-		const FMassMovementParameters& MovementParams = Context.GetExternalData(MovementParamsHandle);
-		const float DesiredSpeed = FMath::Min(
-			MovementParams.GenerateDesiredSpeed(InstanceData.MovementStyle, MassContext.GetEntity().Index) * InstanceData.SpeedScale,
-			MovementParams.MaxSpeed);
-
-		MoveTarget.DesiredSpeed.Set(DesiredSpeed);
-		
-		MoveTarget.CreateNewAction(EMassMovementAction::Move, *Context.GetWorld());
-
+	
 		return true;
 	}
 	
diff --git a/Engine/Plugins/AI/MassAI/Source/MassAIBehavior/Public/Tasks/MassNavMeshPathfollowTask.h b/Engine/Plugins/AI/MassAI/Source/MassAIBehavior/Public/Tasks/MassNavMeshPathfollowTask.h
--- a/Engine/Plugins/AI/MassAI/Source/MassAIBehavior/Public/Tasks/MassNavMeshPathfollowTask.h
+++ b/Engine/Plugins/AI/MassAI/Source/MassAIBehavior/Public/Tasks/MassNavMeshPathfollowTask.h
@@ -13,6 +13,7 @@ struct FTransformFragment;
 struct FMassMoveTargetFragment;
 struct FAgentRadiusFragment;
 struct FMassMovementParameters;
+struct FMassDesiredMovementFragment;
 
 /** FMassNavMeshPathFollowTask movement parameters */
 USTRUCT()
@@ -61,6 +62,7 @@ protected:
 	TStateTreeExternalDataHandle<FTransformFragment> TransformHandle;
 	TStateTreeExternalDataHandle<FMassMoveTargetFragment> MoveTargetHandle;
 	TStateTreeExternalDataHandle<FAgentRadiusFragment> AgentRadiusHandle;
+	TStateTreeExternalDataHandle<FMassDesiredMovementFragment> DesiredMovementHandle;
 	TStateTreeExternalDataHandle<FMassMovementParameters> MovementParamsHandle;
 
 	// Hold a small part of a navmesh path
diff --git a/Engine/Plugins/Runtime/MassGameplay/Source/MassMovement/Public/MassMovementFragments.h b/Engine/Plugins/Runtime/MassGameplay/Source/MassMovement/Public/MassMovementFragments.h
--- a/Engine/Plugins/Runtime/MassGameplay/Source/MassMovement/Public/MassMovementFragments.h
+++ b/Engine/Plugins/Runtime/MassGameplay/Source/MassMovement/Public/MassMovementFragments.h
@@ -38,6 +38,8 @@ struct MASSMOVEMENT_API FMassDesiredMovementFragment : public FMassFragment
 	
 	FVector DesiredVelocity = FVector::ZeroVector;
 
+	float DesiredMaxSpeedOverride = FLT_MAX;
+	
 	// Todo: Could add a DesiredFacing here
 };
 
