[LookAt] added optional target location offset

diff --git a/Engine/Plugins/AI/MassAI/Source/MassAIBehavior/Private/MassLookAtProcessors.cpp b/Engine/Plugins/AI/MassAI/Source/MassAIBehavior/Private/MassLookAtProcessors.cpp
--- a/Engine/Plugins/AI/MassAI/Source/MassAIBehavior/Private/MassLookAtProcessors.cpp
+++ b/Engine/Plugins/AI/MassAI/Source/MassAIBehavior/Private/MassLookAtProcessors.cpp
@@ -20,6 +20,7 @@
 #include "Engine/World.h"
 #include "MassLODFragments.h"
 #include "MassDebugger.h"
+#include "MassLookAtSettings.h"
 #include "MassLookAtSubsystem.h"
 
 
@@ -299,7 +300,7 @@ void UMassLookAtProcessor::FindNewGazeTarget(const UMassNavigationSubsystem& Mas
 
 				// TargetTag is added through the LookAtTargetTrait and Transform was added with it
 				const FTransformFragment& TargetTransform = EntityView.GetFragmentData<FTransformFragment>();
-				const FVector TargetLocation = TargetTransform.GetTransform().GetLocation();
+				const FVector TargetLocation = TargetTransform.GetTransform().GetLocation() + GetDefault<UMassLookAtSettings>()->GetDefaultTargetLocationOffset();
 				FVector Direction = (TargetLocation - Location).GetSafeNormal();
 				Direction = InTransform.InverseTransformVector(Direction);
 
@@ -415,7 +416,7 @@ void UMassLookAtProcessor::UpdateLookAtTrackedEntity(const FMassEntityManager& E
 			const FVector TargetLocation = TargetTransform->GetTransform().GetLocation();
 			const FVector NewGlobalDirection = (TargetLocation - AgentPosition).GetSafeNormal();
 			LookAt.Direction = Transform.InverseTransformVector(NewGlobalDirection);
-			LookAt.MainTargetLocation = TargetLocation;
+			LookAt.MainTargetLocation = TargetLocation + GetDefault<UMassLookAtSettings>()->GetDefaultTargetLocationOffset();
 
 #if WITH_MASSGAMEPLAY_DEBUG
 			if (bDisplayDebug)
@@ -441,7 +442,7 @@ bool UMassLookAtProcessor::UpdateGazeTrackedEntity(const FMassEntityManager& Ent
 			const FVector TargetLocation = TargetTransform->GetTransform().GetLocation();
 			const FVector NewGlobalDirection = (TargetLocation - AgentPosition).GetSafeNormal();
 			LookAt.GazeDirection = Transform.InverseTransformVector(NewGlobalDirection);
-			LookAt.GazeTargetLocation = TargetLocation;
+			LookAt.GazeTargetLocation = TargetLocation + GetDefault<UMassLookAtSettings>()->GetDefaultTargetLocationOffset();
 			bHasTarget = true;
 			
 #if WITH_MASSGAMEPLAY_DEBUG
diff --git a/Engine/Plugins/AI/MassAI/Source/MassAIBehavior/Public/MassLookAtSettings.h b/Engine/Plugins/AI/MassAI/Source/MassAIBehavior/Public/MassLookAtSettings.h
--- a/Engine/Plugins/AI/MassAI/Source/MassAIBehavior/Public/MassLookAtSettings.h
+++ b/Engine/Plugins/AI/MassAI/Source/MassAIBehavior/Public/MassLookAtSettings.h
@@ -24,10 +24,20 @@ public:
 
 	MASSAIBEHAVIOR_API void GetValidPriorityInfos(TArray<FMassLookAtPriorityInfo>& OutInfos) const;
 
+	/** @return Additional offset to add to the target location (aka eyes height) */
+	[[nodiscard]] FVector GetDefaultTargetLocationOffset() const
+	{
+		return DefaultTargetLocationOffset;
+	}
+
 	static MASSAIBEHAVIOR_API FOnMassLookAtPrioritiesChanged OnMassLookAtPrioritiesChanged;
 
 protected:
 
 	UPROPERTY(EditAnywhere, config, Category = LookAt)
 	FMassLookAtPriorityInfo Priorities[static_cast<int32>(EMassLookAtPriorities::MaxPriorities)];
+
+	/** Additional offset to add to the target location (aka eyes height) */
+	UPROPERTY(EditAnywhere, config, Category = LookAt)
+	FVector DefaultTargetLocationOffset = FVector::ZeroVector;
 };
\ No newline at end of file
diff --git a/Engine/Plugins/AI/MassAI/Source/MassAIDebug/Private/GameplayDebuggerCategory_Mass.cpp b/Engine/Plugins/AI/MassAI/Source/MassAIDebug/Private/GameplayDebuggerCategory_Mass.cpp
--- a/Engine/Plugins/AI/MassAI/Source/MassAIDebug/Private/GameplayDebuggerCategory_Mass.cpp
+++ b/Engine/Plugins/AI/MassAI/Source/MassAIDebug/Private/GameplayDebuggerCategory_Mass.cpp
@@ -16,6 +16,7 @@
 #include "MassNavigationFragments.h"
 #include "Steering/MassSteeringFragments.h"
 #include "MassLookAtFragments.h"
+#include "MassLookAtSettings.h"
 #include "MassLookAtSubsystem.h"
 #include "MassStateTreeFragments.h"
 #include "MassStateTreeExecutionContext.h"
@@ -583,18 +584,19 @@ void FGameplayDebuggerCategory_Mass::CollectData(APlayerController* OwnerPC, AAc
 
 					const FVector EntityForward = Transform.GetTransform().GetRotation().GetForwardVector();
 
-					constexpr float EyeHeight = 160.0f; // @todo: add eye height to agent.
+					const FVector TargetOffset = GetDefault<UMassLookAtSettings>()->GetDefaultTargetLocationOffset();
+					const double EyeHeight = TargetOffset.IsZero() ? 160 : TargetOffset.Z;
 
 					// Draw entity position and orientation.
-					FVector BasePos = EntityLocation + FVector(0.0f ,0.0f ,25.0f );
+					FVector BasePos = EntityLocation + FVector(0 ,0 ,25);
 
 					AddShape(FGameplayDebuggerShape::MakeCircle(BasePos, FVector::UpVector, Radius.Radius, FColor::White));
 					AddShape(FGameplayDebuggerShape::MakeSegment(BasePos, BasePos + EntityForward * Radius.Radius * 1.25f, FColor::White));
 
 					// Velocity and steering target
-					BasePos += FVector(0.0f ,0.0f ,5.0f );
+					BasePos += FVector(0 ,0 ,5);
 					AddShape(FGameplayDebuggerShape::MakeArrow(BasePos, BasePos + Velocity.Value, 10.0f, 2.0f, FColor::Yellow));
-					BasePos += FVector(0.0f ,0.0f ,5.0f );
+					BasePos += FVector(0 ,0 ,5);
 					AddShape(FGameplayDebuggerShape::MakeArrow(BasePos, BasePos + Steering.DesiredVelocity, 10.0f, 1.0f, FColorList::Pink));
 
 					// Move target
