Don't use catchup speed scale in animation driven mode to avoid undesirable speed changes.

diff --git a/Engine/Plugins/AI/MassAI/Source/MassNavigation/Private/MassSteeringProcessors.cpp b/Engine/Plugins/AI/MassAI/Source/MassNavigation/Private/MassSteeringProcessors.cpp
--- a/Engine/Plugins/AI/MassAI/Source/MassNavigation/Private/MassSteeringProcessors.cpp
+++ b/Engine/Plugins/AI/MassAI/Source/MassNavigation/Private/MassSteeringProcessors.cpp
@@ -169,29 +169,36 @@ void UMassSteerToMoveTargetProcessor::Execute(FMassEntityManager& EntityManager,
 						FColor::Red, /*Thickness*/2, TEXT("SteerDirection"));
 				}
 #endif //WITH_MASSGAMEPLAY_DEBUG
-				
-				const FVector::FReal DirSpeedScale = UE::MassNavigation::CalcDirectionalSpeedScale(CurrentForward, SteerDirection);
-				FVector::FReal DesiredSpeed = MoveTarget.DesiredSpeed.Get() * DirSpeedScale;
 
-				// Control speed based relation to the forward axis of the move target.
-				FVector::FReal CatchupDesiredSpeed = DesiredSpeed;
-				if (ForwardOffset < 0.)
-				{
-					// Falling behind, catch up
-					const FVector::FReal T = FMath::Min(-ForwardOffset / LookAheadDistance, 1.);
-					CatchupDesiredSpeed = FMath::Lerp(DesiredSpeed, MovementParams.MaxSpeed, T);
-				}
-				else if (ForwardOffset > 0.)
+				FVector::FReal DesiredSpeed = MoveTarget.DesiredSpeed.Get();
+
+				// When being animation driven, animation has authority over the movement, so it might be useful to disable
+				// this catchup mechanic to avoid subtle speed variations affecting animation.
+				if (MovingSteeringParams.bAllowSpeedVariance)
 				{
-					// Ahead, slow down.
-					const FVector::FReal T = FMath::Min(ForwardOffset / LookAheadDistance, 1.);
-					CatchupDesiredSpeed = FMath::Lerp(DesiredSpeed, DesiredSpeed * 0., 1. - FMath::Square(1. - T));
-				}
+					const FVector::FReal DirSpeedScale = UE::MassNavigation::CalcDirectionalSpeedScale(CurrentForward, SteerDirection);
+					DesiredSpeed *= DirSpeedScale;
 
-				// Control speed based on distance to move target. This allows to catch up even if speed above reaches zero.
-				const FVector::FReal DeviantSpeed = FMath::Min(FMath::Abs(SidewaysOffset) / LookAheadDistance, 1.) * DesiredSpeed;
+					// Control speed based relation to the forward axis of the move target.
+					FVector::FReal CatchupDesiredSpeed = DesiredSpeed;
+					if (ForwardOffset < 0.)
+					{
+						// Falling behind, catch up
+						const FVector::FReal T = FMath::Min(-ForwardOffset / LookAheadDistance, 1.);
+						CatchupDesiredSpeed = FMath::Lerp(DesiredSpeed, MovementParams.MaxSpeed, T);
+					}
+					else if (ForwardOffset > 0.)
+					{
+						// Ahead, slow down.
+						const FVector::FReal T = FMath::Min(ForwardOffset / LookAheadDistance, 1.);
+						CatchupDesiredSpeed = FMath::Lerp(DesiredSpeed, DesiredSpeed * 0., 1. - FMath::Square(1. - T));
+					}
 
-				DesiredSpeed = FMath::Max(CatchupDesiredSpeed, DeviantSpeed);
+					// Control speed based on distance to move target. This allows to catch up even if speed above reaches zero.
+					const FVector::FReal DeviantSpeed = FMath::Min(FMath::Abs(SidewaysOffset) / LookAheadDistance, 1.) * DesiredSpeed;
+
+					DesiredSpeed = FMath::Max(CatchupDesiredSpeed, DeviantSpeed);
+				}
 
 				// Slow down towards the end of path.
 				if (MoveTarget.IntentAtGoal == EMassMovementAction::Stand)
diff --git a/Engine/Plugins/AI/MassAI/Source/MassNavigation/Public/Steering/MassSteeringFragments.h b/Engine/Plugins/AI/MassAI/Source/MassNavigation/Public/Steering/MassSteeringFragments.h
--- a/Engine/Plugins/AI/MassAI/Source/MassNavigation/Public/Steering/MassSteeringFragments.h
+++ b/Engine/Plugins/AI/MassAI/Source/MassNavigation/Public/Steering/MassSteeringFragments.h
@@ -62,6 +62,10 @@ struct MASSNAVIGATION_API FMassMovingSteeringParameters : public FMassConstShare
 	/** How much we look ahead when steering. Affects how steeply we steer towards the goal and when to start to slow down at the end of the path. */
 	UPROPERTY(EditAnywhere, Category = "Moving", meta = (ClampMin = "0", ForceUnits="s"))
 	float LookAheadTime = 1.0f;
+
+	/** Allow directional and catchup speed variance. */
+	UPROPERTY(EditAnywhere, Category = "Moving")
+	bool bAllowSpeedVariance = true;
 };
 
 USTRUCT()
