Addition of a dynamic list of ignorable entities for mass avoidance (to be used for player bumps).

diff --git a/Engine/Plugins/AI/MassAI/Source/MassNavigation/Private/MassAvoidanceProcessors.cpp b/Engine/Plugins/AI/MassAI/Source/MassNavigation/Private/MassAvoidanceProcessors.cpp
--- a/Engine/Plugins/AI/MassAI/Source/MassNavigation/Private/MassAvoidanceProcessors.cpp
+++ b/Engine/Plugins/AI/MassAI/Source/MassNavigation/Private/MassAvoidanceProcessors.cpp
@@ -272,6 +272,7 @@ void UMassMovingAvoidanceProcessor::ConfigureQueries(const TSharedRef<FMassEntit
 	EntityQuery.AddRequirement<FTransformFragment>(EMassFragmentAccess::ReadOnly);
 	EntityQuery.AddRequirement<FMassVelocityFragment>(EMassFragmentAccess::ReadOnly);
 	EntityQuery.AddRequirement<FAgentRadiusFragment>(EMassFragmentAccess::ReadOnly);
+	EntityQuery.AddRequirement<FMassAvoidanceEntitiesToIgnoreFragment>(EMassFragmentAccess::ReadOnly, EMassFragmentPresence::Optional);
 	EntityQuery.AddTagRequirement<FMassMediumLODTag>(EMassFragmentPresence::None);
 	EntityQuery.AddTagRequirement<FMassLowLODTag>(EMassFragmentPresence::None);
 	EntityQuery.AddTagRequirement<FMassOffLODTag>(EMassFragmentPresence::None);
@@ -310,6 +311,7 @@ void UMassMovingAvoidanceProcessor::Execute(FMassEntityManager& EntityManager, F
 		const TConstArrayView<FTransformFragment> LocationList = Context.GetFragmentView<FTransformFragment>();
 		const TConstArrayView<FMassVelocityFragment> VelocityList = Context.GetFragmentView<FMassVelocityFragment>();
 		const TConstArrayView<FAgentRadiusFragment> RadiusList = Context.GetFragmentView<FAgentRadiusFragment>();
+		const TConstArrayView<FMassAvoidanceEntitiesToIgnoreFragment> EntitiesToIgnoreList = Context.GetFragmentView<FMassAvoidanceEntitiesToIgnoreFragment>();
 		const TConstArrayView<FMassMoveTargetFragment> MoveTargetList = Context.GetFragmentView<FMassMoveTargetFragment>();
 		const FMassMovingAvoidanceParameters& MovingAvoidanceParams = Context.GetConstSharedFragment<FMassMovingAvoidanceParameters>();
 		const FMassMovementParameters& MovementParams = Context.GetConstSharedFragment<FMassMovementParameters>();
@@ -365,6 +367,12 @@ void UMassMovingAvoidanceProcessor::Execute(FMassEntityManager& EntityManager, F
 			const FTransformFragment& Location = LocationList[EntityIt];
 			const FMassVelocityFragment& Velocity = VelocityList[EntityIt];
 			const FAgentRadiusFragment& RadiusFragment = RadiusList[EntityIt];
+			
+			TConstArrayView<FMassEntityHandle> EntitiesToIgnore =
+				EntitiesToIgnoreList.IsEmpty() ?
+					TArray<FMassEntityHandle,  TFixedAllocator<FMassAvoidanceEntitiesToIgnoreFragment::EntityCount>>() :
+					EntitiesToIgnoreList[EntityIt].EntitiesToIgnore;
+			
 			FMassForceFragment& Force = ForceList[EntityIt];
 
 			// Smaller steering max accel makes the steering more "calm" but less opportunistic, may not find solution, or gets stuck.
@@ -753,6 +761,12 @@ void UMassMovingAvoidanceProcessor::Execute(FMassEntityManager& EntityManager, F
 					continue;
 				}
 
+				// Skip entities to ignore
+				if (UNLIKELY(!EntitiesToIgnore.IsEmpty()) && EntitiesToIgnore.Contains(OtherEntity.Entity))
+				{
+					continue;
+				}
+
 				FSortedObstacle Obstacle;
 				Obstacle.LocationCached = OtherLocation;
 				Obstacle.Forward = Transform.GetRotation().GetForwardVector();
@@ -1033,6 +1047,7 @@ void UMassStandingAvoidanceProcessor::ConfigureQueries(const TSharedRef<FMassEnt
 	EntityQuery.AddRequirement<FMassMoveTargetFragment>(EMassFragmentAccess::ReadOnly);
 	EntityQuery.AddRequirement<FTransformFragment>(EMassFragmentAccess::ReadOnly);
 	EntityQuery.AddRequirement<FAgentRadiusFragment>(EMassFragmentAccess::ReadOnly);
+	EntityQuery.AddRequirement<FMassAvoidanceEntitiesToIgnoreFragment>(EMassFragmentAccess::ReadOnly, EMassFragmentPresence::Optional);
 	EntityQuery.AddTagRequirement<FMassMediumLODTag>(EMassFragmentPresence::None);
 	EntityQuery.AddTagRequirement<FMassLowLODTag>(EMassFragmentPresence::None);
 	EntityQuery.AddTagRequirement<FMassOffLODTag>(EMassFragmentPresence::None);
@@ -1068,6 +1083,7 @@ void UMassStandingAvoidanceProcessor::Execute(FMassEntityManager& EntityManager,
 		const TArrayView<FMassGhostLocationFragment> GhostList = Context.GetMutableFragmentView<FMassGhostLocationFragment>();
 		const TConstArrayView<FTransformFragment> LocationList = Context.GetFragmentView<FTransformFragment>();
 		const TConstArrayView<FAgentRadiusFragment> RadiusList = Context.GetFragmentView<FAgentRadiusFragment>();
+		const TConstArrayView<FMassAvoidanceEntitiesToIgnoreFragment> EntitiesToIgnoreList = Context.GetFragmentView<FMassAvoidanceEntitiesToIgnoreFragment>();
 		const TConstArrayView<FMassMoveTargetFragment> MoveTargetList = Context.GetFragmentView<FMassMoveTargetFragment>();
 		const FMassStandingAvoidanceParameters& StandingParams = Context.GetConstSharedFragment<FMassStandingAvoidanceParameters>();
 
@@ -1115,6 +1131,11 @@ void UMassStandingAvoidanceProcessor::Execute(FMassEntityManager& EntityManager,
 
 			const FTransformFragment& Location = LocationList[EntityIt];
 			const FAgentRadiusFragment& Radius = RadiusList[EntityIt];
+			
+			TConstArrayView<FMassEntityHandle> EntitiesToIgnore =
+				EntitiesToIgnoreList.IsEmpty() ?
+					TArray<FMassEntityHandle,  TFixedAllocator<FMassAvoidanceEntitiesToIgnoreFragment::EntityCount>>() :
+					EntitiesToIgnoreList[EntityIt].EntitiesToIgnore;
 
 			FMassEntityHandle Entity = Context.GetEntity(EntityIt);
 			const FVector AgentLocation = Location.GetTransform().GetTranslation();
@@ -1170,6 +1191,12 @@ void UMassStandingAvoidanceProcessor::Execute(FMassEntityManager& EntityManager,
 					continue;
 				}
 
+				// Skip entities to ignore
+				if (UNLIKELY(!EntitiesToIgnore.IsEmpty()) && EntitiesToIgnore.Contains(OtherEntity.Entity))
+				{
+					continue;
+				}
+				
 				ClosestObstacles.Emplace(OtherEntity.Entity, OtherLocation, OtherTransform.GetTransform().GetRotation().GetForwardVector(), DistSq);
 			}
 			ClosestObstacles.Sort([](const FSortedObstacle& A, const FSortedObstacle& B) { return A.DistSq < B.DistSq; });
diff --git a/Engine/Plugins/AI/MassAI/Source/MassNavigation/Public/Avoidance/MassAvoidanceFragments.h b/Engine/Plugins/AI/MassAI/Source/MassNavigation/Public/Avoidance/MassAvoidanceFragments.h
--- a/Engine/Plugins/AI/MassAI/Source/MassNavigation/Public/Avoidance/MassAvoidanceFragments.h
+++ b/Engine/Plugins/AI/MassAI/Source/MassNavigation/Public/Avoidance/MassAvoidanceFragments.h
@@ -3,6 +3,7 @@
 #pragma once
 
 #include "MassEntityTypes.h"
+#include "MassEntityHandle.h"
 #if UE_ENABLE_INCLUDE_ORDER_DEPRECATED_IN_5_6
 #include "MassNavigationSubsystem.h"
 #endif // UE_ENABLE_INCLUDE_ORDER_DEPRECATED_IN_5_6
@@ -182,9 +183,18 @@ struct MASSNAVIGATION_API FMassNavigationEdgesFragment : public FMassFragment
 {
 	GENERATED_BODY()
 
-	static const int MaxEdgesCount = 16;
+	static constexpr int MaxEdgesCount = 16;
 	TArray<FNavigationAvoidanceEdge, TFixedAllocator<MaxEdgesCount>> AvoidanceEdges;
 
 	/* Set to true if the edges already account for the agent radius. */
 	bool bExtrudedEdges = false;
 };
+
+USTRUCT()
+struct MASSNAVIGATION_API FMassAvoidanceEntitiesToIgnoreFragment : public FMassFragment
+{
+	GENERATED_BODY()
+
+	static constexpr int EntityCount = 4;
+	TArray<FMassEntityHandle, TFixedAllocator<EntityCount>> EntitiesToIgnore;
+};
