[Buildfix] Replace temp allocated array with empty arrayview.

diff --git a/Engine/Plugins/AI/MassAI/Source/MassNavigation/Private/MassAvoidanceProcessors.cpp b/Engine/Plugins/AI/MassAI/Source/MassNavigation/Private/MassAvoidanceProcessors.cpp
--- a/Engine/Plugins/AI/MassAI/Source/MassNavigation/Private/MassAvoidanceProcessors.cpp
+++ b/Engine/Plugins/AI/MassAI/Source/MassNavigation/Private/MassAvoidanceProcessors.cpp
@@ -370,7 +370,7 @@ void UMassMovingAvoidanceProcessor::Execute(FMassEntityManager& EntityManager, F
 			
 			TConstArrayView<FMassEntityHandle> EntitiesToIgnore =
 				EntitiesToIgnoreList.IsEmpty() ?
-					TArray<FMassEntityHandle,  TFixedAllocator<FMassAvoidanceEntitiesToIgnoreFragment::EntityCount>>() :
+				TConstArrayView<FMassEntityHandle>() :
 					EntitiesToIgnoreList[EntityIt].EntitiesToIgnore;
 			
 			FMassForceFragment& Force = ForceList[EntityIt];
@@ -1134,7 +1134,7 @@ void UMassStandingAvoidanceProcessor::Execute(FMassEntityManager& EntityManager,
 			
 			TConstArrayView<FMassEntityHandle> EntitiesToIgnore =
 				EntitiesToIgnoreList.IsEmpty() ?
-					TArray<FMassEntityHandle,  TFixedAllocator<FMassAvoidanceEntitiesToIgnoreFragment::EntityCount>>() :
+				TConstArrayView<FMassEntityHandle>() :
 					EntitiesToIgnoreList[EntityIt].EntitiesToIgnore;
 
 			FMassEntityHandle Entity = Context.GetEntity(EntityIt);
