Changes for NPC strafing prototype. disabled by default

diff --git a/Engine/Plugins/AI/MassAI/Source/MassNavigation/Private/MassSmoothOrientationProcessors.cpp b/Engine/Plugins/AI/MassAI/Source/MassNavigation/Private/MassSmoothOrientationProcessors.cpp
--- a/Engine/Plugins/AI/MassAI/Source/MassNavigation/Private/MassSmoothOrientationProcessors.cpp
+++ b/Engine/Plugins/AI/MassAI/Source/MassNavigation/Private/MassSmoothOrientationProcessors.cpp
@@ -26,7 +26,8 @@ UMassSmoothOrientationProcessor::UMassSmoothOrientationProcessor()
 void UMassSmoothOrientationProcessor::ConfigureQueries(const TSharedRef<FMassEntityManager>& EntityManager)
 {
 	HighResEntityQuery.AddRequirement<FMassMoveTargetFragment>(EMassFragmentAccess::ReadOnly);
-	HighResEntityQuery.AddRequirement<FMassVelocityFragment>(EMassFragmentAccess::ReadWrite);
+	HighResEntityQuery.AddRequirement<FMassVelocityFragment>(EMassFragmentAccess::ReadOnly);
+	HighResEntityQuery.AddRequirement<FMassDesiredMovementFragment>(EMassFragmentAccess::ReadWrite);
 	HighResEntityQuery.AddRequirement<FTransformFragment>(EMassFragmentAccess::ReadWrite);
 	HighResEntityQuery.AddTagRequirement<FMassOffLODTag>(EMassFragmentPresence::None);
 	HighResEntityQuery.AddConstSharedRequirement<FMassSmoothOrientationParameters>(EMassFragmentPresence::All);
@@ -53,7 +54,8 @@ void UMassSmoothOrientationProcessor::Execute(FMassEntityManager& EntityManager,
 
 			const TConstArrayView<FMassMoveTargetFragment> MoveTargetList = Context.GetFragmentView<FMassMoveTargetFragment>();
 			const TArrayView<FTransformFragment> LocationList = Context.GetMutableFragmentView<FTransformFragment>();
-			const TArrayView<FMassVelocityFragment> VelocityList = Context.GetMutableFragmentView<FMassVelocityFragment>();
+			const TConstArrayView<FMassVelocityFragment> VelocityList = Context.GetFragmentView<FMassVelocityFragment>();
+			const TArrayView<FMassDesiredMovementFragment> DesiredMovementList = Context.GetMutableFragmentView<FMassDesiredMovementFragment>();
 
 			for (FMassExecutionContext::FEntityIterator EntityIt = Context.CreateEntityIterator(); EntityIt; ++EntityIt)
 			{
@@ -66,6 +68,7 @@ void UMassSmoothOrientationProcessor::Execute(FMassEntityManager& EntityManager,
 				}
 
 				const FMassVelocityFragment& CurrentVelocity = VelocityList[EntityIt];
+				FMassDesiredMovementFragment& DesiredMovement = DesiredMovementList[EntityIt];
 				FTransform& CurrentTransform = LocationList[EntityIt].GetMutableTransform();
 				const FVector CurrentForward = CurrentTransform.GetRotation().GetForwardVector();
 				const FVector::FReal CurrentHeading = UE::MassNavigation::GetYawFromDirection(CurrentForward);
@@ -102,6 +105,7 @@ void UMassSmoothOrientationProcessor::Execute(FMassEntityManager& EntityManager,
 
 				const FVector::FReal Ratio = MoveTargetWeight / (MoveTargetWeight + VelocityWeight);
 				const FVector::FReal DesiredHeading = UE::MassNavigation::LerpAngle(VelocityHeading, MovementHeading,Ratio);
+				DesiredMovement.DesiredFacing = FQuat(FVector::UpVector, DesiredHeading);
 				
 				const FVector::FReal NewHeading = UE::MassNavigation::ExponentialSmoothingAngle(CurrentHeading, DesiredHeading, DeltaTime, OrientationParams.OrientationSmoothingTime);
 
diff --git a/Engine/Plugins/Runtime/MassGameplay/Source/MassMovement/Public/MassMovementFragments.h b/Engine/Plugins/Runtime/MassGameplay/Source/MassMovement/Public/MassMovementFragments.h
--- a/Engine/Plugins/Runtime/MassGameplay/Source/MassMovement/Public/MassMovementFragments.h
+++ b/Engine/Plugins/Runtime/MassGameplay/Source/MassMovement/Public/MassMovementFragments.h
@@ -37,10 +37,9 @@ struct MASSMOVEMENT_API FMassDesiredMovementFragment : public FMassFragment
 	GENERATED_BODY()
 	
 	FVector DesiredVelocity = FVector::ZeroVector;
+	FQuat DesiredFacing = FQuat::Identity;
 
 	float DesiredMaxSpeedOverride = FLT_MAX;
-	
-	// Todo: Could add a DesiredFacing here
 };
 
 /** Accumulator for steering / avoidance forces to apply to the desired velocity */ 
