[LookAt] added FMassLookAtTargetOffsetFragment to provide per target offset

diff --git a/Engine/Plugins/AI/MassAI/Source/MassAIBehavior/Private/MassLookAtProcessors.cpp b/Engine/Plugins/AI/MassAI/Source/MassAIBehavior/Private/MassLookAtProcessors.cpp
--- a/Engine/Plugins/AI/MassAI/Source/MassAIBehavior/Private/MassLookAtProcessors.cpp
+++ b/Engine/Plugins/AI/MassAI/Source/MassAIBehavior/Private/MassLookAtProcessors.cpp
@@ -105,6 +105,7 @@ UMassLookAtProcessor::UMassLookAtProcessor()
 void UMassLookAtProcessor::ConfigureQueries(const TSharedRef<FMassEntityManager>& EntityManager)
 {
 	EntityQuery_Conditional.AddRequirement<FMassLookAtFragment>(EMassFragmentAccess::ReadWrite);
+	EntityQuery_Conditional.AddRequirement<FMassLookAtTargetOffsetFragment>(EMassFragmentAccess::ReadOnly, EMassFragmentPresence::Optional);
 	EntityQuery_Conditional.AddRequirement<FTransformFragment>(EMassFragmentAccess::ReadOnly);
 	EntityQuery_Conditional.AddRequirement<FMassMoveTargetFragment>(EMassFragmentAccess::ReadOnly, EMassFragmentPresence::Optional);
 	EntityQuery_Conditional.AddRequirement<FMassZoneGraphLaneLocationFragment>(EMassFragmentAccess::ReadOnly, EMassFragmentPresence::Optional);
@@ -133,6 +134,7 @@ void UMassLookAtProcessor::Execute(FMassEntityManager& EntityManager, FMassExecu
 			const UZoneGraphSubsystem& ZoneGraphSubsystem = Context.GetSubsystemChecked<UZoneGraphSubsystem>();
 
 			const TArrayView<FMassLookAtFragment> LookAtList = Context.GetMutableFragmentView<FMassLookAtFragment>();
+			const TConstArrayView<FMassLookAtTargetOffsetFragment> EyeOffsetList = Context.GetFragmentView<FMassLookAtTargetOffsetFragment>();
 			const TConstArrayView<FTransformFragment> TransformList = Context.GetFragmentView<FTransformFragment>();
 			const TConstArrayView<FMassMoveTargetFragment> MoveTargetList = Context.GetFragmentView<FMassMoveTargetFragment>();
 			const TConstArrayView<FMassZoneGraphLaneLocationFragment> ZoneGraphLocationList = Context.GetFragmentView<FMassZoneGraphLaneLocationFragment>();
@@ -300,7 +302,17 @@ void UMassLookAtProcessor::FindNewGazeTarget(const UMassNavigationSubsystem& Mas
 
 				// TargetTag is added through the LookAtTargetTrait and Transform was added with it
 				const FTransformFragment& TargetTransform = EntityView.GetFragmentData<FTransformFragment>();
-				const FVector TargetLocation = TargetTransform.GetTransform().GetLocation() + GetDefault<UMassLookAtSettings>()->GetDefaultTargetLocationOffset();
+				FVector LocationOffset;
+				if (const FMassLookAtTargetOffsetFragment* OffsetFragment = EntityView.GetFragmentDataPtr<FMassLookAtTargetOffsetFragment>())
+				{
+					LocationOffset = OffsetFragment->Offset;
+				}
+				else
+				{
+					LocationOffset = GetDefault<UMassLookAtSettings>()->GetDefaultTargetLocationOffset();
+				}
+
+				const FVector TargetLocation = TargetTransform.GetTransform().TransformPosition(LocationOffset);
 				FVector Direction = (TargetLocation - Location).GetSafeNormal();
 				Direction = InTransform.InverseTransformVector(Direction);
 
@@ -412,11 +424,21 @@ void UMassLookAtProcessor::UpdateLookAtTrackedEntity(const FMassEntityManager& E
 	{
 		if (const FTransformFragment* TargetTransform = EntityManager.GetFragmentDataPtr<FTransformFragment>(LookAt.TrackedEntity))
 		{
+			FVector LocationOffset;
+			if (const FMassLookAtTargetOffsetFragment* OffsetFragment = EntityManager.GetFragmentDataPtr<FMassLookAtTargetOffsetFragment>(LookAt.TrackedEntity))
+			{
+				LocationOffset = OffsetFragment->Offset;
+			}
+			else
+			{
+				LocationOffset = GetDefault<UMassLookAtSettings>()->GetDefaultTargetLocationOffset();
+			}
+
 			const FVector AgentPosition = Transform.GetLocation();
-			const FVector TargetLocation = TargetTransform->GetTransform().GetLocation();
+			const FVector TargetLocation = TargetTransform->GetTransform().TransformPosition(LocationOffset);
 			const FVector NewGlobalDirection = (TargetLocation - AgentPosition).GetSafeNormal();
 			LookAt.Direction = Transform.InverseTransformVector(NewGlobalDirection);
-			LookAt.MainTargetLocation = TargetLocation + GetDefault<UMassLookAtSettings>()->GetDefaultTargetLocationOffset();
+			LookAt.MainTargetLocation = TargetLocation;
 
 #if WITH_MASSGAMEPLAY_DEBUG
 			if (bDisplayDebug)
@@ -438,11 +460,22 @@ bool UMassLookAtProcessor::UpdateGazeTrackedEntity(const FMassEntityManager& Ent
 	{
 		if (const FTransformFragment* TargetTransform = EntityManager.GetFragmentDataPtr<FTransformFragment>(LookAt.GazeTrackedEntity))
 		{
+			FVector LocationOffset;
+			if (const FMassLookAtTargetOffsetFragment* OffsetFragment = EntityManager.GetFragmentDataPtr<FMassLookAtTargetOffsetFragment>(LookAt.GazeTrackedEntity))
+			{
+				LocationOffset = OffsetFragment->Offset;
+			}
+			else
+			{
+				LocationOffset = GetDefault<UMassLookAtSettings>()->GetDefaultTargetLocationOffset();
+			}
+
 			const FVector AgentPosition = Transform.GetLocation();
-			const FVector TargetLocation = TargetTransform->GetTransform().GetLocation();
+			const FVector TargetLocation = TargetTransform->GetTransform().TransformPosition(LocationOffset);
 			const FVector NewGlobalDirection = (TargetLocation - AgentPosition).GetSafeNormal();
 			LookAt.GazeDirection = Transform.InverseTransformVector(NewGlobalDirection);
-			LookAt.GazeTargetLocation = TargetLocation + GetDefault<UMassLookAtSettings>()->GetDefaultTargetLocationOffset();
+			LookAt.GazeTargetLocation = TargetLocation;
+
 			bHasTarget = true;
 			
 #if WITH_MASSGAMEPLAY_DEBUG
diff --git a/Engine/Plugins/AI/MassAI/Source/MassAIBehavior/Public/MassLookAtFragments.h b/Engine/Plugins/AI/MassAI/Source/MassAIBehavior/Public/MassLookAtFragments.h
--- a/Engine/Plugins/AI/MassAI/Source/MassAIBehavior/Public/MassLookAtFragments.h
+++ b/Engine/Plugins/AI/MassAI/Source/MassAIBehavior/Public/MassLookAtFragments.h
@@ -34,6 +34,25 @@ enum class EMassLookAtGazeMode : uint8
 	Glance,
 };
 
+/**
+ * Fragment providing additional offset that could be used by look at processors to get more accurate locations
+ */
+USTRUCT()
+struct MASSAIBEHAVIOR_API FMassLookAtTargetOffsetFragment : public FMassFragment
+{
+	GENERATED_BODY()
+
+	FMassLookAtTargetOffsetFragment() = default;
+	explicit FMassLookAtTargetOffsetFragment(const FVector& Offset)
+		: Offset(Offset)
+	{
+	}
+
+	/** Offset in local space to add to the target transform to get the final location */
+	UPROPERTY(Transient)
+	FVector Offset = FVector::ZeroVector;
+};
+
 /**
  * Struct that holds all parameters of the current entity look at 
  */
diff --git a/Engine/Plugins/AI/MassAI/Source/MassAIDebug/Private/GameplayDebuggerCategory_Mass.cpp b/Engine/Plugins/AI/MassAI/Source/MassAIDebug/Private/GameplayDebuggerCategory_Mass.cpp
--- a/Engine/Plugins/AI/MassAI/Source/MassAIDebug/Private/GameplayDebuggerCategory_Mass.cpp
+++ b/Engine/Plugins/AI/MassAI/Source/MassAIDebug/Private/GameplayDebuggerCategory_Mass.cpp
@@ -509,6 +509,7 @@ void FGameplayDebuggerCategory_Mass::CollectData(APlayerController* OwnerPC, AAc
 		EntityQuery.AddRequirement<FMassForceFragment>(EMassFragmentAccess::ReadOnly);
 		EntityQuery.AddRequirement<FMassMoveTargetFragment>(EMassFragmentAccess::ReadOnly);
 		EntityQuery.AddRequirement<FMassLookAtFragment>(EMassFragmentAccess::ReadOnly, EMassFragmentPresence::Optional);
+		EntityQuery.AddRequirement<FMassLookAtTargetOffsetFragment>(EMassFragmentAccess::ReadOnly, EMassFragmentPresence::Optional);
 		EntityQuery.AddRequirement<FMassSimulationLODFragment>(EMassFragmentAccess::ReadOnly, EMassFragmentPresence::Optional);
 		EntityQuery.AddRequirement<FMassZoneGraphShortPathFragment>(EMassFragmentAccess::ReadOnly, EMassFragmentPresence::Optional);
 		EntityQuery.AddRequirement<FMassNavMeshShortPathFragment>(EMassFragmentAccess::ReadOnly, EMassFragmentPresence::Optional);
@@ -538,6 +539,7 @@ void FGameplayDebuggerCategory_Mass::CollectData(APlayerController* OwnerPC, AAc
 				const TConstArrayView<FMassForceFragment> ForceList = Context.GetFragmentView<FMassForceFragment>();
 				const TConstArrayView<FMassMoveTargetFragment> MoveTargetList = Context.GetFragmentView<FMassMoveTargetFragment>();
 				const TConstArrayView<FMassLookAtFragment> LookAtList = Context.GetFragmentView<FMassLookAtFragment>();
+				const TConstArrayView<FMassLookAtTargetOffsetFragment> LookAtOffsetList = Context.GetFragmentView<FMassLookAtTargetOffsetFragment>();
 				const TConstArrayView<FMassSimulationLODFragment> SimLODList = Context.GetFragmentView<FMassSimulationLODFragment>();
 				const TConstArrayView<FMassZoneGraphShortPathFragment> ZoneGraphShortPathList = Context.GetFragmentView<FMassZoneGraphShortPathFragment>();
 				const TConstArrayView<FMassNavMeshShortPathFragment> NavMeshShortPathList = Context.GetFragmentView<FMassNavMeshShortPathFragment>();
@@ -584,9 +586,6 @@ void FGameplayDebuggerCategory_Mass::CollectData(APlayerController* OwnerPC, AAc
 
 					const FVector EntityForward = Transform.GetTransform().GetRotation().GetForwardVector();
 
-					const FVector TargetOffset = GetDefault<UMassLookAtSettings>()->GetDefaultTargetLocationOffset();
-					const double EyeHeight = TargetOffset.IsZero() ? 160 : TargetOffset.Z;
-
 					// Draw entity position and orientation.
 					FVector BasePos = EntityLocation + FVector(0 ,0 ,25);
 
@@ -612,15 +611,26 @@ void FGameplayDebuggerCategory_Mass::CollectData(APlayerController* OwnerPC, AAc
 						constexpr float LookArrowHeadSize = 20.f;
 						constexpr float LookArrowThickness = 2.f;
 
-						BasePos = EntityLocation + FVector(0,0,EyeHeight);
+						BasePos = EntityLocation;
+
+						// Apply the offset for the entity looking at something
+						if (LookAtOffsetList.Num())
+						{
+							BasePos += LookAtOffsetList[EntityIndex].Offset;
+						}
+						else
+						{
+							const FVector TargetOffset = GetDefault<UMassLookAtSettings>()->GetDefaultTargetLocationOffset();
+							BasePos += TargetOffset.IsZero() ? FVector(0, 0, 160) : TargetOffset;
+						}
 
 						const FMassLookAtFragment& LookAt = LookAtList[EntityIndex];
 						const FVector WorldLookDirection = Transform.GetTransform().TransformVector(LookAt.Direction);
 						bool bLookArrowDrawn = false;
 						if (LookAt.LookAtMode == EMassLookAtMode::LookAtEntity && EntityManager.IsEntityValid(LookAt.TrackedEntity))
 						{
+							// Target location is expected to include its own offset
 							FVector TargetPosition = LookAt.MainTargetLocation;
-							TargetPosition.Z = BasePos.Z;
 							AddShape(FGameplayDebuggerShape::MakeCircle(TargetPosition, FVector::UpVector, Radius.Radius, FColor::Red));
 							AddShape(FGameplayDebuggerShape::MakeArrow(BasePos, TargetPosition, TargetArrowHeadSize, TargetArrowThickness, FColor::Red));
 
@@ -631,8 +641,8 @@ void FGameplayDebuggerCategory_Mass::CollectData(APlayerController* OwnerPC, AAc
 
 						if (LookAt.bRandomGazeEntities && EntityManager.IsEntityValid(LookAt.GazeTrackedEntity))
 						{
+							// Target location is expected to include its own offset
 							FVector TargetPosition = LookAt.GazeTargetLocation;
-							TargetPosition.Z = BasePos.Z;
 							AddShape(FGameplayDebuggerShape::MakeCircle(TargetPosition, FVector::UpVector, Radius.Radius, FColorList::Goldenrod));
 							AddShape(FGameplayDebuggerShape::MakeArrow(BasePos, TargetPosition, TargetArrowHeadSize, TargetArrowThickness, FColorList::Goldenrod));
 						}
