[LookAt] replaced and merged FMassLookAtTargetTag, FMassLookAtTargetOffsetFragment and FMassLookAtTargetGridCellLocationFragment by FMassLookAtTargetFragment
- FMassLookAtTargetTrait will add the new fragment, set the default offset and adjust to the capsule height when associated actor is available
- Added priority to the fragment so when selecting random target we can prioritize some of them
- Added 'TargetOffsetFromCapsuleHeight' in MassLookAtSettings to tweak the target position from the top of the capsule
- Now all entities with FMassLookAtTargetFragment are added to the LookAt grid so the processor no longer need to query the obstacle grid

diff --git a/Engine/Plugins/AI/MassAI/Source/MassAIBehavior/Private/MassLookAtProcessors.cpp b/Engine/Plugins/AI/MassAI/Source/MassAIBehavior/Private/MassLookAtProcessors.cpp
--- a/Engine/Plugins/AI/MassAI/Source/MassAIBehavior/Private/MassLookAtProcessors.cpp
+++ b/Engine/Plugins/AI/MassAI/Source/MassAIBehavior/Private/MassLookAtProcessors.cpp
@@ -1,44 +1,33 @@
 // Copyright Epic Games, Inc. All Rights Reserved.
 
 #include "MassLookAtProcessors.h"
-#include "MassEntityView.h"
+#include "Algo/RandomShuffle.h"
+#include "Curves/BezierUtilities.h"
 #include "MassAIBehaviorTypes.h"
 #include "MassCommonFragments.h"
-#include "MassExecutionContext.h"
+#include "MassDebugger.h"
+#include "MassLODFragments.h"
 #include "MassLookAtFragments.h"
+#include "MassLookAtSubsystem.h"
 #include "MassMovementFragments.h"
 #include "MassNavigationFragments.h"
-#include "MassNavigationSubsystem.h"
-#include "MassRepresentationFragments.h"
-#include "Math/UnrealMathUtility.h"
-#include "VisualLogger/VisualLogger.h"
+#include "MassRepresentationTypes.h"
 #include "MassZoneGraphNavigationFragments.h"
-#include "ZoneGraphSubsystem.h"
+#include "VisualLogger/VisualLogger.h"
 #include "ZoneGraphQuery.h"
-#include "Curves/BezierUtilities.h"
-#include "Algo/RandomShuffle.h"
-#include "Engine/World.h"
-#include "MassLODFragments.h"
-#include "MassDebugger.h"
-#include "MassLookAtSettings.h"
-#include "MassLookAtSubsystem.h"
-
+#include "ZoneGraphSubsystem.h"
 
 namespace UE::MassBehavior
 {
 	namespace Tweakables
 	{
 		float TrajectoryLookAhead = 600.f;
-		bool bConsiderObstacleTargets = true;
 	}
 
 	FAutoConsoleVariableRef CVars[] =
 	{
 		FAutoConsoleVariableRef(TEXT("ai.mass.LookAt.TrajectoryLookAhead"), Tweakables::TrajectoryLookAhead,
 								TEXT("Distance (in cm) further along the look at trajectory (based on current path) to look at while moving."), ECVF_Cheat),
-
-		FAutoConsoleVariableRef(TEXT("ai.mass.LookAt.ConsiderObstacleTargets"), Tweakables::bConsiderObstacleTargets,
-								TEXT("Indicates if gaze targets can be selected from the entities registered in the obstacle grid."), ECVF_Cheat),
 	};
 
 	// Clamps direction vector to a cone specified by the cone angle along X-axis
@@ -105,7 +94,6 @@ UMassLookAtProcessor::UMassLookAtProcessor()
 void UMassLookAtProcessor::ConfigureQueries(const TSharedRef<FMassEntityManager>& EntityManager)
 {
 	EntityQuery_Conditional.AddRequirement<FMassLookAtFragment>(EMassFragmentAccess::ReadWrite);
-	EntityQuery_Conditional.AddRequirement<FMassLookAtTargetOffsetFragment>(EMassFragmentAccess::ReadOnly, EMassFragmentPresence::Optional);
 	EntityQuery_Conditional.AddRequirement<FTransformFragment>(EMassFragmentAccess::ReadOnly);
 	EntityQuery_Conditional.AddRequirement<FMassMoveTargetFragment>(EMassFragmentAccess::ReadOnly, EMassFragmentPresence::Optional);
 	EntityQuery_Conditional.AddRequirement<FMassZoneGraphLaneLocationFragment>(EMassFragmentAccess::ReadOnly, EMassFragmentPresence::Optional);
@@ -134,7 +122,6 @@ void UMassLookAtProcessor::Execute(FMassEntityManager& EntityManager, FMassExecu
 			const UZoneGraphSubsystem& ZoneGraphSubsystem = Context.GetSubsystemChecked<UZoneGraphSubsystem>();
 
 			const TArrayView<FMassLookAtFragment> LookAtList = Context.GetMutableFragmentView<FMassLookAtFragment>();
-			const TConstArrayView<FMassLookAtTargetOffsetFragment> EyeOffsetList = Context.GetFragmentView<FMassLookAtTargetOffsetFragment>();
 			const TConstArrayView<FTransformFragment> TransformList = Context.GetFragmentView<FTransformFragment>();
 			const TConstArrayView<FMassMoveTargetFragment> MoveTargetList = Context.GetFragmentView<FMassMoveTargetFragment>();
 			const TConstArrayView<FMassZoneGraphLaneLocationFragment> ZoneGraphLocationList = Context.GetFragmentView<FMassZoneGraphLaneLocationFragment>();
@@ -269,12 +256,43 @@ void UMassLookAtProcessor::FindNewGazeTarget(const UMassNavigationSubsystem& Mas
 		const FVector QueryOrigin = Transform.TransformPosition(FVector(0.5f*QueryExtent, 0.f, 0.f));
 		const FBox QueryBox = FBox(QueryOrigin - 0.5f*Extent, QueryOrigin + 0.5f*Extent);
 
-		auto FindTargetFunc = [&EntityManager, &LastTrackedEntity, CosAngleThreshold, &LookAt, Entity]
-		(TStridedView<FMassEntityHandle> InCandidates, const FTransform& InTransform)
+		// Process from LookAt targets grid
+		TArray<UE::Mass::LookAt::FTargetHashGrid2D::ItemIDType, TInlineAllocator<UE::Mass::LookAt::HashGridResultsSoftLimit>> NearbyEntities;
+		if (LookAtSystem.Query(QueryBox, NearbyEntities))
 		{
-			const FVector Location = InTransform.GetLocation();
-			for (FMassEntityHandle CandidateEntity : InCandidates)
+			NearbyEntities.Sort([]
+			(const UE::Mass::LookAt::FTargetHashGrid2D::ItemIDType& A, const UE::Mass::LookAt::FTargetHashGrid2D::ItemIDType& B)
 			{
+				return A.Priority < B.Priority;
+			});
+
+			TOptional<uint8> LastPriority;
+			TArray<int32, TInlineAllocator<UE::Mass::LookAt::HashGridResultsSoftLimit>> SpanIndices;
+			for (int32 ItemIndex = 0; ItemIndex < NearbyEntities.Num(); ItemIndex++)
+			{
+				const UE::Mass::LookAt::FTargetHashGridItem& Item = NearbyEntities[ItemIndex];
+				if (!LastPriority.IsSet() || Item.Priority != LastPriority.GetValue())
+				{
+					LastPriority = Item.Priority;
+					SpanIndices.Add(ItemIndex);
+				}
+			}
+
+			// We'll pick the first entity that passes, this ensures that it's a random one.
+			// For now higher priority targets are always selected first
+			for (int32 SpanIndex = 0; SpanIndex < SpanIndices.Num(); SpanIndex++)
+			{
+				const int32 FirstItemIndex = SpanIndices[SpanIndex];
+				const int32 NextSpanItemIndex = SpanIndex + 1 < SpanIndices.Num() ? SpanIndices[SpanIndex + 1] : NearbyEntities.Num();
+				const int32 NumItems = NextSpanItemIndex - FirstItemIndex;
+				Algo::RandomShuffle(MakeArrayView(&NearbyEntities[FirstItemIndex], NumItems));
+			}
+
+			const FVector Location = Transform.GetLocation();
+			for (const UE::Mass::LookAt::FTargetHashGrid2D::ItemIDType& Item : NearbyEntities)
+			{
+				FMassEntityHandle CandidateEntity = Item.TargetEntity;
+
 				// This can happen if we remove entities in the system.
 				if (!EntityManager.IsEntityValid(CandidateEntity))
 				{
@@ -294,27 +312,13 @@ void UMassLookAtProcessor::FindNewGazeTarget(const UMassNavigationSubsystem& Mas
 					continue;
 				}
 
-				FMassEntityView EntityView(EntityManager, CandidateEntity);
-				if (!EntityView.HasTag<FMassLookAtTargetTag>())
-				{
-					continue;
-				}
+				// TargetFragment and Transform are added through the LookAtTargetTrait
+				const FTransformFragment& TargetTransform = EntityManager.GetFragmentDataChecked<FTransformFragment>(CandidateEntity);
+				const FMassLookAtTargetFragment& TargetFragment = EntityManager.GetFragmentDataChecked<FMassLookAtTargetFragment>(CandidateEntity);
 
-				// TargetTag is added through the LookAtTargetTrait and Transform was added with it
-				const FTransformFragment& TargetTransform = EntityView.GetFragmentData<FTransformFragment>();
-				FVector LocationOffset;
-				if (const FMassLookAtTargetOffsetFragment* OffsetFragment = EntityView.GetFragmentDataPtr<FMassLookAtTargetOffsetFragment>())
-				{
-					LocationOffset = OffsetFragment->Offset;
-				}
-				else
-				{
-					LocationOffset = GetDefault<UMassLookAtSettings>()->GetDefaultTargetLocationOffset();
-				}
-
-				const FVector TargetLocation = TargetTransform.GetTransform().TransformPosition(LocationOffset);
+				const FVector TargetLocation = TargetTransform.GetTransform().TransformPosition(TargetFragment.Offset);
 				FVector Direction = (TargetLocation - Location).GetSafeNormal();
-				Direction = InTransform.InverseTransformVector(Direction);
+				Direction = Transform.InverseTransformVector(Direction);
 
 				const bool bIsTargetInView = FVector::DotProduct(FVector::ForwardVector, Direction) > CosAngleThreshold;
 				if (bIsTargetInView)
@@ -322,59 +326,27 @@ void UMassLookAtProcessor::FindNewGazeTarget(const UMassNavigationSubsystem& Mas
 					LookAt.GazeDirection = Direction;
 					LookAt.GazeTrackedEntity = CandidateEntity;
 					LookAt.GazeTargetLocation = TargetLocation;
-					return true;
+					bTargetFound = true;
+					break;
 				}
 
 				// Allow to pick entities out of view if they are moving towards us.
-				if (const FMassVelocityFragment* Velocity = EntityView.GetFragmentDataPtr<FMassVelocityFragment>())
+				if (const FMassVelocityFragment* Velocity = EntityManager.GetFragmentDataPtr<FMassVelocityFragment>(CandidateEntity))
 				{
-					const FVector MoveDirection = InTransform.InverseTransformVector(Velocity->Value.GetSafeNormal());
+					const FVector MoveDirection = Transform.InverseTransformVector(Velocity->Value.GetSafeNormal());
 
-					// Direction negated as it is from the agent to target, and we want target to agent. 
+					// Direction negated as it is from the agent to target, and we want target to agent.
 					const bool bIsTargetMovingTowards = FVector::DotProduct(MoveDirection, -Direction) > CosAngleThreshold;
 					if (bIsTargetMovingTowards)
 					{
 						LookAt.GazeDirection = Direction;
 						LookAt.GazeTrackedEntity = CandidateEntity;
 						LookAt.GazeTargetLocation = TargetLocation;
-						return true;
+						bTargetFound = true;
+						break;
 					}
 				}
 			}
-
-			return false;
-		};
-
-		// Process from LookAt targets grid first
-		{
-			TArray<UE::Mass::LookAt::FTargetHashGrid2D::ItemIDType> NearbyEntities;
-			NearbyEntities.Reserve(16);
-			if (LookAtSystem.Query(QueryBox, NearbyEntities))
-			{
-				// We'll pick the first entity that passes, this ensure that it's random one.
-				Algo::RandomShuffle(NearbyEntities);
-
-				TStridedView<FMassEntityHandle> NearbyLookAtTargetEntities = MakeStridedView(NearbyEntities);
-				bTargetFound = FindTargetFunc(NearbyLookAtTargetEntities, Transform);
-			}
-		}
-
-		// Process from obstacles grid
-		if (!bTargetFound && UE::MassBehavior::Tweakables::bConsiderObstacleTargets)
-		{
-			TArray<FNavigationObstacleHashGrid2D::ItemIDType> NearbyEntities;
-			NearbyEntities.Reserve(16);
-			const FNavigationObstacleHashGrid2D& Grid = MassNavSystem.GetObstacleGrid();
-			if (Grid.GetItems().Num() > 0)
-			{
-				Grid.Query(QueryBox, NearbyEntities);
-
-				// We'll pick the first entity that passes, this ensure that it's random one.
-				Algo::RandomShuffle(NearbyEntities);
-
-				TStridedView<FMassEntityHandle> NearbyAvoidanceObstacleEntities = MakeStridedView(NearbyEntities, &FNavigationObstacleHashGrid2D::ItemIDType::Entity);
-				bTargetFound = FindTargetFunc(NearbyAvoidanceObstacleEntities, Transform);
-			}
 		}
 	}
 
@@ -422,72 +394,50 @@ void UMassLookAtProcessor::UpdateLookAtTrackedEntity(const FMassEntityManager& E
 	// Update direction toward target
 	if (EntityManager.IsEntityValid(LookAt.TrackedEntity))
 	{
-		if (const FTransformFragment* TargetTransform = EntityManager.GetFragmentDataPtr<FTransformFragment>(LookAt.TrackedEntity))
-		{
-			FVector LocationOffset;
-			if (const FMassLookAtTargetOffsetFragment* OffsetFragment = EntityManager.GetFragmentDataPtr<FMassLookAtTargetOffsetFragment>(LookAt.TrackedEntity))
-			{
-				LocationOffset = OffsetFragment->Offset;
-			}
-			else
-			{
-				LocationOffset = GetDefault<UMassLookAtSettings>()->GetDefaultTargetLocationOffset();
-			}
-
-			const FVector AgentPosition = Transform.GetLocation();
-			const FVector TargetLocation = TargetTransform->GetTransform().TransformPosition(LocationOffset);
-			const FVector NewGlobalDirection = (TargetLocation - AgentPosition).GetSafeNormal();
-			LookAt.Direction = Transform.InverseTransformVector(NewGlobalDirection);
-			LookAt.MainTargetLocation = TargetLocation;
+		const FTransformFragment& TargetTransform = EntityManager.GetFragmentDataChecked<FTransformFragment>(LookAt.TrackedEntity);
+		const FVector LocationOffset = EntityManager.GetFragmentDataChecked<FMassLookAtTargetFragment>(LookAt.TrackedEntity).Offset;
+		const FVector AgentPosition = Transform.GetLocation();
+		const FVector TargetLocation = TargetTransform.GetTransform().TransformPosition(LocationOffset);
+		const FVector NewGlobalDirection = (TargetLocation - AgentPosition).GetSafeNormal();
+		LookAt.Direction = Transform.InverseTransformVector(NewGlobalDirection);
+		LookAt.MainTargetLocation = TargetLocation;
 
 #if WITH_MASSGAMEPLAY_DEBUG
-			if (bDisplayDebug)
-			{
-				const FVector ZOffset(0.f,0.f,DebugZOffset);
-				UE_VLOG_SEGMENT_THICK(this, LogMassBehavior, Display, AgentPosition + ZOffset, AgentPosition + ZOffset + 100.f * NewGlobalDirection, FColor::White, /*Thickness*/3, TEXT("LookAt Track"));
-			}
-#endif
+		if (bDisplayDebug)
+		{
+			const FVector ZOffset(0.f,0.f,DebugZOffset);
+			UE_VLOG_SEGMENT_THICK(this, LogMassBehavior, Display, AgentPosition + ZOffset, AgentPosition + ZOffset + 100.f * NewGlobalDirection, FColor::White, /*Thickness*/3, TEXT("LookAt Track"));
 		}
+#endif
 	}
 }
 
 bool UMassLookAtProcessor::UpdateGazeTrackedEntity(const FMassEntityManager& EntityManager, const FTransform& Transform, const bool bDisplayDebug, FMassLookAtFragment& LookAt) const
 {
 	bool bHasTarget = false;
-	
+
 	// Update direction toward gaze target
 	if (LookAt.GazeTrackedEntity.IsSet() && EntityManager.IsEntityValid(LookAt.GazeTrackedEntity))
 	{
-		if (const FTransformFragment* TargetTransform = EntityManager.GetFragmentDataPtr<FTransformFragment>(LookAt.GazeTrackedEntity))
-		{
-			FVector LocationOffset;
-			if (const FMassLookAtTargetOffsetFragment* OffsetFragment = EntityManager.GetFragmentDataPtr<FMassLookAtTargetOffsetFragment>(LookAt.GazeTrackedEntity))
-			{
-				LocationOffset = OffsetFragment->Offset;
-			}
-			else
-			{
-				LocationOffset = GetDefault<UMassLookAtSettings>()->GetDefaultTargetLocationOffset();
-			}
+		const FTransformFragment& TargetTransform = EntityManager.GetFragmentDataChecked<FTransformFragment>(LookAt.GazeTrackedEntity);
+		const FMassLookAtTargetFragment& OffsetFragment = EntityManager.GetFragmentDataChecked<FMassLookAtTargetFragment>(LookAt.GazeTrackedEntity);
 
-			const FVector AgentPosition = Transform.GetLocation();
-			const FVector TargetLocation = TargetTransform->GetTransform().TransformPosition(LocationOffset);
-			const FVector NewGlobalDirection = (TargetLocation - AgentPosition).GetSafeNormal();
-			LookAt.GazeDirection = Transform.InverseTransformVector(NewGlobalDirection);
-			LookAt.GazeTargetLocation = TargetLocation;
+		const FVector AgentPosition = Transform.GetLocation();
+		const FVector TargetLocation = TargetTransform.GetTransform().TransformPosition(OffsetFragment.Offset);
+		const FVector NewGlobalDirection = (TargetLocation - AgentPosition).GetSafeNormal();
+		LookAt.GazeDirection = Transform.InverseTransformVector(NewGlobalDirection);
+		LookAt.GazeTargetLocation = TargetLocation;
+
+		bHasTarget = true;
 
-			bHasTarget = true;
-			
 #if WITH_MASSGAMEPLAY_DEBUG
-			if (bDisplayDebug)
-			{
-				const FVector ZOffset(0.f,0.f,DebugZOffset);
-				UE_VLOG_SEGMENT_THICK(this, LogMassBehavior, Display, AgentPosition + ZOffset, AgentPosition + ZOffset + 100.f * NewGlobalDirection, FColor(160,160,160), /*Thickness*/3, TEXT("Gaze Track"));
-			}
-#endif
+		if (bDisplayDebug)
+		{
+			const FVector ZOffset(0.f, 0.f, DebugZOffset);
+			UE_VLOG_SEGMENT_THICK(this, LogMassBehavior, Display, AgentPosition + ZOffset, AgentPosition + ZOffset + 100.f * NewGlobalDirection, FColor(160, 160, 160), /*Thickness*/3, TEXT("Gaze Track"));
 		}
+#endif
 	}
-
 	return bHasTarget;
 }
 
@@ -628,8 +578,7 @@ void UMassLookAtTargetGridProcessor::ConfigureQueries(const TSharedRef<FMassEnti
 {
 	FMassEntityQuery BaseQuery(EntityManager);
 	BaseQuery.AddRequirement<FTransformFragment>(EMassFragmentAccess::ReadOnly);
-	BaseQuery.AddRequirement<FMassLookAtTargetGridCellLocationFragment>(EMassFragmentAccess::ReadWrite);
-	BaseQuery.AddTagRequirement<FMassLookAtTargetTag>(EMassFragmentPresence::All);
+	BaseQuery.AddRequirement<FMassLookAtTargetFragment>(EMassFragmentAccess::ReadWrite);
 	BaseQuery.AddSubsystemRequirement<UMassLookAtSubsystem>(EMassFragmentAccess::ReadWrite);
 
 	AddToGridQuery = BaseQuery;
@@ -657,15 +606,16 @@ void UMassLookAtTargetGridProcessor::Execute(FMassEntityManager& EntityManager,
 	AddToGridQuery.ForEachEntityChunk(Context, [](FMassExecutionContext& Context)
 	{
 		UMassLookAtSubsystem& LookAtSubsystem = Context.GetMutableSubsystemChecked<UMassLookAtSubsystem>();
-		TConstArrayView<FTransformFragment> LocationList = Context.GetFragmentView<FTransformFragment>();
-		TArrayView<FMassLookAtTargetGridCellLocationFragment> CellLocationList = Context.GetMutableFragmentView<FMassLookAtTargetGridCellLocationFragment>();
+		const TConstArrayView<FTransformFragment> LocationList = Context.GetFragmentView<FTransformFragment>();
+		const TArrayView<FMassLookAtTargetFragment> TargetList = Context.GetMutableFragmentView<FMassLookAtTargetFragment>();
 
 		for (FMassExecutionContext::FEntityIterator EntityIt = Context.CreateEntityIterator(); EntityIt; ++EntityIt)
 		{
+			FMassLookAtTargetFragment& Target = TargetList[EntityIt];
 			const FMassEntityHandle Entity = Context.GetEntity(EntityIt);
 			const FVector NewPos = LocationList[EntityIt].GetTransform().GetLocation();
 			const FBox NewBounds(NewPos - FVector(Radius, Radius, 0.f), NewPos + FVector(Radius, Radius, 0.f));
-			CellLocationList[EntityIt].CellLocation = LookAtSubsystem.AddTarget(Entity, NewBounds);
+			Target.CellLocation = LookAtSubsystem.AddTarget(Entity, Target, NewBounds);
 
 			Context.Defer().AddTag<FMassInLookAtTargetGridTag>(Entity);
 		}
@@ -674,15 +624,15 @@ void UMassLookAtTargetGridProcessor::Execute(FMassEntityManager& EntityManager,
 	UpdateGridQuery.ForEachEntityChunk(Context, [](FMassExecutionContext& Context)
 	{
 		UMassLookAtSubsystem& LookAtSubsystem = Context.GetMutableSubsystemChecked<UMassLookAtSubsystem>();
-		TConstArrayView<FTransformFragment> LocationList = Context.GetFragmentView<FTransformFragment>();
-		TArrayView<FMassLookAtTargetGridCellLocationFragment> CellLocationList = Context.GetMutableFragmentView<FMassLookAtTargetGridCellLocationFragment>();
-		TArray<TTuple<const FMassEntityHandle, UE::Mass::LookAt::FTargetHashGrid2D::FCellLocation&, const FBox&>> AllUpdates;
+		const TConstArrayView<FTransformFragment> LocationList = Context.GetFragmentView<FTransformFragment>();
+		const TArrayView<FMassLookAtTargetFragment> CellLocationList = Context.GetMutableFragmentView<FMassLookAtTargetFragment>();
+		TArray<TTuple<const FMassEntityHandle, FMassLookAtTargetFragment&, const FBox&>> AllUpdates;
 
 		for (FMassExecutionContext::FEntityIterator EntityIt = Context.CreateEntityIterator(); EntityIt; ++EntityIt)
 		{
 			const FVector NewPos = LocationList[EntityIt].GetTransform().GetLocation();
 			const FBox NewBounds(NewPos - FVector(Radius, Radius, 0.f), NewPos + FVector(Radius, Radius, 0.f));
-			AllUpdates.Emplace(Context.GetEntity(EntityIt), CellLocationList[EntityIt].CellLocation, NewBounds);
+			AllUpdates.Emplace(Context.GetEntity(EntityIt), CellLocationList[EntityIt], NewBounds);
 		}
 
 		LookAtSubsystem.BatchMoveTarget(AllUpdates);
@@ -691,13 +641,13 @@ void UMassLookAtTargetGridProcessor::Execute(FMassEntityManager& EntityManager,
 	RemoveFromGridQuery.ForEachEntityChunk(Context, [](FMassExecutionContext& Context)
 	{
 		UMassLookAtSubsystem& LookAtSubsystem = Context.GetMutableSubsystemChecked<UMassLookAtSubsystem>();
-		TArrayView<FMassLookAtTargetGridCellLocationFragment> CellLocationList = Context.GetMutableFragmentView<FMassLookAtTargetGridCellLocationFragment>();
+		const TArrayView<FMassLookAtTargetFragment> TargetList = Context.GetMutableFragmentView<FMassLookAtTargetFragment>();
 
 		for (FMassExecutionContext::FEntityIterator EntityIt = Context.CreateEntityIterator(); EntityIt; ++EntityIt)
 		{
 			const FMassEntityHandle Entity = Context.GetEntity(EntityIt);
-			LookAtSubsystem.RemoveTarget(Entity, CellLocationList[EntityIt].CellLocation);
-			CellLocationList[EntityIt].CellLocation = UE::Mass::LookAt::FTargetHashGrid2D::FCellLocation();
+			LookAtSubsystem.RemoveTarget(Entity, TargetList[EntityIt]);
+			TargetList[EntityIt].CellLocation = UE::Mass::LookAt::FTargetHashGrid2D::FCellLocation();
 
 			Context.Defer().RemoveTag<FMassInLookAtTargetGridTag>(Entity);
 		}
@@ -710,14 +660,14 @@ void UMassLookAtTargetGridProcessor::Execute(FMassEntityManager& EntityManager,
 UMassLookAtTargetRemoverProcessor::UMassLookAtTargetRemoverProcessor()
 	: Query(*this)
 {
-	ObservedType = FMassLookAtTargetGridCellLocationFragment::StaticStruct();
+	ObservedType = FMassLookAtTargetFragment::StaticStruct();
 	Operation = EMassObservedOperation::Remove;
 	ExecutionFlags = static_cast<int32>(EProcessorExecutionFlags::AllNetModes);
 }
 
 void UMassLookAtTargetRemoverProcessor::ConfigureQueries(const TSharedRef<FMassEntityManager>& EntityManager)
 {
-	Query.AddRequirement<FMassLookAtTargetGridCellLocationFragment>(EMassFragmentAccess::ReadWrite);
+	Query.AddRequirement<FMassLookAtTargetFragment>(EMassFragmentAccess::ReadWrite);
 	Query.AddSubsystemRequirement<UMassLookAtSubsystem>(EMassFragmentAccess::ReadWrite);
 }
 
@@ -726,11 +676,11 @@ void UMassLookAtTargetRemoverProcessor::Execute(FMassEntityManager& EntityManage
 	Query.ForEachEntityChunk(Context, [this](FMassExecutionContext& Context)
 	{
 		UMassLookAtSubsystem& LookAtSubsystem = Context.GetMutableSubsystemChecked<UMassLookAtSubsystem>();
-		const TArrayView<FMassLookAtTargetGridCellLocationFragment> CellLocationList = Context.GetMutableFragmentView<FMassLookAtTargetGridCellLocationFragment>();
+		const TArrayView<FMassLookAtTargetFragment> TargetList = Context.GetMutableFragmentView<FMassLookAtTargetFragment>();
 
 		for (FMassExecutionContext::FEntityIterator EntityIt = Context.CreateEntityIterator(); EntityIt; ++EntityIt)
 		{
-			LookAtSubsystem.RemoveTarget(Context.GetEntity(EntityIt), CellLocationList[EntityIt].CellLocation);
+			LookAtSubsystem.RemoveTarget(Context.GetEntity(EntityIt), TargetList[EntityIt]);
 		}
 	});
 }
@@ -757,7 +707,7 @@ void UMassLookAtRequestInitializer::Execute(FMassEntityManager& EntityManager, F
 	UMassLookAtSubsystem& LookAtSubsystem = Context.GetMutableSubsystemChecked<UMassLookAtSubsystem>();
 	TArray<UMassLookAtSubsystem::FRequest> Requests;
 
-	Query.ForEachEntityChunk(Context, [&LookAtSubsystem, &Requests](FMassExecutionContext& Context)
+	Query.ForEachEntityChunk(Context, [&Requests](FMassExecutionContext& Context)
 	{
 		const TConstArrayView<FMassLookAtRequestFragment> RequestList = Context.GetFragmentView<FMassLookAtRequestFragment>();
 
@@ -793,10 +743,8 @@ void UMassLookAtRequestDeinitializer::Execute(FMassEntityManager& EntityManager,
 	UMassLookAtSubsystem& LookAtSubsystem = Context.GetMutableSubsystemChecked<UMassLookAtSubsystem>();
 	TArray<FMassEntityHandle> Requests;
 
-	Query.ForEachEntityChunk(Context, [&LookAtSubsystem, &Requests](FMassExecutionContext& Context)
+	Query.ForEachEntityChunk(Context, [&Requests](FMassExecutionContext& Context)
 	{
-		const TConstArrayView<FMassLookAtRequestFragment> RequestList = Context.GetFragmentView<FMassLookAtRequestFragment>();
-
 		Requests.Reserve(Requests.Num() + Context.GetNumEntities());
 		for (FMassExecutionContext::FEntityIterator EntityIt = Context.CreateEntityIterator(); EntityIt; ++EntityIt)
 		{
diff --git a/Engine/Plugins/AI/MassAI/Source/MassAIBehavior/Private/MassLookAtSubsystem.cpp b/Engine/Plugins/AI/MassAI/Source/MassAIBehavior/Private/MassLookAtSubsystem.cpp
--- a/Engine/Plugins/AI/MassAI/Source/MassAIBehavior/Private/MassLookAtSubsystem.cpp
+++ b/Engine/Plugins/AI/MassAI/Source/MassAIBehavior/Private/MassLookAtSubsystem.cpp
@@ -89,6 +89,13 @@ static FAutoConsoleCommandWithWorldArgsAndOutputDevice CmdSendLookAtPlayerReques
 				return;
 			}
 
+			const FMassLookAtTargetFragment* LookAtTargetFragment = EntitySubsystem->GetMutableEntityManager().GetFragmentDataPtr<FMassLookAtTargetFragment>(PlayerEntity);
+			if (LookAtTargetFragment == nullptr)
+			{
+				OutputDevice.Log(ELogVerbosity::Error, TEXT("Command failed: missing required LookAtTarget fragment for the MassEntity associated to the player"));
+				return;
+			}
+
 			TSharedPtr<FMassCommandBuffer> CommandBuffer = MakeShareable(new FMassCommandBuffer());
 			if (bCreateRequest)
 			{
diff --git a/Engine/Plugins/AI/MassAI/Source/MassAIBehavior/Private/MassLookAtTargetTrait.cpp b/Engine/Plugins/AI/MassAI/Source/MassAIBehavior/Private/MassLookAtTargetTrait.cpp
--- a/Engine/Plugins/AI/MassAI/Source/MassAIBehavior/Private/MassLookAtTargetTrait.cpp
+++ b/Engine/Plugins/AI/MassAI/Source/MassAIBehavior/Private/MassLookAtTargetTrait.cpp
@@ -1,12 +1,38 @@
 // Copyright Epic Games, Inc. All Rights Reserved.
 
 #include "MassLookAtTargetTrait.h"
+#include "GameFramework/Actor.h"
 #include "MassLookAtFragments.h"
 #include "MassCommonFragments.h"
 #include "MassEntityTemplateRegistry.h"
+#include "MassEntityView.h"
+#include "MassLookAtSettings.h"
+#include "Components/CapsuleComponent.h"
 
 void UMassLookAtTargetTrait::BuildTemplate(FMassEntityTemplateBuildContext& BuildContext, const UWorld& World) const
 {
-	BuildContext.AddTag<FMassLookAtTargetTag>();
+	const UMassLookAtSettings* Settings = GetDefault<UMassLookAtSettings>();
+
 	BuildContext.AddFragment<FTransformFragment>();
+	BuildContext.AddFragment_GetRef<FMassLookAtTargetFragment>().Offset = Settings->GetDefaultTargetLocationOffset();
+
+	if (bShouldUseCapsuleComponentToSetTargetOffset)
+	{
+		BuildContext.GetMutableObjectFragmentInitializers().Add([=](UObject& Owner, const FMassEntityView& EntityView, const EMassTranslationDirection)
+			{
+				const AActor* AsActor = Cast<AActor>(&Owner);
+
+				const UCapsuleComponent* CapsuleComponent = AsActor
+					? AsActor->FindComponentByClass<UCapsuleComponent>()
+					: Cast<UCapsuleComponent>(&Owner);
+
+				if (CapsuleComponent != nullptr)
+				{
+					double Height = 2 * CapsuleComponent->GetScaledCapsuleHalfHeight();
+					Height *= Settings->GetTargetHeightRatio();
+					Height += Settings->GetFixedOffsetFromTargetHeight();
+					EntityView.GetFragmentData<FMassLookAtTargetFragment>().Offset = FVector(0, 0, Height);
+				}
+			});
+	}
 }
diff --git a/Engine/Plugins/AI/MassAI/Source/MassAIBehavior/Public/MassLookAtFragments.h b/Engine/Plugins/AI/MassAI/Source/MassAIBehavior/Public/MassLookAtFragments.h
--- a/Engine/Plugins/AI/MassAI/Source/MassAIBehavior/Public/MassLookAtFragments.h
+++ b/Engine/Plugins/AI/MassAI/Source/MassAIBehavior/Public/MassLookAtFragments.h
@@ -34,25 +34,6 @@ enum class EMassLookAtGazeMode : uint8
 	Glance,
 };
 
-/**
- * Fragment providing additional offset that could be used by look at processors to get more accurate locations
- */
-USTRUCT()
-struct MASSAIBEHAVIOR_API FMassLookAtTargetOffsetFragment : public FMassFragment
-{
-	GENERATED_BODY()
-
-	FMassLookAtTargetOffsetFragment() = default;
-	explicit FMassLookAtTargetOffsetFragment(const FVector& Offset)
-		: Offset(Offset)
-	{
-	}
-
-	/** Offset in local space to add to the target transform to get the final location */
-	UPROPERTY(Transient)
-	FVector Offset = FVector::ZeroVector;
-};
-
 /**
  * Struct that holds all parameters of the current entity look at 
  */
@@ -157,7 +138,7 @@ struct MASSAIBEHAVIOR_API FMassLookAtFragment : public FMassFragment
  * Special tag to mark an entity that could be tracked by the LookAt
  */
 USTRUCT()
-struct MASSAIBEHAVIOR_API FMassLookAtTargetTag : public FMassTag
+struct UE_DEPRECATED(5.6, "Use FMassLookAtTargetFragment instead") MASSAIBEHAVIOR_API FMassLookAtTargetTag : public FMassTag
 {
 	GENERATED_BODY()
 };
@@ -191,7 +172,7 @@ struct MASSAIBEHAVIOR_API FMassLookAtRequestFragment : public FMassFragment
 	FMassEntityHandle TargetEntity;
 
 	UPROPERTY(Transient)
-	FMassLookAtPriority Priority;
+	FMassLookAtPriority Priority{static_cast<uint8>(EMassLookAtPriorities::LowestPriority)};
 
 	UPROPERTY(Transient)
 	EMassLookAtMode LookAtMode = EMassLookAtMode::LookForward;
@@ -215,12 +196,23 @@ struct MASSAIBEHAVIOR_API FMassLookAtRequestFragment : public FMassFragment
 	}
 };
 
-/** Cell location for LookAt targets */
+/**
+ * Fragment to mark an entity that could be tracked by the LookAt processor
+ * and providing information that could be used to get more accurate locations.
+ */
 USTRUCT()
-struct MASSAIBEHAVIOR_API FMassLookAtTargetGridCellLocationFragment : public FMassFragment
+struct MASSAIBEHAVIOR_API FMassLookAtTargetFragment : public FMassFragment
 {
 	GENERATED_BODY()
 
+	/** Offset in local space to add to the target transform to get the final location */
+	UPROPERTY(Transient)
+	FVector Offset = FVector::ZeroVector;
+
+	/** When a viewer is searching for a random target this priority will influence the selected target */
+	UPROPERTY(Transient)
+	FMassLookAtPriority Priority{static_cast<uint8>(EMassLookAtPriorities::LowestPriority)};
+
 	UE::Mass::LookAt::FTargetHashGrid2D::FCellLocation CellLocation;
 };
 
diff --git a/Engine/Plugins/AI/MassAI/Source/MassAIBehavior/Public/MassLookAtSettings.h b/Engine/Plugins/AI/MassAI/Source/MassAIBehavior/Public/MassLookAtSettings.h
--- a/Engine/Plugins/AI/MassAI/Source/MassAIBehavior/Public/MassLookAtSettings.h
+++ b/Engine/Plugins/AI/MassAI/Source/MassAIBehavior/Public/MassLookAtSettings.h
@@ -30,6 +30,27 @@ public:
 		return DefaultTargetLocationOffset;
 	}
 
+	/**
+	 * Specifies a modifier applied to the target height, if available, to adjust the final location to look at.
+	 * Value of 0 represents the base location and a value of 1 the base location with the full target height offset.
+	 * TargetHeightRatio is applied before FixedOffsetFromTargetHeight* 
+	 * @return Modifier applied on the target height to compute the final location to look at
+	 */
+	[[nodiscard]] double GetTargetHeightRatio() const
+	{
+		return TargetHeightRatio;
+	}
+
+	/**
+	 * Specifies the offset to add to the target height, if available, to adjust the final location to look at.
+	 * Negative value should be used to lower the target.
+	 * @return Offset applied on the target height to compute the final location to look at
+	 */
+	[[nodiscard]] double GetFixedOffsetFromTargetHeight() const
+	{
+		return FixedOffsetFromTargetHeight;
+	}
+
 	static MASSAIBEHAVIOR_API FOnMassLookAtPrioritiesChanged OnMassLookAtPrioritiesChanged;
 
 protected:
@@ -37,7 +58,28 @@ protected:
 	UPROPERTY(EditAnywhere, config, Category = LookAt)
 	FMassLookAtPriorityInfo Priorities[static_cast<int32>(EMassLookAtPriorities::MaxPriorities)];
 
-	/** Additional offset to add to the target location (aka eyes height) */
+	/**
+	 * Additional offset added to the target base location (i.e., TransformFragment).
+	 * This is used by default when no specific height information (i.e., LookAtTargetTrait initializer) is available.
+	 */
 	UPROPERTY(EditAnywhere, config, Category = LookAt)
 	FVector DefaultTargetLocationOffset = FVector::ZeroVector;
+
+	/**
+	 * Optional height modifier ratio applied to the target height, if available, to adjust the final location to look at.
+	 * Value of 0 represents the base location and a value of 1 the base location with the full target height offset.
+	 * TargetHeightRatio is applied before FixedOffsetFromTargetHeight
+	 * @see FixedOffsetFromTargetHeight
+	 */
+	UPROPERTY(EditAnywhere, config, Category = LookAt, meta = (EditCondition = bShouldUseCapsuleComponentToSetTargetOffset))
+	double TargetHeightRatio = 1.0;
+
+	/**
+	 * Optional fixed offset (in cm) added to the target height, if available, to adjust the final location to look at.
+	 * Negative value should be used to lower the target.
+	 * TargetHeightRatio is applied before FixedOffsetFromTargetHeight
+	 * @see TargetHeightRatio
+	 */
+	UPROPERTY(EditAnywhere, config, Category = LookAt, meta = (EditCondition = bShouldUseCapsuleComponentToSetTargetOffset))
+	double FixedOffsetFromTargetHeight = 0;
 };
\ No newline at end of file
diff --git a/Engine/Plugins/AI/MassAI/Source/MassAIBehavior/Public/MassLookAtSubsystem.h b/Engine/Plugins/AI/MassAI/Source/MassAIBehavior/Public/MassLookAtSubsystem.h
--- a/Engine/Plugins/AI/MassAI/Source/MassAIBehavior/Public/MassLookAtSubsystem.h
+++ b/Engine/Plugins/AI/MassAI/Source/MassAIBehavior/Public/MassLookAtSubsystem.h
@@ -47,26 +47,27 @@ public:
 	/**
 	 * Adds new item to the grid for a given entity handle representing a LookAt target.
 	 * @param InEntity - Entity used to identify the item in the grid.
+	 * @param InTarget - Fragment providing details about the target (i.e., priority)
 	 * @param InBounds - Bounding box of the entity.
 	 * @return The cell location where the item was added to, can be used later to move or remove the item.
 	 */
-	[[nodiscard]] UE::Mass::LookAt::FTargetHashGrid2D::FCellLocation AddTarget(const FMassEntityHandle InEntity, const FBox& InBounds)
+	[[nodiscard]] UE::Mass::LookAt::FTargetHashGrid2D::FCellLocation AddTarget(const FMassEntityHandle InEntity, const FMassLookAtTargetFragment& InTarget, const FBox& InBounds)
 	{
 		UE_MT_SCOPED_WRITE_ACCESS(TargetGridAccessDetector);
-		return TargetGrid.Add(InEntity, InBounds);
+		return TargetGrid.Add(UE::Mass::LookAt::FTargetHashGrid2D::ItemIDType{InEntity, InTarget.Priority.Get()}, InBounds);
 	}
 
 	/**
 	 * Moves item based on the entity handle, its previous cell location and its new bounding box.
 	 * @param InEntity - Entity used to identify the item in the grid.
-	 * @param InOldCellLocation - The same cell location the item was previously added or moved with.
+	 * @param InTarget - Fragment providing details about the target (i.e., cell location the item was previously added or moved with).
 	 * @param InNewBounds - New bounds of the item
 	 * @return The cell location where the item is now added to.
 	 */
-	[[nodiscard]] UE::Mass::LookAt::FTargetHashGrid2D::FCellLocation MoveTarget(const FMassEntityHandle InEntity, const UE::Mass::LookAt::FTargetHashGrid2D::FCellLocation InOldCellLocation, const FBox& InNewBounds)
+	[[nodiscard]] UE::Mass::LookAt::FTargetHashGrid2D::FCellLocation MoveTarget(const FMassEntityHandle InEntity, const FMassLookAtTargetFragment& InTarget, const FBox& InNewBounds)
 	{
 		UE_MT_SCOPED_WRITE_ACCESS(TargetGridAccessDetector);
-		return TargetGrid.Move(InEntity, InOldCellLocation, InNewBounds);
+		return TargetGrid.Move(UE::Mass::LookAt::FTargetHashGrid2D::ItemIDType{InEntity, InTarget.Priority.Get()}, InTarget.CellLocation, InNewBounds);
 	}
 
 	/**
@@ -74,26 +75,29 @@ public:
 	 * @param InUpdates - List of entity, cell location reference and new bounds to use for each item to update in the grid.
 	 * @note Cell location provided in the tuple is a reference and will be updated in the process.
 	 */
-	void BatchMoveTarget(TArrayView<TTuple<const FMassEntityHandle, UE::Mass::LookAt::FTargetHashGrid2D::FCellLocation&, const FBox&>> InUpdates)
+	void BatchMoveTarget(const TArrayView<TTuple<const FMassEntityHandle, FMassLookAtTargetFragment&, const FBox&>> InUpdates)
 	{
 		UE_MT_SCOPED_WRITE_ACCESS(TargetGridAccessDetector);
 		TArray<UE::Mass::LookAt::FTargetHashGrid2D::FCellLocation> UpdatedLocations;
 		UpdatedLocations.Reserve(InUpdates.Num());
-		for (TTuple<const FMassEntityHandle, UE::Mass::LookAt::FTargetHashGrid2D::FCellLocation&, const FBox&>& UpdateInfo : InUpdates)
+		for (TTuple<const FMassEntityHandle, FMassLookAtTargetFragment&, const FBox&>& UpdateInfo : InUpdates)
 		{
-			UpdateInfo.Get<1>() = TargetGrid.Move(/*Entity*/UpdateInfo.Get<0>(), /*OldCellLocation*/UpdateInfo.Get<1>(), /*NewBounds*/UpdateInfo.Get<2>());
+			UpdateInfo.Get<1>().CellLocation = TargetGrid.Move(
+				/*Item*/UE::Mass::LookAt::FTargetHashGrid2D::ItemIDType{UpdateInfo.Get<0>(), UpdateInfo.Get<1>().Priority.Get()},
+				/*OldCellLocation*/UpdateInfo.Get<1>().CellLocation,
+				/*NewBounds*/UpdateInfo.Get<2>());
 		}
 	}
 
 	/**
 	 * Removes item based on the entity handle and the cell location it was added with.
 	 * @param InEntity - Entity used to identify the item in the grid.
-	 * @param InOldCellLocation - The same cell location the item was previously added or moved with.
+	 * @param InTarget - Fragment providing details about the target (i.e., cell location the item was previously added or moved with).
 	 */
-	void RemoveTarget(const FMassEntityHandle InEntity, const UE::Mass::LookAt::FTargetHashGrid2D::FCellLocation InOldCellLocation)
+	void RemoveTarget(const FMassEntityHandle InEntity, const FMassLookAtTargetFragment& InTarget)
 	{
 		UE_MT_SCOPED_WRITE_ACCESS(TargetGridAccessDetector);
-		TargetGrid.Remove(InEntity, InOldCellLocation);
+		TargetGrid.Remove(UE::Mass::LookAt::FTargetHashGrid2D::ItemIDType{InEntity, InTarget.Priority.Get()}, InTarget.CellLocation);
 	}
 
 	/**
@@ -101,7 +105,8 @@ public:
 	 * @param InQueryBox - Query bounding box.
 	 * @param OutEntities - Result of the query, entity handles of potentially overlapping items.
 	 */
-	bool Query(const FBox& InQueryBox, TArray<UE::Mass::LookAt::FTargetHashGrid2D::ItemIDType>& OutEntities) const
+	template <typename OutT>
+	bool Query(const FBox& InQueryBox, OutT& OutEntities) const
 	{
 		UE_MT_SCOPED_READ_ACCESS(TargetGridAccessDetector);
 		TargetGrid.Query(InQueryBox, OutEntities);
diff --git a/Engine/Plugins/AI/MassAI/Source/MassAIBehavior/Public/MassLookAtTargetTrait.h b/Engine/Plugins/AI/MassAI/Source/MassAIBehavior/Public/MassLookAtTargetTrait.h
--- a/Engine/Plugins/AI/MassAI/Source/MassAIBehavior/Public/MassLookAtTargetTrait.h
+++ b/Engine/Plugins/AI/MassAI/Source/MassAIBehavior/Public/MassLookAtTargetTrait.h
@@ -16,4 +16,10 @@ class UMassLookAtTargetTrait : public UMassEntityTraitBase
 protected:
 
 	virtual void BuildTemplate(FMassEntityTemplateBuildContext& BuildContext, const UWorld& World) const override;
+
+	/**
+	 * Indicates whether the trait will use an initializer to set target offset using the height of the capsule component if available.
+	 */
+	UPROPERTY(EditAnywhere, config, Category = LookAt)
+	bool bShouldUseCapsuleComponentToSetTargetOffset = true;
 };
diff --git a/Engine/Plugins/AI/MassAI/Source/MassAIBehavior/Public/MassLookAtTypes.h b/Engine/Plugins/AI/MassAI/Source/MassAIBehavior/Public/MassLookAtTypes.h
--- a/Engine/Plugins/AI/MassAI/Source/MassAIBehavior/Public/MassLookAtTypes.h
+++ b/Engine/Plugins/AI/MassAI/Source/MassAIBehavior/Public/MassLookAtTypes.h
@@ -11,8 +11,32 @@ namespace UE::Mass::LookAt
 {
 	constexpr int32 HashGridLevelsOfHierarchy = 2;
 	constexpr int32 HashGridRatioBetweenLevels = 4;
+	constexpr int32 HashGridResultsSoftLimit = 16;
 
-	using FTargetHashGrid2D = THierarchicalHashGrid2D<HashGridLevelsOfHierarchy, HashGridRatioBetweenLevels, FMassEntityHandle>;
+	struct FTargetHashGridItem
+	{
+		FTargetHashGridItem(const FMassEntityHandle& TargetEntity, const uint8 Priority)
+			: TargetEntity(TargetEntity)
+			, Priority(Priority)
+		{
+		}
+
+		bool operator==(const FTargetHashGridItem& RHS) const
+		{
+			// We only allow a single entry per entity so no need to compare the priority
+			return TargetEntity == RHS.TargetEntity;
+		}
+
+		bool operator!=(const FTargetHashGridItem& RHS) const
+		{
+			return !(*this == RHS);
+		}
+
+		FMassEntityHandle TargetEntity;
+		uint8 Priority = MAX_uint8;
+	};
+
+	using FTargetHashGrid2D = THierarchicalHashGrid2D<HashGridLevelsOfHierarchy, HashGridRatioBetweenLevels, FTargetHashGridItem>;
 }
 
 /**
diff --git a/Engine/Plugins/AI/MassAI/Source/MassAIDebug/Private/GameplayDebuggerCategory_Mass.cpp b/Engine/Plugins/AI/MassAI/Source/MassAIDebug/Private/GameplayDebuggerCategory_Mass.cpp
--- a/Engine/Plugins/AI/MassAI/Source/MassAIDebug/Private/GameplayDebuggerCategory_Mass.cpp
+++ b/Engine/Plugins/AI/MassAI/Source/MassAIDebug/Private/GameplayDebuggerCategory_Mass.cpp
@@ -509,7 +509,7 @@ void FGameplayDebuggerCategory_Mass::CollectData(APlayerController* OwnerPC, AAc
 		EntityQuery.AddRequirement<FMassForceFragment>(EMassFragmentAccess::ReadOnly);
 		EntityQuery.AddRequirement<FMassMoveTargetFragment>(EMassFragmentAccess::ReadOnly);
 		EntityQuery.AddRequirement<FMassLookAtFragment>(EMassFragmentAccess::ReadOnly, EMassFragmentPresence::Optional);
-		EntityQuery.AddRequirement<FMassLookAtTargetOffsetFragment>(EMassFragmentAccess::ReadOnly, EMassFragmentPresence::Optional);
+		EntityQuery.AddRequirement<FMassLookAtTargetFragment>(EMassFragmentAccess::ReadOnly, EMassFragmentPresence::Optional);
 		EntityQuery.AddRequirement<FMassSimulationLODFragment>(EMassFragmentAccess::ReadOnly, EMassFragmentPresence::Optional);
 		EntityQuery.AddRequirement<FMassZoneGraphShortPathFragment>(EMassFragmentAccess::ReadOnly, EMassFragmentPresence::Optional);
 		EntityQuery.AddRequirement<FMassNavMeshShortPathFragment>(EMassFragmentAccess::ReadOnly, EMassFragmentPresence::Optional);
@@ -539,7 +539,7 @@ void FGameplayDebuggerCategory_Mass::CollectData(APlayerController* OwnerPC, AAc
 				const TConstArrayView<FMassForceFragment> ForceList = Context.GetFragmentView<FMassForceFragment>();
 				const TConstArrayView<FMassMoveTargetFragment> MoveTargetList = Context.GetFragmentView<FMassMoveTargetFragment>();
 				const TConstArrayView<FMassLookAtFragment> LookAtList = Context.GetFragmentView<FMassLookAtFragment>();
-				const TConstArrayView<FMassLookAtTargetOffsetFragment> LookAtOffsetList = Context.GetFragmentView<FMassLookAtTargetOffsetFragment>();
+				const TConstArrayView<FMassLookAtTargetFragment> LookAtTargetList = Context.GetFragmentView<FMassLookAtTargetFragment>();
 				const TConstArrayView<FMassSimulationLODFragment> SimLODList = Context.GetFragmentView<FMassSimulationLODFragment>();
 				const TConstArrayView<FMassZoneGraphShortPathFragment> ZoneGraphShortPathList = Context.GetFragmentView<FMassZoneGraphShortPathFragment>();
 				const TConstArrayView<FMassNavMeshShortPathFragment> NavMeshShortPathList = Context.GetFragmentView<FMassNavMeshShortPathFragment>();
@@ -614,9 +614,9 @@ void FGameplayDebuggerCategory_Mass::CollectData(APlayerController* OwnerPC, AAc
 						BasePos = EntityLocation;
 
 						// Apply the offset for the entity looking at something
-						if (LookAtOffsetList.Num())
+						if (LookAtTargetList.Num())
 						{
-							BasePos += LookAtOffsetList[EntityIndex].Offset;
+							BasePos += LookAtTargetList[EntityIndex].Offset;
 						}
 						else
 						{
diff --git a/Engine/Source/Runtime/AIModule/Public/HierarchicalHashGrid2D.h b/Engine/Source/Runtime/AIModule/Public/HierarchicalHashGrid2D.h
--- a/Engine/Source/Runtime/AIModule/Public/HierarchicalHashGrid2D.h
+++ b/Engine/Source/Runtime/AIModule/Public/HierarchicalHashGrid2D.h
@@ -337,7 +337,8 @@ public:
 	  * @param Bounds - Query bounding box.
 	  * @param OutResults - Result of the query, IDs of potentially overlapping items.
 	  */
-	void QuerySmall(const FBox& Bounds, TArray<ItemIDType>& OutResults) const
+	template <typename OutT>
+	void QuerySmall(const FBox& Bounds, OutT& OutResults) const
 	{
 		// Return items from buckets on all levels
 		for (int32 Level = 0; Level < NumLevels; Level++)
@@ -377,7 +378,8 @@ public:
 	  * @param Bounds - Query bounding box.
 	  * @param OutResults - Result of the query, IDs of potentially overlapping items.
 	  */
-	void Query(const FBox& Bounds, TArray<ItemIDType>& OutResults) const
+	template <typename OutT>
+	void Query(const FBox& Bounds, OutT& OutResults) const
 	{
 		FCellRect Rects[NumLevels];
 		FCellRectIterator Iters[NumLevels];
