[LookAt] updated processor to consider FMassLookAtTargetFragment as optional

diff --git a/Engine/Plugins/AI/MassAI/Source/MassAIBehavior/Private/MassLookAtProcessors.cpp b/Engine/Plugins/AI/MassAI/Source/MassAIBehavior/Private/MassLookAtProcessors.cpp
--- a/Engine/Plugins/AI/MassAI/Source/MassAIBehavior/Private/MassLookAtProcessors.cpp
+++ b/Engine/Plugins/AI/MassAI/Source/MassAIBehavior/Private/MassLookAtProcessors.cpp
@@ -420,10 +420,12 @@ bool UMassLookAtProcessor::UpdateGazeTrackedEntity(const FMassEntityManager& Ent
 	if (LookAt.GazeTrackedEntity.IsSet() && EntityManager.IsEntityValid(LookAt.GazeTrackedEntity))
 	{
 		const FTransformFragment& TargetTransform = EntityManager.GetFragmentDataChecked<FTransformFragment>(LookAt.GazeTrackedEntity);
-		const FMassLookAtTargetFragment& OffsetFragment = EntityManager.GetFragmentDataChecked<FMassLookAtTargetFragment>(LookAt.GazeTrackedEntity);
+		const FMassLookAtTargetFragment* OffsetFragment = EntityManager.GetFragmentDataPtr<FMassLookAtTargetFragment>(LookAt.GazeTrackedEntity);
 
 		const FVector AgentPosition = Transform.GetLocation();
-		const FVector TargetLocation = TargetTransform.GetTransform().TransformPosition(OffsetFragment.Offset);
+		const FVector TargetLocation = OffsetFragment != nullptr
+			? TargetTransform.GetTransform().TransformPosition(OffsetFragment->Offset)
+			: TargetTransform.GetTransform().GetLocation();
 		const FVector NewGlobalDirection = (TargetLocation - AgentPosition).GetSafeNormal();
 		LookAt.GazeDirection = Transform.InverseTransformVector(NewGlobalDirection);
 		LookAt.GazeTargetLocation = TargetLocation;
