PR #13068: Avoid data copy in UMassNavMeshPathFollowProcessor

#rnx

diff --git a/Engine/Plugins/AI/MassAI/Source/MassNavMeshNavigation/Private/MassNavMeshNavigationProcessors.cpp b/Engine/Plugins/AI/MassAI/Source/MassNavMeshNavigation/Private/MassNavMeshNavigationProcessors.cpp
--- a/Engine/Plugins/AI/MassAI/Source/MassNavMeshNavigation/Private/MassNavMeshNavigationProcessors.cpp
+++ b/Engine/Plugins/AI/MassAI/Source/MassNavMeshNavigation/Private/MassNavMeshNavigationProcessors.cpp
@@ -114,8 +114,8 @@ void UMassNavMeshPathFollowProcessor::Execute(FMassEntityManager& EntityManager,
 						}
 					}
 
-					const FMassNavMeshPathPoint PointA = ShortPath.Points[ClosestSegmentIndex];
-					const FMassNavMeshPathPoint PointB = ShortPath.Points[ClosestSegmentIndex+1];
+					const FMassNavMeshPathPoint& PointA = ShortPath.Points[ClosestSegmentIndex];
+					const FMassNavMeshPathPoint& PointB = ShortPath.Points[ClosestSegmentIndex+1];
 					const float DistanceOnClosestSegment =  (EntityLocation - PointA.Position).Dot((PointB.Position - PointA.Position).GetSafeNormal());
 					const float EntityProgress = ShortPath.Points[ClosestSegmentIndex].Distance.Get() + DistanceOnClosestSegment;
 					MoveTarget.EntityDistanceToGoal = FMath::Max(0.f, ShortPath.Points[UpdatePointIndex].Distance.Get() - EntityProgress);
