Mass Gameplay Debug Tool -> Added keybindings to toggle lookat and entity descriptions,
with the additon of small updates to the entity debug display

diff --git a/Engine/Plugins/AI/MassAI/Source/MassAIDebug/Private/GameplayDebuggerCategory_Mass.cpp b/Engine/Plugins/AI/MassAI/Source/MassAIDebug/Private/GameplayDebuggerCategory_Mass.cpp
--- a/Engine/Plugins/AI/MassAI/Source/MassAIDebug/Private/GameplayDebuggerCategory_Mass.cpp
+++ b/Engine/Plugins/AI/MassAI/Source/MassAIDebug/Private/GameplayDebuggerCategory_Mass.cpp
@@ -100,6 +100,8 @@ FGameplayDebuggerCategory_Mass::FOnConsoleCommandBroadcastDelegate FGameplayDebu
 FGameplayDebuggerCategory_Mass::FOnConsoleCommandBroadcastDelegate FGameplayDebuggerCategory_Mass::OnToggleNearEntityOverviewBroadcast;
 FGameplayDebuggerCategory_Mass::FOnConsoleCommandBroadcastDelegate FGameplayDebuggerCategory_Mass::OnToggleNearEntityAvoidanceBroadcast;
 FGameplayDebuggerCategory_Mass::FOnConsoleCommandBroadcastDelegate FGameplayDebuggerCategory_Mass::OnToggleNearEntityPathBroadcast;
+FGameplayDebuggerCategory_Mass::FOnConsoleCommandBroadcastDelegate FGameplayDebuggerCategory_Mass::OnToggleEntityLookAtBroadcast;
+FGameplayDebuggerCategory_Mass::FOnConsoleCommandBroadcastDelegate FGameplayDebuggerCategory_Mass::OnToggleEntityDescriptionBroadcast;
 FGameplayDebuggerCategory_Mass::FOnConsoleCommandBroadcastDelegate FGameplayDebuggerCategory_Mass::OnToggleDebugLocalEntityManagerBroadcast;
 FGameplayDebuggerCategory_Mass::FOnConsoleCommandBroadcastDelegate FGameplayDebuggerCategory_Mass::OnTogglePickedActorAsViewerBroadcast;
 FGameplayDebuggerCategory_Mass::FOnConsoleCommandBroadcastDelegate FGameplayDebuggerCategory_Mass::OnToggleDrawViewersBroadcast;
@@ -119,6 +121,8 @@ FGameplayDebuggerCategory_Mass::FGameplayDebuggerCategory_Mass()
 	bShowNearEntityOverview = true;
 	bShowNearEntityAvoidance = false;
 	bShowNearEntityPath = false;
+	bShowEntityLookAt = false;
+	bShowEntityDescription = true;
 	bMarkEntityBeingDebugged = true;
 	bDebugLocalEntityManager = false;
 	bShowViewers = false;
@@ -130,7 +134,9 @@ FGameplayDebuggerCategory_Mass::FGameplayDebuggerCategory_Mass()
 	BindKeyPress(EKeys::D.GetFName(), FGameplayDebuggerInputModifier::Shift, this, &FGameplayDebuggerCategory_Mass::OnToggleEntityDetails, EGameplayDebuggerInputMode::Replicated);
 	BindKeyPress(EKeys::O.GetFName(), FGameplayDebuggerInputModifier::Shift, this, &FGameplayDebuggerCategory_Mass::OnToggleNearEntityOverview, EGameplayDebuggerInputMode::Replicated);
 	BindKeyPress(EKeys::V.GetFName(), FGameplayDebuggerInputModifier::Shift, this, &FGameplayDebuggerCategory_Mass::OnToggleNearEntityAvoidance, EGameplayDebuggerInputMode::Replicated);
-	BindKeyPress(EKeys::C.GetFName(), FGameplayDebuggerInputModifier::Shift, this, &FGameplayDebuggerCategory_Mass::OnToggleNearEntityPath, EGameplayDebuggerInputMode::Replicated);
+	BindKeyPress(EKeys::N.GetFName(), FGameplayDebuggerInputModifier::Shift, this, &FGameplayDebuggerCategory_Mass::OnToggleNearEntityPath, EGameplayDebuggerInputMode::Replicated);
+	BindKeyPress(EKeys::K.GetFName(), FGameplayDebuggerInputModifier::Shift, this, &FGameplayDebuggerCategory_Mass::OnToggleEntityLookAt, EGameplayDebuggerInputMode::Replicated);
+	BindKeyPress(EKeys::E.GetFName(), FGameplayDebuggerInputModifier::Shift, this, &FGameplayDebuggerCategory_Mass::OnToggleEntityDescription, EGameplayDebuggerInputMode::Replicated);
 	ToggleDebugLocalEntityManagerInputIndex = GetNumInputHandlers();
 	BindKeyPress(EKeys::L.GetFName(), FGameplayDebuggerInputModifier::Shift, this, &FGameplayDebuggerCategory_Mass::OnToggleDebugLocalEntityManager, EGameplayDebuggerInputMode::Local);
 	BindKeyPress(EKeys::Add.GetFName(), FGameplayDebuggerInputModifier::Shift, this, &FGameplayDebuggerCategory_Mass::OnIncreaseSearchRange, EGameplayDebuggerInputMode::Replicated);
@@ -152,6 +158,8 @@ FGameplayDebuggerCategory_Mass::FGameplayDebuggerCategory_Mass()
 		ConsoleCommands.Emplace(TEXT("gdt.mass.ToggleNearEntityOverview"), TEXT(""), FConsoleCommandWithWorldDelegate::CreateLambda([](UWorld* InWorld) { OnToggleNearEntityOverviewBroadcast.Broadcast(InWorld); }));
 		ConsoleCommands.Emplace(TEXT("gdt.mass.ToggleNearEntityAvoidance"), TEXT(""), FConsoleCommandWithWorldDelegate::CreateLambda([](UWorld* InWorld) { OnToggleNearEntityAvoidanceBroadcast.Broadcast(InWorld); }));
 		ConsoleCommands.Emplace(TEXT("gdt.mass.ToggleNearEntityPath"), TEXT(""), FConsoleCommandWithWorldDelegate::CreateLambda([](UWorld* InWorld) { OnToggleNearEntityPathBroadcast.Broadcast(InWorld); }));
+		ConsoleCommands.Emplace(TEXT("gdt.mass.ToggleEntityLookAt"), TEXT(""), FConsoleCommandWithWorldDelegate::CreateLambda([](UWorld* InWorld) { OnToggleEntityLookAtBroadcast.Broadcast(InWorld); }));
+		ConsoleCommands.Emplace(TEXT("gdt.mass.ToggleEntityDescription"), TEXT(""), FConsoleCommandWithWorldDelegate::CreateLambda([](UWorld* InWorld) { OnToggleEntityDescriptionBroadcast.Broadcast(InWorld); }));
 		ConsoleCommands.Emplace(TEXT("gdt.mass.ToggleDebugLocalEntityManager"), TEXT(""), FConsoleCommandWithWorldDelegate::CreateLambda([](UWorld* InWorld) { OnToggleDebugLocalEntityManagerBroadcast.Broadcast(InWorld); }));
 		ConsoleCommands.Emplace(TEXT("gdt.mass.TogglePickedActorAsViewer"), TEXT(""), FConsoleCommandWithWorldDelegate::CreateLambda([](UWorld* InWorld) { OnTogglePickedActorAsViewerBroadcast.Broadcast(InWorld); }));
 		ConsoleCommands.Emplace(TEXT("gdt.mass.ToggleDrawViewers"), TEXT(""), FConsoleCommandWithWorldDelegate::CreateLambda([](UWorld* InWorld) { OnToggleDrawViewersBroadcast.Broadcast(InWorld); }));	
@@ -166,6 +174,8 @@ FGameplayDebuggerCategory_Mass::FGameplayDebuggerCategory_Mass()
 	ConsoleCommandHandles.Add(FDelegateHandlePair(&OnToggleNearEntityOverviewBroadcast, OnToggleNearEntityOverviewBroadcast.AddLambda([this](UWorld* InWorld) { if (InWorld == GetWorldFromReplicator()) { OnToggleNearEntityOverview(); }})));
 	ConsoleCommandHandles.Add(FDelegateHandlePair(&OnToggleNearEntityAvoidanceBroadcast, OnToggleNearEntityAvoidanceBroadcast.AddLambda([this](UWorld* InWorld) { if (InWorld == GetWorldFromReplicator()) { OnToggleNearEntityAvoidance(); }})));
 	ConsoleCommandHandles.Add(FDelegateHandlePair(&OnToggleNearEntityPathBroadcast, OnToggleNearEntityPathBroadcast.AddLambda([this](UWorld* InWorld) { if (InWorld == GetWorldFromReplicator()) { OnToggleNearEntityPath(); }})));
+	ConsoleCommandHandles.Add(FDelegateHandlePair(&OnToggleEntityLookAtBroadcast, OnToggleEntityLookAtBroadcast.AddLambda([this](UWorld* InWorld) { if (InWorld == GetWorldFromReplicator()) { OnToggleEntityLookAt(); }})));
+	ConsoleCommandHandles.Add(FDelegateHandlePair(&OnToggleEntityDescriptionBroadcast, OnToggleEntityDescriptionBroadcast.AddLambda([this](UWorld* InWorld) { if (InWorld == GetWorldFromReplicator()) { OnToggleEntityDescription(); }})));
 	ConsoleCommandHandles.Add(FDelegateHandlePair(&OnToggleDebugLocalEntityManagerBroadcast, OnToggleDebugLocalEntityManagerBroadcast.AddLambda([this](UWorld* InWorld) { if (InWorld == GetWorldFromReplicator()) { OnToggleDebugLocalEntityManager(); }})));
 	ConsoleCommandHandles.Add(FDelegateHandlePair(&OnTogglePickedActorAsViewerBroadcast, OnTogglePickedActorAsViewerBroadcast.AddLambda([this](UWorld* InWorld) { if (InWorld == GetWorldFromReplicator()) { OnTogglePickedActorAsViewer(); }})));
 	ConsoleCommandHandles.Add(FDelegateHandlePair(&OnToggleDrawViewersBroadcast, OnToggleDrawViewersBroadcast.AddLambda([this](UWorld* InWorld) { if (InWorld == GetWorldFromReplicator()) { OnToggleDrawViewers(); }})));
@@ -587,23 +597,22 @@ void FGameplayDebuggerCategory_Mass::CollectData(APlayerController* OwnerPC, AAc
 					const FVector EntityForward = Transform.GetTransform().GetRotation().GetForwardVector();
 
 					// Draw entity position and orientation.
-					FVector BasePos = EntityLocation + FVector(0 ,0 ,25);
+					const FVector ZBaseOffset(0, 0, 20);
+					const FVector ZDeltaOffset(0, 0, 1);
+					FVector BasePos = EntityLocation + ZBaseOffset;
 
 					AddShape(FGameplayDebuggerShape::MakeCircle(BasePos, FVector::UpVector, Radius.Radius, FColor::White));
-					AddShape(FGameplayDebuggerShape::MakeSegment(BasePos, BasePos + EntityForward * Radius.Radius * 1.25f, FColor::White));
+					AddShape(FGameplayDebuggerShape::MakeSegment(BasePos, BasePos + EntityForward * Radius.Radius, FColor::White));
 
 					// Velocity and steering target
-					BasePos += FVector(0 ,0 ,5);
+					BasePos += ZDeltaOffset;
 					AddShape(FGameplayDebuggerShape::MakeArrow(BasePos, BasePos + Velocity.Value, 10.0f, 2.0f, FColor::Yellow));
-					BasePos += FVector(0 ,0 ,5);
-					AddShape(FGameplayDebuggerShape::MakeArrow(BasePos, BasePos + Steering.DesiredVelocity, 10.0f, 1.0f, FColorList::Pink));
 
-					// Move target
-					const FVector MoveBasePos = MoveTarget.Center + FVector(0,0,5);
-					AddShape(FGameplayDebuggerShape::MakeArrow(MoveBasePos - MoveTarget.Forward * Radius.Radius, MoveBasePos + MoveTarget.Forward * Radius.Radius, 10.0f, 2.0f, FColorList::MediumVioletRed));
+					BasePos += ZDeltaOffset;
+					AddShape(FGameplayDebuggerShape::MakeArrow(BasePos, BasePos + Steering.DesiredVelocity, 10.0f, 1.0f, FColorList::Pink));
 
 					// Look at
-					if (bHasLookAt)
+					if (bShowEntityLookAt && bHasLookAt)
 					{
 						constexpr FVector::FReal LookArrowLength = 100.;
 						constexpr float TargetArrowHeadSize = 10.f;
@@ -659,17 +668,20 @@ void FGameplayDebuggerCategory_Mass::CollectData(APlayerController* OwnerPC, AAc
 						const FMassSmartObjectUserFragment& SOUser = SOUserList[EntityIndex];
 						if (SOUser.InteractionHandle.IsValid())
 						{
-							const FVector ZOffset = FVector(0.0f , 0.0f , 25.0f );
 							const FTransform SlotTransform = SmartObjectSubsystem->GetSlotTransform(SOUser.InteractionHandle).Get(FTransform::Identity);
 							const FVector SlotLocation = SlotTransform.GetLocation();
-							AddShape(FGameplayDebuggerShape::MakeSegment(EntityLocation + ZOffset, SlotLocation + ZOffset, 3.0f, FColorList::Orange));
+							AddShape(FGameplayDebuggerShape::MakeSegment(EntityLocation + ZBaseOffset, SlotLocation + ZBaseOffset, 3.0f, FColorList::Orange));
 						}
 					}
 
 					// Path
 					if (bShowNearEntityPath)
 					{
-						const FVector ZOffset = FVector(0.f , 0.f , 25.f );
+						// Move target
+						const FVector MoveBasePos = MoveTarget.Center + ZBaseOffset + ZDeltaOffset;
+						AddShape(FGameplayDebuggerShape::MakeCircle(MoveBasePos, FVector::UpVector, 5.0f, FColorList::MediumVioletRed));
+						AddShape(FGameplayDebuggerShape::MakeArrow(MoveBasePos,
+							MoveBasePos + MoveTarget.Forward * Radius.Radius, 10.0f, 1.0f, FColorList::MediumVioletRed));
 					
 						if (!ZoneGraphShortPathList.IsEmpty())
 						{
@@ -679,13 +691,13 @@ void FGameplayDebuggerCategory_Mass::CollectData(APlayerController* OwnerPC, AAc
 							{
 								const FMassZoneGraphPathPoint& CurrPoint = ShortPath.Points[PointIndex];
 								const FMassZoneGraphPathPoint& NextPoint = ShortPath.Points[PointIndex + 1];
-								AddShape(FGameplayDebuggerShape::MakeSegment(CurrPoint.Position + ZOffset, NextPoint.Position + ZOffset, /*Thickness*/3.f, FColorList::Grey));
+								AddShape(FGameplayDebuggerShape::MakeSegment(CurrPoint.Position + ZBaseOffset, NextPoint.Position + ZBaseOffset, /*Thickness*/3.f, FColorList::Grey));
 							}
 					
 							for (uint8 PointIndex = 0; PointIndex < ShortPath.NumPoints; PointIndex++)
 							{
 								const FMassZoneGraphPathPoint& CurrPoint = ShortPath.Points[PointIndex];
-								const FVector CurrBase = CurrPoint.Position + ZOffset;
+								const FVector CurrBase = CurrPoint.Position + ZBaseOffset;
 								// Lane tangents
 								AddShape(FGameplayDebuggerShape::MakeSegment(CurrBase, CurrBase + CurrPoint.Tangent.GetVector() * 50.f, /*Thickness*/1.f, FColorList::LightGrey));
 							}
@@ -699,13 +711,13 @@ void FGameplayDebuggerCategory_Mass::CollectData(APlayerController* OwnerPC, AAc
 							{
 								const FMassNavMeshPathPoint& CurrPoint = ShortPath.Points[PointIndex];
 								const FMassNavMeshPathPoint& NextPoint = ShortPath.Points[PointIndex + 1];
-								AddShape(FGameplayDebuggerShape::MakeSegment(CurrPoint.Position + ZOffset, NextPoint.Position + ZOffset, /*Thickness*/3.f, FColorList::GreenYellow));
+								AddShape(FGameplayDebuggerShape::MakeSegment(CurrPoint.Position + ZBaseOffset, NextPoint.Position + ZBaseOffset, /*Thickness*/3.f, FColorList::GreenYellow));
 							}
 					
 							for (uint8 PointIndex = 0; PointIndex < ShortPath.NumPoints; PointIndex++)
 							{
 								const FMassNavMeshPathPoint& CurrPoint = ShortPath.Points[PointIndex];
-								const FVector CurrBase = CurrPoint.Position + ZOffset;
+								const FVector CurrBase = CurrPoint.Position + ZBaseOffset;
 								// Path tangents
 								AddShape(FGameplayDebuggerShape::MakeSegment(CurrBase, CurrBase + CurrPoint.Tangent.GetVector() * 50.f, /*Thickness*/1.f, FColorList::LightGrey));
 							}
@@ -719,7 +731,7 @@ void FGameplayDebuggerCategory_Mass::CollectData(APlayerController* OwnerPC, AAc
 						// Standing avoidance.
 						if (Ghost.IsValid(MoveTarget.GetCurrentActionID()))
 						{
-							FVector GhostBasePos = Ghost.Location + FVector(0.0f ,0.0f ,25.0f );
+							FVector GhostBasePos = Ghost.Location + ZBaseOffset;
 							AddShape(FGameplayDebuggerShape::MakeCircle(GhostBasePos, FVector::UpVector, Radius.Radius, FColorList::LightGrey));
 							GhostBasePos += FVector(0,0,5);
 							AddShape(FGameplayDebuggerShape::MakeArrow(GhostBasePos, GhostBasePos + Ghost.Velocity, 10.0f, 2.0f, FColorList::LightGrey));
@@ -730,7 +742,7 @@ void FGameplayDebuggerCategory_Mass::CollectData(APlayerController* OwnerPC, AAc
 					}
 					
 					// Status
-					if (DistanceToEntitySq < FMath::Square(MaxViewDistance * 0.5f))
+					if (bShowEntityDescription && DistanceToEntitySq < FMath::Square(MaxViewDistance * 0.5f))
 					{
 						FString Status;
 
@@ -780,7 +792,7 @@ void FGameplayDebuggerCategory_Mass::CollectData(APlayerController* OwnerPC, AAc
 							*UEnum::GetDisplayValueAsText(MoveTarget.IntentAtGoal).ToString(), MoveTarget.DistanceToGoal);
 
 						// Look
-						if (bHasLookAt)
+						if (bShowEntityLookAt && bHasLookAt)
 						{
 							const FMassLookAtFragment& LookAt = LookAtList[EntityIndex];
 							const double RemainingTime = LookAt.GazeDuration - (CurrentTime - LookAt.GazeStartTime);
@@ -874,19 +886,21 @@ void FGameplayDebuggerCategory_Mass::DrawData(APlayerController* OwnerPC, FGamep
 	CanvasContext.Printf(TEXT("[{yellow}%s{white}] %s Agent Fragments"), *GetInputHandlerDescription(2), bShowAgentFragments ? TEXT("Hide") : TEXT("Show"));
 	if (bShowAgentFragments)
 	{
-		CanvasContext.Printf(TEXT("[{yellow}%s{white}] %s Entity details"), *GetInputHandlerDescription(4), bShowEntityDetails ? TEXT("Hide") : TEXT("Show"));
+		CanvasContext.Printf(TEXT("[{yellow}%s{white}] %s Entity Details"), *GetInputHandlerDescription(4), bShowEntityDetails ? TEXT("Hide") : TEXT("Show"));
 	}
 	else
 	{
-		CanvasContext.Printf(TEXT("{grey}[%s] Entity details [enable Agent Fragments]{white}"), *GetInputHandlerDescription(4));
+		CanvasContext.Printf(TEXT("{grey}[%s] Entity Details [enable Agent Fragments]{white}"), *GetInputHandlerDescription(4));
 	}
 	CanvasContext.Printf(TEXT("[{yellow}%s{white}] Pick Entity"), *GetInputHandlerDescription(3));
-	CanvasContext.Printf(TEXT("[{yellow}%s{white}] Toggle Picked Actor as viewer"), *GetInputHandlerDescription(TogglePickedActorAsViewerInputIndex));
+	CanvasContext.Printf(TEXT("[{yellow}%s{white}] Toggle Picked Actor as Viewer"), *GetInputHandlerDescription(TogglePickedActorAsViewerInputIndex));
 	CanvasContext.Printf(TEXT("[{yellow}%s{white}] %s LOD Viewers"), *GetInputHandlerDescription(ToggleDrawViewersInputIndex), bShowViewers ? TEXT("Hide") : TEXT("Show"));
 	CanvasContext.Printf(TEXT("[{yellow}%s{white}] Reset Actor LOD Viewers"), *GetInputHandlerDescription(ClearViewersInputIndex));
-	CanvasContext.Printf(TEXT("[{yellow}%s{white}] %s Entity overview"), *GetInputHandlerDescription(5), bShowNearEntityOverview ? TEXT("Hide") : TEXT("Show"));
-	CanvasContext.Printf(TEXT("[{yellow}%s{white}] %s Entity avoidance"), *GetInputHandlerDescription(6), bShowNearEntityAvoidance ? TEXT("Hide") : TEXT("Show"));
-	CanvasContext.Printf(TEXT("[{yellow}%s{white}] %s Entity path"), *GetInputHandlerDescription(7), bShowNearEntityPath ? TEXT("Hide") : TEXT("Show"));
+	CanvasContext.Printf(TEXT("[{yellow}%s{white}] %s Entity Overview"), *GetInputHandlerDescription(5), bShowNearEntityOverview ? TEXT("Hide") : TEXT("Show"));
+	CanvasContext.Printf(TEXT("[{yellow}%s{white}] %s Entity Avoidance"), *GetInputHandlerDescription(6), bShowNearEntityAvoidance ? TEXT("Hide") : TEXT("Show"));
+	CanvasContext.Printf(TEXT("[{yellow}%s{white}] %s Entity Path"), *GetInputHandlerDescription(7), bShowNearEntityPath ? TEXT("Hide") : TEXT("Show"));
+	CanvasContext.Printf(TEXT("[{yellow}%s{white}] %s Entity LookAt"), *GetInputHandlerDescription(8), bShowEntityLookAt ? TEXT("Hide") : TEXT("Show"));
+	CanvasContext.Printf(TEXT("[{yellow}%s{white}] %s Entity Description"), *GetInputHandlerDescription(9), bShowEntityDescription ? TEXT("Hide") : TEXT("Show"));
 
 	if (IsCategoryLocal() && !IsCategoryAuth())
 	{
diff --git a/Engine/Plugins/AI/MassAI/Source/MassAIDebug/Public/GameplayDebuggerCategory_Mass.h b/Engine/Plugins/AI/MassAI/Source/MassAIDebug/Public/GameplayDebuggerCategory_Mass.h
--- a/Engine/Plugins/AI/MassAI/Source/MassAIDebug/Public/GameplayDebuggerCategory_Mass.h
+++ b/Engine/Plugins/AI/MassAI/Source/MassAIDebug/Public/GameplayDebuggerCategory_Mass.h
@@ -43,6 +43,14 @@ protected:
 	void OnToggleNearEntityOverview() { bShowNearEntityOverview = !bShowNearEntityOverview; }
 	void OnToggleNearEntityAvoidance() { bShowNearEntityAvoidance = !bShowNearEntityAvoidance; }
 	void OnToggleNearEntityPath() { bShowNearEntityPath = !bShowNearEntityPath; }
+	void OnToggleEntityLookAt() 
+	{ 
+		bShowEntityLookAt = !bShowEntityLookAt; 
+	}
+	void OnToggleEntityDescription() 
+	{ 
+		bShowEntityDescription = !bShowEntityDescription; 
+	}
 	void OnToggleDebugLocalEntityManager();
 	void OnIncreaseSearchRange();
 	void OnDecreaseSearchRange();
@@ -69,6 +77,8 @@ protected:
 	bool bShowNearEntityOverview;
 	bool bShowNearEntityAvoidance;
 	bool bShowNearEntityPath;
+	bool bShowEntityLookAt;
+	bool bShowEntityDescription;
 	bool bMarkEntityBeingDebugged;
 	bool bDebugLocalEntityManager;
 	bool bShowViewers;
@@ -99,6 +109,8 @@ protected:
 	static FOnConsoleCommandBroadcastDelegate OnToggleNearEntityOverviewBroadcast;
 	static FOnConsoleCommandBroadcastDelegate OnToggleNearEntityAvoidanceBroadcast;
 	static FOnConsoleCommandBroadcastDelegate OnToggleNearEntityPathBroadcast; 	
+	static FOnConsoleCommandBroadcastDelegate OnToggleEntityLookAtBroadcast;
+	static FOnConsoleCommandBroadcastDelegate OnToggleEntityDescriptionBroadcast;
 	static FOnConsoleCommandBroadcastDelegate OnToggleDebugLocalEntityManagerBroadcast; 	
 	static FOnConsoleCommandBroadcastDelegate OnTogglePickedActorAsViewerBroadcast;
 	static FOnConsoleCommandBroadcastDelegate OnToggleDrawViewersBroadcast;
