[LookAt]
- exposed target priority in the LookAtTrait
- added fragment to disable systemic LookAt for actors bound to cinematic in Sequencer

diff --git a/Engine/Plugins/AI/MassAI/Source/MassAIBehavior/Private/MassLookAtTargetTrait.cpp b/Engine/Plugins/AI/MassAI/Source/MassAIBehavior/Private/MassLookAtTargetTrait.cpp
--- a/Engine/Plugins/AI/MassAI/Source/MassAIBehavior/Private/MassLookAtTargetTrait.cpp
+++ b/Engine/Plugins/AI/MassAI/Source/MassAIBehavior/Private/MassLookAtTargetTrait.cpp
@@ -14,7 +14,9 @@ void UMassLookAtTargetTrait::BuildTemplate(FMassEntityTemplateBuildContext& Buil
 	const UMassLookAtSettings* Settings = GetDefault<UMassLookAtSettings>();
 
 	BuildContext.AddFragment<FTransformFragment>();
-	BuildContext.AddFragment_GetRef<FMassLookAtTargetFragment>().Offset = Settings->GetDefaultTargetLocationOffset();
+	FMassLookAtTargetFragment& Target = BuildContext.AddFragment_GetRef<FMassLookAtTargetFragment>();
+	Target.Offset = Settings->GetDefaultTargetLocationOffset();
+	Target.Priority = Priority;
 
 	if (bShouldUseCapsuleComponentToSetTargetOffset)
 	{
diff --git a/Engine/Plugins/AI/MassAI/Source/MassAIBehavior/Public/MassLookAtTargetTrait.h b/Engine/Plugins/AI/MassAI/Source/MassAIBehavior/Public/MassLookAtTargetTrait.h
--- a/Engine/Plugins/AI/MassAI/Source/MassAIBehavior/Public/MassLookAtTargetTrait.h
+++ b/Engine/Plugins/AI/MassAI/Source/MassAIBehavior/Public/MassLookAtTargetTrait.h
@@ -3,6 +3,7 @@
 #pragma once
 
 #include "MassEntityTraitBase.h"
+#include "MassLookAtTypes.h"
 #if UE_ENABLE_INCLUDE_ORDER_DEPRECATED_IN_5_6
 #include "CoreMinimal.h"
 #endif // UE_ENABLE_INCLUDE_ORDER_DEPRECATED_IN_5_6
@@ -22,4 +23,8 @@ protected:
 	 */
 	UPROPERTY(EditAnywhere, config, Category = LookAt)
 	bool bShouldUseCapsuleComponentToSetTargetOffset = true;
+
+	/** Priority assigned to the target to influence target selection */
+	UPROPERTY(EditAnywhere, config, Category = LookAt)
+	FMassLookAtPriority Priority{static_cast<uint8>(EMassLookAtPriorities::LowestPriority)};
 };
