[LookAt] minor fixes

diff --git a/Engine/Plugins/AI/MassAI/Source/MassAIBehavior/Private/MassLookAtProcessors.cpp b/Engine/Plugins/AI/MassAI/Source/MassAIBehavior/Private/MassLookAtProcessors.cpp
--- a/Engine/Plugins/AI/MassAI/Source/MassAIBehavior/Private/MassLookAtProcessors.cpp
+++ b/Engine/Plugins/AI/MassAI/Source/MassAIBehavior/Private/MassLookAtProcessors.cpp
@@ -395,9 +395,12 @@ void UMassLookAtProcessor::UpdateLookAtTrackedEntity(const FMassEntityManager& E
 	if (EntityManager.IsEntityValid(LookAt.TrackedEntity))
 	{
 		const FTransformFragment& TargetTransform = EntityManager.GetFragmentDataChecked<FTransformFragment>(LookAt.TrackedEntity);
-		const FVector LocationOffset = EntityManager.GetFragmentDataChecked<FMassLookAtTargetFragment>(LookAt.TrackedEntity).Offset;
+		const FMassLookAtTargetFragment* TargetFragment = EntityManager.GetFragmentDataPtr<FMassLookAtTargetFragment>(LookAt.TrackedEntity);
+
 		const FVector AgentPosition = Transform.GetLocation();
-		const FVector TargetLocation = TargetTransform.GetTransform().TransformPosition(LocationOffset);
+		const FVector TargetLocation = TargetFragment != nullptr
+			? TargetTransform.GetTransform().TransformPosition(TargetFragment->Offset)
+			: TargetTransform.GetTransform().GetLocation();
 		const FVector NewGlobalDirection = (TargetLocation - AgentPosition).GetSafeNormal();
 		LookAt.Direction = Transform.InverseTransformVector(NewGlobalDirection);
 		LookAt.MainTargetLocation = TargetLocation;
@@ -420,11 +423,11 @@ bool UMassLookAtProcessor::UpdateGazeTrackedEntity(const FMassEntityManager& Ent
 	if (LookAt.GazeTrackedEntity.IsSet() && EntityManager.IsEntityValid(LookAt.GazeTrackedEntity))
 	{
 		const FTransformFragment& TargetTransform = EntityManager.GetFragmentDataChecked<FTransformFragment>(LookAt.GazeTrackedEntity);
-		const FMassLookAtTargetFragment* OffsetFragment = EntityManager.GetFragmentDataPtr<FMassLookAtTargetFragment>(LookAt.GazeTrackedEntity);
+		const FMassLookAtTargetFragment* TargetFragment = EntityManager.GetFragmentDataPtr<FMassLookAtTargetFragment>(LookAt.GazeTrackedEntity);
 
 		const FVector AgentPosition = Transform.GetLocation();
-		const FVector TargetLocation = OffsetFragment != nullptr
-			? TargetTransform.GetTransform().TransformPosition(OffsetFragment->Offset)
+		const FVector TargetLocation = TargetFragment != nullptr
+			? TargetTransform.GetTransform().TransformPosition(TargetFragment->Offset)
 			: TargetTransform.GetTransform().GetLocation();
 		const FVector NewGlobalDirection = (TargetLocation - AgentPosition).GetSafeNormal();
 		LookAt.GazeDirection = Transform.InverseTransformVector(NewGlobalDirection);
diff --git a/Engine/Plugins/AI/MassAI/Source/MassAIBehavior/Private/MassLookAtSubsystem.cpp b/Engine/Plugins/AI/MassAI/Source/MassAIBehavior/Private/MassLookAtSubsystem.cpp
--- a/Engine/Plugins/AI/MassAI/Source/MassAIBehavior/Private/MassLookAtSubsystem.cpp
+++ b/Engine/Plugins/AI/MassAI/Source/MassAIBehavior/Private/MassLookAtSubsystem.cpp
@@ -76,7 +76,7 @@ static FAutoConsoleCommandWithWorldArgsAndOutputDevice CmdSendLookAtPlayerReques
 			}
 
 			UMassEntitySubsystem* EntitySubsystem = World->GetSubsystem<UMassEntitySubsystem>();
-			if (ActorSubsystem == nullptr)
+			if (EntitySubsystem == nullptr)
 			{
 				OutputDevice.Log(ELogVerbosity::Error, TEXT("Command failed: unable to find UMassEntitySubsystem"));
 				return;
