[LookAt] fix hash grid update

diff --git a/Engine/Plugins/AI/MassAI/Source/MassAIBehavior/Private/MassLookAtProcessors.cpp b/Engine/Plugins/AI/MassAI/Source/MassAIBehavior/Private/MassLookAtProcessors.cpp
--- a/Engine/Plugins/AI/MassAI/Source/MassAIBehavior/Private/MassLookAtProcessors.cpp
+++ b/Engine/Plugins/AI/MassAI/Source/MassAIBehavior/Private/MassLookAtProcessors.cpp
@@ -631,7 +631,7 @@ void UMassLookAtTargetGridProcessor::Execute(FMassEntityManager& EntityManager,
 		UMassLookAtSubsystem& LookAtSubsystem = Context.GetMutableSubsystemChecked<UMassLookAtSubsystem>();
 		const TConstArrayView<FTransformFragment> LocationList = Context.GetFragmentView<FTransformFragment>();
 		const TArrayView<FMassLookAtTargetFragment> CellLocationList = Context.GetMutableFragmentView<FMassLookAtTargetFragment>();
-		TArray<TTuple<const FMassEntityHandle, FMassLookAtTargetFragment&, const FBox&>> AllUpdates;
+		TArray<TTuple<const FMassEntityHandle, FMassLookAtTargetFragment&, const FBox>> AllUpdates;
 
 		for (FMassExecutionContext::FEntityIterator EntityIt = Context.CreateEntityIterator(); EntityIt; ++EntityIt)
 		{
diff --git a/Engine/Plugins/AI/MassAI/Source/MassAIBehavior/Public/MassLookAtSubsystem.h b/Engine/Plugins/AI/MassAI/Source/MassAIBehavior/Public/MassLookAtSubsystem.h
--- a/Engine/Plugins/AI/MassAI/Source/MassAIBehavior/Public/MassLookAtSubsystem.h
+++ b/Engine/Plugins/AI/MassAI/Source/MassAIBehavior/Public/MassLookAtSubsystem.h
@@ -75,12 +75,12 @@ public:
 	 * @param InUpdates - List of entity, cell location reference and new bounds to use for each item to update in the grid.
 	 * @note Cell location provided in the tuple is a reference and will be updated in the process.
 	 */
-	void BatchMoveTarget(const TArrayView<TTuple<const FMassEntityHandle, FMassLookAtTargetFragment&, const FBox&>> InUpdates)
+	void BatchMoveTarget(const TArrayView<TTuple<const FMassEntityHandle, FMassLookAtTargetFragment&, const FBox>> InUpdates)
 	{
 		UE_MT_SCOPED_WRITE_ACCESS(TargetGridAccessDetector);
 		TArray<UE::Mass::LookAt::FTargetHashGrid2D::FCellLocation> UpdatedLocations;
 		UpdatedLocations.Reserve(InUpdates.Num());
-		for (TTuple<const FMassEntityHandle, FMassLookAtTargetFragment&, const FBox&>& UpdateInfo : InUpdates)
+		for (TTuple<const FMassEntityHandle, FMassLookAtTargetFragment&, const FBox>& UpdateInfo : InUpdates)
 		{
 			UpdateInfo.Get<1>().CellLocation = TargetGrid.Move(
 				/*Item*/UE::Mass::LookAt::FTargetHashGrid2D::ItemIDType{UpdateInfo.Get<0>(), UpdateInfo.Get<1>().Priority.Get()},
