[MassGameplay] fixed LookAt's entity destruction commands to properly indicate these are "Destroy"-type commands

diff --git a/Engine/Plugins/AI/MassAI/Source/MassAIBehavior/Private/MassLookAtSubsystem.cpp b/Engine/Plugins/AI/MassAI/Source/MassAIBehavior/Private/MassLookAtSubsystem.cpp
--- a/Engine/Plugins/AI/MassAI/Source/MassAIBehavior/Private/MassLookAtSubsystem.cpp
+++ b/Engine/Plugins/AI/MassAI/Source/MassAIBehavior/Private/MassLookAtSubsystem.cpp
@@ -137,7 +137,7 @@ static FAutoConsoleCommandWithWorldArgsAndOutputDevice CmdSendLookAtPlayerReques
 			else
 			{
 				// Delete all entities created for debug requests
-				CommandBuffer->PushCommand<FMassDeferredCreateCommand>([](FMassEntityManager& InEntityManager)
+				CommandBuffer->PushCommand<FMassDeferredDestroyCommand>([](FMassEntityManager& InEntityManager)
 					{
 						InEntityManager.BatchDestroyEntities(DebugRequests);
 						DebugRequests.Reset();
diff --git a/Engine/Plugins/AI/MassAI/Source/MassAIBehavior/Private/Tasks/MassLookAtTask.cpp b/Engine/Plugins/AI/MassAI/Source/MassAIBehavior/Private/Tasks/MassLookAtTask.cpp
--- a/Engine/Plugins/AI/MassAI/Source/MassAIBehavior/Private/Tasks/MassLookAtTask.cpp
+++ b/Engine/Plugins/AI/MassAI/Source/MassAIBehavior/Private/Tasks/MassLookAtTask.cpp
@@ -95,7 +95,7 @@ void FMassLookAtTask::ExitState(FStateTreeExecutionContext& Context, const FStat
 	if (InstanceData.Request.IsSet())
 	{
 		MassStateTreeContext.GetMassEntityExecutionContext().Defer()
-			.PushCommand<FMassDeferredCreateCommand>([Request = InstanceData.Request](FMassEntityManager& InEntityManager)
+			.PushCommand<FMassDeferredDestroyCommand>([Request = InstanceData.Request](FMassEntityManager& InEntityManager)
 			{
 				InEntityManager.DestroyEntity(Request);
 			});
