[MassGameplay] updated GameplayDebugger category to handle more fragments as optional.
Any entity with a transform fragment will now have a representation using a default red circle when FAgentRadiusFragment is missing

diff --git a/Engine/Plugins/AI/MassAI/Source/MassAIDebug/Private/GameplayDebuggerCategory_Mass.cpp b/Engine/Plugins/AI/MassAI/Source/MassAIDebug/Private/GameplayDebuggerCategory_Mass.cpp
--- a/Engine/Plugins/AI/MassAI/Source/MassAIDebug/Private/GameplayDebuggerCategory_Mass.cpp
+++ b/Engine/Plugins/AI/MassAI/Source/MassAIDebug/Private/GameplayDebuggerCategory_Mass.cpp
@@ -508,16 +508,16 @@ void FGameplayDebuggerCategory_Mass::CollectData(APlayerController* OwnerPC, AAc
 	if (bShowNearEntityOverview && OwnerPC)
 	{
 		FMassEntityQuery EntityQuery(EntityManager.AsShared());
-		EntityQuery.AddRequirement<FMassStateTreeInstanceFragment>(EMassFragmentAccess::ReadOnly);
-		EntityQuery.AddConstSharedRequirement<FMassStateTreeSharedFragment>();
+		EntityQuery.AddRequirement<FMassStateTreeInstanceFragment>(EMassFragmentAccess::ReadOnly, EMassFragmentPresence::Optional);
+		EntityQuery.AddConstSharedRequirement<FMassStateTreeSharedFragment>(EMassFragmentPresence::Optional);
 		EntityQuery.AddRequirement<FTransformFragment>(EMassFragmentAccess::ReadOnly);
-		EntityQuery.AddRequirement<FAgentRadiusFragment>(EMassFragmentAccess::ReadOnly);
-		EntityQuery.AddRequirement<FMassSteeringFragment>(EMassFragmentAccess::ReadOnly);
-		EntityQuery.AddRequirement<FMassStandingSteeringFragment>(EMassFragmentAccess::ReadOnly);
+		EntityQuery.AddRequirement<FAgentRadiusFragment>(EMassFragmentAccess::ReadOnly, EMassFragmentPresence::Optional);
+		EntityQuery.AddRequirement<FMassSteeringFragment>(EMassFragmentAccess::ReadOnly, EMassFragmentPresence::Optional);
+		EntityQuery.AddRequirement<FMassStandingSteeringFragment>(EMassFragmentAccess::ReadOnly, EMassFragmentPresence::Optional);
 		EntityQuery.AddRequirement<FMassGhostLocationFragment>(EMassFragmentAccess::ReadOnly, EMassFragmentPresence::Optional);
-		EntityQuery.AddRequirement<FMassVelocityFragment>(EMassFragmentAccess::ReadOnly);
-		EntityQuery.AddRequirement<FMassForceFragment>(EMassFragmentAccess::ReadOnly);
-		EntityQuery.AddRequirement<FMassMoveTargetFragment>(EMassFragmentAccess::ReadOnly);
+		EntityQuery.AddRequirement<FMassVelocityFragment>(EMassFragmentAccess::ReadOnly, EMassFragmentPresence::Optional);
+		EntityQuery.AddRequirement<FMassForceFragment>(EMassFragmentAccess::ReadOnly, EMassFragmentPresence::Optional);
+		EntityQuery.AddRequirement<FMassMoveTargetFragment>(EMassFragmentAccess::ReadOnly, EMassFragmentPresence::Optional);
 		EntityQuery.AddRequirement<FMassLookAtFragment>(EMassFragmentAccess::ReadOnly, EMassFragmentPresence::Optional);
 		EntityQuery.AddRequirement<FMassLookAtTargetFragment>(EMassFragmentAccess::ReadOnly, EMassFragmentPresence::Optional);
 		EntityQuery.AddRequirement<FMassSimulationLODFragment>(EMassFragmentAccess::ReadOnly, EMassFragmentPresence::Optional);
@@ -554,18 +554,23 @@ void FGameplayDebuggerCategory_Mass::CollectData(APlayerController* OwnerPC, AAc
 				const TConstArrayView<FMassZoneGraphShortPathFragment> ZoneGraphShortPathList = Context.GetFragmentView<FMassZoneGraphShortPathFragment>();
 				const TConstArrayView<FMassNavMeshShortPathFragment> NavMeshShortPathList = Context.GetFragmentView<FMassNavMeshShortPathFragment>();
 				const TConstArrayView<FMassSmartObjectUserFragment> SOUserList = Context.GetFragmentView<FMassSmartObjectUserFragment>();
-				const FMassStateTreeSharedFragment& SharedStateTree = Context.GetConstSharedFragment<FMassStateTreeSharedFragment>();
 
+				const bool bHasForce = (ForceList.Num() > 0);
+				const bool bHasGhostLocation = (GhostList.Num() > 0);
 				const bool bHasLOD = (SimLODList.Num() > 0);
 				const bool bHasLookAt = (LookAtList.Num() > 0);
+				const bool bHasMoveTarget = (MoveTargetList.Num() > 0);
+				const bool bHasRadius = (RadiusList.Num() > 0);
 				const bool bHasSOUser = (SOUserList.Num() > 0);
+				const bool bHasStandingSteering = (StandingSteeringList.Num() > 0);
+				const bool bHasStateTree = (StateTreeInstanceList.Num() > 0);
+				const bool bHasSteering = (SteeringList.Num() > 0);
+				const bool bHasVelocity = (VelocityList.Num() > 0);
 
 				const UGameplayDebuggerUserSettings* Settings = GetDefault<UGameplayDebuggerUserSettings>();
 				const FVector::FReal MaxViewDistance = Settings->MaxViewDistance;
 				const FVector::FReal MinViewDirDot = FMath::Cos(FMath::DegreesToRadians(Settings->MaxViewAngle));
 
-				const UStateTree* StateTree = SharedStateTree.StateTree;
-
 				for (int32 EntityIndex = 0; EntityIndex < NumEntities; ++EntityIndex)
 				{
 					FMassEntityHandle Entity = Context.GetEntity(EntityIndex);
@@ -585,31 +590,37 @@ void FGameplayDebuggerCategory_Mass::CollectData(APlayerController* OwnerPC, AAc
 						continue;
 					}
 
-					const FAgentRadiusFragment& Radius = RadiusList[EntityIndex];
-					const FMassSteeringFragment& Steering = SteeringList[EntityIndex];
-					const FMassStandingSteeringFragment& StandingSteering = StandingSteeringList[EntityIndex];
-					const FMassVelocityFragment& Velocity = VelocityList[EntityIndex];
-					const FMassForceFragment& Force = ForceList[EntityIndex];
-					const FMassMoveTargetFragment& MoveTarget = MoveTargetList[EntityIndex];
-					const FMassSimulationLODFragment& SimLOD = bHasLOD ? SimLODList[EntityIndex] : FMassSimulationLODFragment();
-					const FMassStateTreeInstanceFragment& StateTreeInstance = StateTreeInstanceList[EntityIndex];
-
-					const FVector EntityForward = Transform.GetTransform().GetRotation().GetForwardVector();
-
 					// Draw entity position and orientation.
 					const FVector ZBaseOffset(0, 0, 20);
 					const FVector ZDeltaOffset(0, 0, 1);
 					FVector BasePos = EntityLocation + ZBaseOffset;
+					const FVector EntityForward = Transform.GetTransform().GetRotation().GetForwardVector();
+					constexpr float DebugShapeRadius = 10.f;
 
-					AddShape(FGameplayDebuggerShape::MakeCircle(BasePos, FVector::UpVector, Radius.Radius, FColor::White));
-					AddShape(FGameplayDebuggerShape::MakeSegment(BasePos, BasePos + EntityForward * Radius.Radius, FColor::White));
+					if (bHasRadius)
+					{
+						AddShape(FGameplayDebuggerShape::MakeCircle(BasePos, FVector::UpVector, RadiusList[EntityIndex].Radius, FColor::White));
+						AddShape(FGameplayDebuggerShape::MakeSegment(BasePos, BasePos + EntityForward * RadiusList[EntityIndex].Radius, FColor::White));
+					}
+					else
+					{
+						AddShape(FGameplayDebuggerShape::MakeCircle(BasePos, FVector::UpVector, DebugShapeRadius, FColor::Red));
+						AddShape(FGameplayDebuggerShape::MakeSegment(BasePos, BasePos + EntityForward * DebugShapeRadius, FColor::Red));
+					}
 
-					// Velocity and steering target
-					BasePos += ZDeltaOffset;
-					AddShape(FGameplayDebuggerShape::MakeArrow(BasePos, BasePos + Velocity.Value, 10.0f, 2.0f, FColor::Yellow));
+					if (bHasVelocity)
+					{
+						// Velocity target
+						BasePos += ZDeltaOffset;
+						AddShape(FGameplayDebuggerShape::MakeArrow(BasePos, BasePos + VelocityList[EntityIndex].Value, 10.0f, 2.0f, FColor::Yellow));
+					}
 
-					BasePos += ZDeltaOffset;
-					AddShape(FGameplayDebuggerShape::MakeArrow(BasePos, BasePos + Steering.DesiredVelocity, 10.0f, 1.0f, FColorList::Pink));
+					if (bHasSteering)
+					{
+						// Steering target
+						BasePos += ZDeltaOffset;
+						AddShape(FGameplayDebuggerShape::MakeArrow(BasePos, BasePos + SteeringList[EntityIndex].DesiredVelocity, 10.0f, 1.0f, FColorList::Pink));
+					}
 
 					// Look at
 					if (bShowEntityLookAt && bHasLookAt)
@@ -640,7 +651,7 @@ void FGameplayDebuggerCategory_Mass::CollectData(APlayerController* OwnerPC, AAc
 						{
 							// Target location is expected to include its own offset
 							FVector TargetPosition = LookAt.MainTargetLocation;
-							AddShape(FGameplayDebuggerShape::MakeCircle(TargetPosition, FVector::UpVector, Radius.Radius, FColor::Red));
+							AddShape(FGameplayDebuggerShape::MakeCircle(TargetPosition, FVector::UpVector, bHasRadius ? RadiusList[EntityIndex].Radius : DebugShapeRadius, FColor::Red));
 							AddShape(FGameplayDebuggerShape::MakeArrow(BasePos, TargetPosition, TargetArrowHeadSize, TargetArrowThickness, FColor::Red));
 
 							const FVector::FReal TargetDistance = FMath::Max(LookArrowLength, FVector::DotProduct(WorldLookDirection, TargetPosition - BasePos));
@@ -652,7 +663,7 @@ void FGameplayDebuggerCategory_Mass::CollectData(APlayerController* OwnerPC, AAc
 						{
 							// Target location is expected to include its own offset
 							FVector TargetPosition = LookAt.GazeTargetLocation;
-							AddShape(FGameplayDebuggerShape::MakeCircle(TargetPosition, FVector::UpVector, Radius.Radius, FColorList::Goldenrod));
+							AddShape(FGameplayDebuggerShape::MakeCircle(TargetPosition, FVector::UpVector, bHasRadius ? RadiusList[EntityIndex].Radius : DebugShapeRadius, FColorList::Goldenrod));
 							AddShape(FGameplayDebuggerShape::MakeArrow(BasePos, TargetPosition, TargetArrowHeadSize, TargetArrowThickness, FColorList::Goldenrod));
 						}
 
@@ -675,13 +686,14 @@ void FGameplayDebuggerCategory_Mass::CollectData(APlayerController* OwnerPC, AAc
 					}
 
 					// Path
-					if (bShowNearEntityPath)
+					if (bShowNearEntityPath && bHasMoveTarget && bHasRadius)
 					{
 						// Move target
+						const FMassMoveTargetFragment& MoveTarget = MoveTargetList[EntityIndex];
 						const FVector MoveBasePos = MoveTarget.Center + ZBaseOffset + ZDeltaOffset;
 						AddShape(FGameplayDebuggerShape::MakeCircle(MoveBasePos, FVector::UpVector, 5.0f, FColorList::MediumVioletRed));
 						AddShape(FGameplayDebuggerShape::MakeArrow(MoveBasePos,
-							MoveBasePos + MoveTarget.Forward * Radius.Radius, 10.0f, 1.0f, FColorList::MediumVioletRed));
+							MoveBasePos + MoveTarget.Forward * RadiusList[EntityIndex].Radius, 10.0f, 1.0f, FColorList::MediumVioletRed));
 					
 						if (!ZoneGraphShortPathList.IsEmpty())
 						{
@@ -724,8 +736,10 @@ void FGameplayDebuggerCategory_Mass::CollectData(APlayerController* OwnerPC, AAc
 						}
 					}
 
-					if (bShowNearEntityAvoidance && GhostList.Num() > 0)
+					if (bShowNearEntityAvoidance && bHasMoveTarget && bHasRadius && bHasGhostLocation)
 					{
+						const FMassMoveTargetFragment& MoveTarget = MoveTargetList[EntityIndex];
+						const FAgentRadiusFragment& Radius = RadiusList[EntityIndex];
 						const FMassGhostLocationFragment& Ghost = GhostList[EntityIndex];
 
 						// Standing avoidance.
@@ -736,8 +750,11 @@ void FGameplayDebuggerCategory_Mass::CollectData(APlayerController* OwnerPC, AAc
 							GhostBasePos += FVector(0,0,5);
 							AddShape(FGameplayDebuggerShape::MakeArrow(GhostBasePos, GhostBasePos + Ghost.Velocity, 10.0f, 2.0f, FColorList::LightGrey));
 
-							const FVector GhostTargetBasePos = StandingSteering.TargetLocation + FVector(0.0f ,0.0f ,25.0f );
-							AddShape(FGameplayDebuggerShape::MakeCircle(GhostTargetBasePos, FVector::UpVector, Radius.Radius * 0.75f, FColorList::Orange));
+							if (bHasStandingSteering)
+							{
+								const FVector GhostTargetBasePos = StandingSteeringList[EntityIndex].TargetLocation + FVector(0.0f, 0.0f, 25.0f);
+								AddShape(FGameplayDebuggerShape::MakeCircle(GhostTargetBasePos, FVector::UpVector, Radius.Radius * 0.75f, FColorList::Orange));
+							}
 						}
 					}
 					
@@ -749,49 +766,68 @@ void FGameplayDebuggerCategory_Mass::CollectData(APlayerController* OwnerPC, AAc
 						// Entity name
 						Status += TEXT("{orange}");
 						Status += Entity.DebugGetDescription();
-						Status += TEXT(" {white}LOD ");
-						switch (SimLOD.LOD) {
-						case EMassLOD::High:
-							Status += TEXT("High");
-							break;
-						case EMassLOD::Medium:
-							Status += TEXT("Med");
-							break;
-						case EMassLOD::Low:
-							Status += TEXT("Low");
-							break;
-						case EMassLOD::Off:
-							Status += TEXT("Off");
-							break;
-						default:
-							Status += TEXT("?");
-							break;
-						}
-						Status += TEXT("\n");
-						
-						// Current StateTree task
-						if (StateTree != nullptr)
+
+						// LOD
+						if (bHasLOD)
 						{
-							if (FStateTreeInstanceData* InstanceData = MassStateTreeSubsystem->GetInstanceData(StateTreeInstance.InstanceHandle))
+							Status += TEXT(" {white}LOD ");
+							switch (SimLODList[EntityIndex].LOD)
 							{
-								FStateTreeReadOnlyExecutionContext StateTreeContext(MassStateTreeSubsystem, StateTree, *InstanceData);
-
-								Status += StateTreeContext.GetActiveStateName();
-								Status += FString::Printf(TEXT("  {yellow}%d{white}\n"), StateTreeContext.GetStateChangeCount());
+							case EMassLOD::High:
+								Status += TEXT("High");
+								break;
+							case EMassLOD::Medium:
+								Status += TEXT("Med");
+								break;
+							case EMassLOD::Low:
+								Status += TEXT("Low");
+								break;
+							case EMassLOD::Off:
+								Status += TEXT("Off");
+								break;
+							default:
+								Status += TEXT("?");
+								break;
 							}
-							else
+						}
+						Status += TEXT("\n");
+
+						// StateTree
+						if (bHasStateTree)
+						{
+							// Current StateTree task
+							// Optional shared fragment is expected to be present for entities with a StateTree instance
+							const FMassStateTreeSharedFragment& SharedStateTree = Context.GetConstSharedFragment<FMassStateTreeSharedFragment>();
+							if (const UStateTree* StateTree = SharedStateTree.StateTree)
 							{
-								Status += TEXT("{red}<No StateTree instance>{white}\n");
+								if (FStateTreeInstanceData* InstanceData = MassStateTreeSubsystem->GetInstanceData(StateTreeInstanceList[EntityIndex].InstanceHandle))
+								{
+									FStateTreeReadOnlyExecutionContext StateTreeContext(MassStateTreeSubsystem, StateTree, *InstanceData);
+									Status += StateTreeContext.GetActiveStateName();
+									Status += FString::Printf(TEXT("  {yellow}%d{white}\n"), StateTreeContext.GetStateChangeCount());
+								}
+								else
+								{
+									Status += TEXT("{red}<No StateTree instance>{white}\n");
+								}
 							}
 						}
 
 						// Movement info
-						Status += FString::Printf(TEXT("{yellow}%s/%03d {lightgrey}Speed:{white}%.1f {lightgrey}Force:{white}%.1f\n"),
-							*UEnum::GetDisplayValueAsText(MoveTarget.GetCurrentAction()).ToString(), MoveTarget.GetCurrentActionID(), Velocity.Value.Length(), Force.Value.Length());
-						Status += FString::Printf(TEXT("{pink}-> %s {white}Dist: %.1f\n"),
-							*UEnum::GetDisplayValueAsText(MoveTarget.IntentAtGoal).ToString(), MoveTarget.DistanceToGoal);
+						if (bHasMoveTarget)
+						{
+							const FMassMoveTargetFragment& MoveTarget = MoveTargetList[EntityIndex];
+							if (bHasVelocity && bHasForce)
+							{
+								Status += FString::Printf(TEXT("{yellow}%s/%03d {lightgrey}Speed:{white}%.1f {lightgrey}Force:{white}%.1f\n"),
+									*UEnum::GetDisplayValueAsText(MoveTarget.GetCurrentAction()).ToString(), MoveTarget.GetCurrentActionID(), VelocityList[EntityIndex].Value.Length(), ForceList[EntityIndex].Value.Length());
+							}
+
+							Status += FString::Printf(TEXT("{pink}-> %s {white}Dist: %.1f\n"),
+								*UEnum::GetDisplayValueAsText(MoveTarget.IntentAtGoal).ToString(), MoveTarget.DistanceToGoal);
+						}
 
-						// Look
+						// LookAt
 						if (bShowEntityLookAt && bHasLookAt)
 						{
 							const FMassLookAtFragment& LookAt = LookAtList[EntityIndex];
