[LookAt] added BP exposed functions to Create/Delete requests

diff --git a/Engine/Plugins/AI/MassAI/Source/MassAIBehavior/Private/MassLookAtSubsystem.cpp b/Engine/Plugins/AI/MassAI/Source/MassAIBehavior/Private/MassLookAtSubsystem.cpp
--- a/Engine/Plugins/AI/MassAI/Source/MassAIBehavior/Private/MassLookAtSubsystem.cpp
+++ b/Engine/Plugins/AI/MassAI/Source/MassAIBehavior/Private/MassLookAtSubsystem.cpp
@@ -10,6 +10,87 @@
 
 namespace UE::Mass::LookAt::Private
 {
+void FlushCommands(FMassEntityManager& EntityManager, const TSharedPtr<FMassCommandBuffer>& CommandBuffer)
+{
+	ExecuteOnGameThread(UE_SOURCE_LOCATION, [WeakEntityManager = EntityManager.AsWeak(), CommandBuffer]()
+		{
+			if (const TSharedPtr<FMassEntityManager> SharedEntityManager = WeakEntityManager.Pin())
+			{
+				SharedEntityManager->FlushCommands(CommandBuffer);
+			}
+			else
+			{
+				CommandBuffer->CancelCommands();
+			}
+		});
+}
+
+bool CreateRequests(FMassEntityManager& InEntityManager
+	, const FMassArchetypeHandle& RequestArchetype
+	, const TConstArrayView<FMassEntityHandle> Viewers
+	, const TConstArrayView<FMassEntityHandle> Requests
+	, const FMassEntityHandle TargetEntity
+	, const FMassLookAtPriority Priority)
+{
+	if (Viewers.IsEmpty())
+	{
+		return false;
+	}
+
+	checkf(Viewers.Num() == Requests.Num(), TEXT("Number of reserved entities for requests must match the number of provided viewer entities."));
+
+	const TSharedRef<FMassEntityManager::FEntityCreationContext> CreationContext = InEntityManager.BatchCreateReservedEntities(RequestArchetype, /*SharedFragmentValues*/{}, Requests);
+
+	for (int Index = 0; Index < Requests.Num(); ++Index)
+	{
+		InEntityManager.SetEntityFragmentValues(Requests[Index]
+			, { FInstancedStruct::Make(FMassLookAtRequestFragment{Viewers[Index], Priority, EMassLookAtMode::LookAtEntity, TargetEntity}) });
+
+		UE_VLOG_UELOG(InEntityManager.GetOwner(), LogMassBehavior, Log, TEXT("Created LookAt Request '%s', Target '%s', Priority=%s")
+			, *LexToString(Requests[Index])
+			, *LexToString(TargetEntity)
+			, *LexToString(Priority.Get()))
+	}
+
+	return true;
+}
+
+bool CreateTargets(FMassEntityManager& InEntityManager
+	, const FMassArchetypeHandle& TargetArchetype
+	, const TConstArrayView<FMassEntityHandle> Targets
+	, const TConstArrayView<const FTransform> Transforms
+	, const FMassLookAtPriority Priority)
+{
+	if (Targets.IsEmpty())
+	{
+		return false;
+	}
+
+	checkf(Targets.Num() == Transforms.Num(), TEXT("Number of reserved entities for targets must match the number of provided transforms."));
+
+	// This needs to stay in sync with 'TargetArchetype' created on subsystem initialization
+	TArray<FInstancedStruct, TInlineAllocator<2>> FragmentInstanceList;
+	FragmentInstanceList.Add(FInstancedStruct::Make(FMassLookAtTargetFragment{ .Offset = FVector::ZeroVector, .Priority = Priority }));
+	FTransformFragment& TransformFragment = FragmentInstanceList.Add_GetRef({ FInstancedStruct::Make(FTransformFragment{})}).GetMutable<FTransformFragment>();
+
+	const TSharedRef<FMassEntityManager::FEntityCreationContext> CreationContext =
+		InEntityManager.BatchCreateReservedEntities(TargetArchetype, /*SharedFragmentValues*/{}, Targets);
+
+	for (int Index = 0; Index < Targets.Num(); ++Index)
+	{
+		TransformFragment.SetTransform(Transforms[Index]);
+		InEntityManager.SetEntityFragmentValues(Targets[Index], FragmentInstanceList);
+
+		UE_VLOG_UELOG(InEntityManager.GetOwner(), LogMassBehavior, Log, TEXT("Created LookAtTarget '%s' at '%s'")
+			, *LexToString(Targets[Index])
+			, *LexToString(TransformFragment.GetTransform().ToString()))
+		UE_VLOG_LOCATION(InEntityManager.GetOwner(), LogMassBehavior, Display
+			, TransformFragment.GetTransform().GetLocation(), /*Thickness*/50, FColor::Yellow, TEXT(""));
+	}
+
+	return true;
+}
+
 TArray<FMassEntityHandle> DebugRequests;
 
 static FAutoConsoleCommandWithWorldArgsAndOutputDevice CmdSendLookAtPlayerRequestToAll(
@@ -89,10 +170,10 @@ static FAutoConsoleCommandWithWorldArgsAndOutputDevice CmdSendLookAtPlayerReques
 				return;
 			}
 
-			const FMassLookAtTargetFragment* LookAtTargetFragment = EntitySubsystem->GetMutableEntityManager().GetFragmentDataPtr<FMassLookAtTargetFragment>(PlayerEntity);
-			if (LookAtTargetFragment == nullptr)
+			const UMassLookAtSubsystem* LookAtSubsystem = World->GetSubsystem<UMassLookAtSubsystem>();
+			if (LookAtSubsystem == nullptr)
 			{
-				OutputDevice.Log(ELogVerbosity::Error, TEXT("Command failed: missing required LookAtTarget fragment for the MassEntity associated to the player"));
+				OutputDevice.Log(ELogVerbosity::Error, TEXT("Command failed: unable to find UMassLookAtSubsystem"));
 				return;
 			}
 
@@ -100,38 +181,18 @@ static FAutoConsoleCommandWithWorldArgsAndOutputDevice CmdSendLookAtPlayerReques
 			if (bCreateRequest)
 			{
 				// Create requests for all entities with a MassLookAtFragment
-				CommandBuffer->PushCommand<FMassDeferredCreateCommand>([PlayerEntity, PriorityLevel](FMassEntityManager& InEntityManager)
+				CommandBuffer->PushCommand<FMassDeferredCreateCommand>([PlayerEntity, PriorityLevel, RequestArchetype = LookAtSubsystem->DebugGetRequestArchetype()](FMassEntityManager& InEntityManager)
 					{
 						FMassEntityQuery EntityQuery(InEntityManager.AsShared());
 						EntityQuery.AddRequirement<FMassLookAtFragment>(EMassFragmentAccess::ReadOnly);
-
-						// @todo LookAt: replace this part by FMassEntityQuery::GetMatchingEntityHandles() once available
-						TArray<FMassEntityHandle> Viewers;
-						FMassExecutionContext Context(InEntityManager, 0.0f);
-						EntityQuery.ForEachEntityChunk(Context, [&Viewers](const FMassExecutionContext& Context)
-							{
-								Viewers.Append(Context.GetEntities());
-							});
-
-						if (Viewers.Num() > 0)
-						{
-							FMassArchetypeCompositionDescriptor Composition;
-							Composition.Fragments.Add(*FMassLookAtRequestFragment::StaticStruct());
-
-							TArray<FMassEntityHandle> Requests;
-							const TSharedRef<FMassEntityManager::FEntityCreationContext> CreationContext =
-								InEntityManager.BatchCreateEntities(InEntityManager.CreateArchetype(Composition), Viewers.Num(), Requests);
-
-							DebugRequests.Append(Requests);
-							check(Viewers.Num() == Requests.Num());
-							const FMassLookAtPriority Priority(PriorityLevel);
-
-							for (int Index = 0; Index < Requests.Num(); ++Index)
-							{
-								InEntityManager.SetEntityFragmentValues(Requests[Index]
-									, { FInstancedStruct::Make(FMassLookAtRequestFragment{Viewers[Index], Priority, EMassLookAtMode::LookAtEntity, PlayerEntity}) });
-							}
-						}
+						TArray<FMassEntityHandle> Viewers = EntityQuery.GetMatchingEntityHandles();
+
+						CreateRequests(InEntityManager
+							, RequestArchetype
+							, Viewers
+							, InEntityManager.BatchReserveEntities(Viewers.Num(), DebugRequests)
+							, PlayerEntity
+							, FMassLookAtPriority{ PriorityLevel });
 					});
 			}
 			else
@@ -144,22 +205,100 @@ static FAutoConsoleCommandWithWorldArgsAndOutputDevice CmdSendLookAtPlayerReques
 					});
 			}
 
-			ExecuteOnGameThread(UE_SOURCE_LOCATION, [WeakEntityManager = EntitySubsystem->GetMutableEntityManager().AsWeak(), CommandBuffer]()
-				{
-					if (const TSharedPtr<FMassEntityManager> SharedEntityManager = WeakEntityManager.Pin())
-					{
-						SharedEntityManager->FlushCommands(CommandBuffer);
-					}
-					else
-					{
-						CommandBuffer->CancelCommands();
-					}
-				});
+			FlushCommands(EntitySubsystem->GetMutableEntityManager(), CommandBuffer);
 		}
 	));
 
 } // UE::Mass::LookAt::Private
 
+FMassLookAtRequestHandle UMassLookAtSubsystem::CreateRequest(AActor* ViewerActor, FMassLookAtPriority Priority, const FVector TargetLocation) const
+{
+	const UWorld* World = GetWorld();
+
+	UMassActorSubsystem* ActorSubsystem = World->GetSubsystem<UMassActorSubsystem>();
+	if (ActorSubsystem == nullptr)
+	{
+		UE_VLOG_UELOG(this, LogMassBehavior, Error, TEXT("%hs failed: unable to find MassActorSubsystem"), __FUNCTION__);
+		return {};
+	}
+
+	UMassEntitySubsystem* EntitySubsystem = World->GetSubsystem<UMassEntitySubsystem>();
+	if (EntitySubsystem == nullptr)
+	{
+		UE_VLOG_UELOG(this, LogMassBehavior, Error, TEXT("%hs failed: unable to find UMassEntitySubsystem"), __FUNCTION__);
+		return {};
+	}
+
+	const FMassEntityHandle ViewerEntity = ActorSubsystem->GetEntityHandleFromActor(ViewerActor);
+	if (!ViewerEntity.IsSet())
+	{
+		UE_VLOG_UELOG(this, LogMassBehavior, Error, TEXT("%hs failed: unable to find a MassEntity associated to the player"), __FUNCTION__);
+		return {};
+	}
+
+	// @todo: once available, consider using incoming EntityBuilder improvements for this whole creation step
+	FMassEntityManager& EntityManager = EntitySubsystem->GetMutableEntityManager();
+	FMassEntityHandle TargetEntity = EntityManager.ReserveEntity();
+	FMassEntityHandle RequestEntity = EntityManager.ReserveEntity();
+	FTransform TargetTransform(TargetLocation);
+
+	// Push command to create a new entity representing a LookAt target along with a request to look at it
+	const TSharedPtr<FMassCommandBuffer> CommandBuffer = MakeShareable(new FMassCommandBuffer());
+	CommandBuffer->PushCommand<FMassDeferredCreateCommand>(
+		[
+			TargetArchetype = TargetArchetype
+			, RequestArchetype = RequestArchetype
+			, RequestEntity
+			, ViewerEntity
+			, TargetEntity
+			, TargetTransform
+			, Priority
+		]
+		(FMassEntityManager& InEntityManager)
+		{
+			UE::Mass::LookAt::Private::CreateTargets(InEntityManager, TargetArchetype, { TargetEntity }, { TargetTransform }, Priority);
+			UE::Mass::LookAt::Private::CreateRequests(InEntityManager, RequestArchetype, { ViewerEntity }, { RequestEntity }, TargetEntity, Priority);
+		});
+
+	UE::Mass::LookAt::Private::FlushCommands(EntityManager, CommandBuffer);
+
+	return FMassLookAtRequestHandle{RequestEntity, TargetEntity};
+}
+
+void UMassLookAtSubsystem::DeleteRequest(const FMassLookAtRequestHandle RequestHandle) const
+{
+	// Simple validation when none of the entities are set since it is probably due to a bad data setup.
+	// Otherwise, the EntityManager can process gracefully the handles, valid or not.
+	if (!RequestHandle.Request.IsSet()
+		&& !RequestHandle.Target.IsSet())
+	{
+		UE_VLOG_UELOG(this, LogMassBehavior, Error, TEXT("%hs failed: invalid request handle"), __FUNCTION__);
+		return;
+	}
+
+	const UWorld* World = GetWorld();
+	if (World == nullptr)
+	{
+		UE_VLOG_UELOG(this, LogMassBehavior, Error, TEXT("%hs failed: unable to find World"), __FUNCTION__);
+		return;
+	}
+
+	UMassEntitySubsystem* EntitySubsystem = World->GetSubsystem<UMassEntitySubsystem>();
+	if (EntitySubsystem == nullptr)
+	{
+		UE_VLOG_UELOG(this, LogMassBehavior, Error, TEXT("%hs failed: unable to find UMassEntitySubsystem"), __FUNCTION__);
+		return;
+	}
+
+	const TSharedPtr<FMassCommandBuffer> CommandBuffer = MakeShareable(new FMassCommandBuffer());
+	CommandBuffer->PushCommand<FMassDeferredDestroyCommand>([RequestHandle](FMassEntityManager& InEntityManager)
+		{
+			InEntityManager.BatchDestroyEntities({ RequestHandle.Request, RequestHandle.Target });
+		});
+
+	UE::Mass::LookAt::Private::FlushCommands(EntitySubsystem->GetMutableEntityManager(), CommandBuffer);
+}
+
 //----------------------------------------------------------------------//
 //  UMassLookAtSubsystem
 //----------------------------------------------------------------------//
@@ -167,6 +306,17 @@ void UMassLookAtSubsystem::Initialize(FSubsystemCollectionBase& Collection)
 {
 	Super::Initialize(Collection);
 
+	FMassEntityManager& EntityManager = Collection.InitializeDependency<UMassEntitySubsystem>()->GetMutableEntityManager();
+
+	// Create Mass archetype for entities representing requests
+	const FMassArchetypeCompositionDescriptor RequestComposition{ FMassFragmentBitSet(*FMassLookAtRequestFragment::StaticStruct()) };
+	RequestArchetype = EntityManager.CreateArchetype(RequestComposition);
+
+	// Create Mass archetype for entities representing targets
+	TArray<const UScriptStruct*, TInlineAllocator<2>> TargetFragmentTypes{FMassLookAtTargetFragment::StaticStruct(), FTransformFragment::StaticStruct()};
+	const FMassArchetypeCompositionDescriptor TargetComposition{ FMassFragmentBitSet(TargetFragmentTypes) };
+	TargetArchetype = EntityManager.CreateArchetype(TargetComposition);
+
 	OverrideSubsystemTraits<UMassLookAtSubsystem>(Collection);
 }
 
diff --git a/Engine/Plugins/AI/MassAI/Source/MassAIBehavior/Public/MassLookAtSubsystem.h b/Engine/Plugins/AI/MassAI/Source/MassAIBehavior/Public/MassLookAtSubsystem.h
--- a/Engine/Plugins/AI/MassAI/Source/MassAIBehavior/Public/MassLookAtSubsystem.h
+++ b/Engine/Plugins/AI/MassAI/Source/MassAIBehavior/Public/MassLookAtSubsystem.h
@@ -8,6 +8,18 @@
 #include "MassLookAtTypes.h"
 #include "MassLookAtSubsystem.generated.h"
 
+USTRUCT(BlueprintType)
+struct FMassLookAtRequestHandle
+{
+	GENERATED_BODY()
+
+	UPROPERTY()
+	FMassEntityHandle Request;
+
+	UPROPERTY()
+	FMassEntityHandle Target;
+};
+
 /**
  * Subsystem that keeps track of the LookAt targets
  */
@@ -28,6 +40,20 @@ public:
 		bool bActive = false;
 	};
 
+	/**
+	 * Creates a new LookAt request using the provided parameters for the mass entity associated to 'ViewerActor', if any.
+	 * Method will log an error if it fails (e.g., no associated mass entity).
+	 */
+	UFUNCTION(BlueprintCallable, Category = "LookAt", meta = (ReturnDisplayName="Request Handle"))
+	FMassLookAtRequestHandle CreateRequest(AActor* ViewerActor, FMassLookAtPriority Priority, FVector TargetLocation) const;
+
+	/**
+	 * Removes given request from the active LookAt requests.
+	 * Method will log an error if it fails (e.g., invalid request handle).
+	 */
+	UFUNCTION(BlueprintCallable, Category = "LookAt")
+	void DeleteRequest(FMassLookAtRequestHandle RequestHandle) const;
+
 	/**
 	 * Appends the provided requests to the active ones, then updates the LookAt fragments
 	 * of the referenced entity if required.
@@ -119,6 +145,18 @@ public:
 		return TargetGrid.GetItems().Num();
 	}
 
+	/** @return Mass archetype to create entities representing a LookAt request */
+	const FMassArchetypeHandle& DebugGetRequestArchetype() const
+	{
+		return RequestArchetype;
+	}
+
+	/** @return Mass archetype to create entities representing a LookAt target */
+	const FMassArchetypeHandle& DebugGetTargetArchetype() const
+	{
+		return TargetArchetype;
+	}
+
 #if WITH_MASSGAMEPLAY_DEBUG
 	/** @return String detailing all requests registered for a given entity. */
 	FString DebugGetRequestsString(FMassEntityHandle InEntity) const;
@@ -161,6 +199,12 @@ private:
 
 	/** Multithread access detector to detect threading issues with the hash grid */
 	UE_MT_DECLARE_RW_ACCESS_DETECTOR(TargetGridAccessDetector);
+
+	/** Cached archetype to create entities representing a LookAt target */
+	FMassArchetypeHandle TargetArchetype;
+
+	/** Cached archetype to create entities representing a LookAt request */
+	FMassArchetypeHandle RequestArchetype;
 };
 
 template<>
diff --git a/Engine/Plugins/AI/MassAI/Source/MassAIBehavior/Public/MassLookAtTypes.h b/Engine/Plugins/AI/MassAI/Source/MassAIBehavior/Public/MassLookAtTypes.h
--- a/Engine/Plugins/AI/MassAI/Source/MassAIBehavior/Public/MassLookAtTypes.h
+++ b/Engine/Plugins/AI/MassAI/Source/MassAIBehavior/Public/MassLookAtTypes.h
@@ -65,7 +65,7 @@ enum class EMassLookAtPriorities : uint8
 /**
  * Struct used as a priority selector exposed to the Editor
  */
-USTRUCT()
+USTRUCT(BlueprintType)
 struct FMassLookAtPriority
 {
 	GENERATED_BODY()
