Fix for smooth orientation processor not working for root motion driven characters (unless strafe warping was applied). The smooth orientation processor now takes the desired velocity as its input instead of the current

diff --git a/Engine/Plugins/AI/MassAI/Source/MassNavigation/Private/MassSmoothOrientationProcessors.cpp b/Engine/Plugins/AI/MassAI/Source/MassNavigation/Private/MassSmoothOrientationProcessors.cpp
--- a/Engine/Plugins/AI/MassAI/Source/MassNavigation/Private/MassSmoothOrientationProcessors.cpp
+++ b/Engine/Plugins/AI/MassAI/Source/MassNavigation/Private/MassSmoothOrientationProcessors.cpp
@@ -26,7 +26,6 @@ UMassSmoothOrientationProcessor::UMassSmoothOrientationProcessor()
 void UMassSmoothOrientationProcessor::ConfigureQueries(const TSharedRef<FMassEntityManager>& EntityManager)
 {
 	HighResEntityQuery.AddRequirement<FMassMoveTargetFragment>(EMassFragmentAccess::ReadOnly);
-	HighResEntityQuery.AddRequirement<FMassVelocityFragment>(EMassFragmentAccess::ReadOnly);
 	HighResEntityQuery.AddRequirement<FMassDesiredMovementFragment>(EMassFragmentAccess::ReadWrite);
 	HighResEntityQuery.AddRequirement<FTransformFragment>(EMassFragmentAccess::ReadWrite);
 	HighResEntityQuery.AddTagRequirement<FMassOffLODTag>(EMassFragmentPresence::None);
@@ -54,7 +53,6 @@ void UMassSmoothOrientationProcessor::Execute(FMassEntityManager& EntityManager,
 
 			const TConstArrayView<FMassMoveTargetFragment> MoveTargetList = Context.GetFragmentView<FMassMoveTargetFragment>();
 			const TArrayView<FTransformFragment> LocationList = Context.GetMutableFragmentView<FTransformFragment>();
-			const TConstArrayView<FMassVelocityFragment> VelocityList = Context.GetFragmentView<FMassVelocityFragment>();
 			const TArrayView<FMassDesiredMovementFragment> DesiredMovementList = Context.GetMutableFragmentView<FMassDesiredMovementFragment>();
 
 			for (FMassExecutionContext::FEntityIterator EntityIt = Context.CreateEntityIterator(); EntityIt; ++EntityIt)
@@ -67,7 +65,6 @@ void UMassSmoothOrientationProcessor::Execute(FMassEntityManager& EntityManager,
 					continue;
 				}
 
-				const FMassVelocityFragment& CurrentVelocity = VelocityList[EntityIt];
 				FMassDesiredMovementFragment& DesiredMovement = DesiredMovementList[EntityIt];
 				FTransform& CurrentTransform = LocationList[EntityIt].GetMutableTransform();
 				const FVector CurrentForward = CurrentTransform.GetRotation().GetForwardVector();
@@ -100,7 +97,7 @@ void UMassSmoothOrientationProcessor::Execute(FMassEntityManager& EntityManager,
 					VelocityWeight = OrientationParams.Standing.VelocityWeight;
 				}
 				
-				const FVector::FReal VelocityHeading = UE::MassNavigation::GetYawFromDirection(CurrentVelocity.Value);
+				const FVector::FReal VelocityHeading = UE::MassNavigation::GetYawFromDirection(DesiredMovement.DesiredVelocity);
 				const FVector::FReal MovementHeading = UE::MassNavigation::GetYawFromDirection(MoveTarget.Forward);
 
 				const FVector::FReal Ratio = MoveTargetWeight / (MoveTargetWeight + VelocityWeight);
