[Looked] Fixed CreateRequest BP function to not be pure

diff --git a/Engine/Plugins/AI/MassAI/Source/MassAIBehavior/Private/Tasks/MassLookAtTask.cpp b/Engine/Plugins/AI/MassAI/Source/MassAIBehavior/Private/Tasks/MassLookAtTask.cpp
--- a/Engine/Plugins/AI/MassAI/Source/MassAIBehavior/Private/Tasks/MassLookAtTask.cpp
+++ b/Engine/Plugins/AI/MassAI/Source/MassAIBehavior/Private/Tasks/MassLookAtTask.cpp
@@ -37,8 +37,7 @@ EStateTreeRunStatus FMassLookAtTask::EnterState(FStateTreeExecutionContext& Cont
 			[&Context, SelfEntity = MassStateTreeContext.GetEntity(), &InstanceData, Mode = LookAtMode, Priority = Priority, InterpolationSpeed = InterpolationSpeed]
 			(FMassEntityManager& InEntityManager)
 			{
-				FMassArchetypeCompositionDescriptor Composition;
-				Composition.Fragments.Add(*FMassLookAtRequestFragment::StaticStruct());
+				const FMassArchetypeCompositionDescriptor Composition{FMassFragmentBitSet(*FMassLookAtRequestFragment::StaticStruct())};
 
 				constexpr int32 NumRequestsToCreate = 1;
 				TArray<FMassEntityHandle> Requests;
diff --git a/Engine/Plugins/AI/MassAI/Source/MassAIBehavior/Public/MassLookAtSubsystem.h b/Engine/Plugins/AI/MassAI/Source/MassAIBehavior/Public/MassLookAtSubsystem.h
--- a/Engine/Plugins/AI/MassAI/Source/MassAIBehavior/Public/MassLookAtSubsystem.h
+++ b/Engine/Plugins/AI/MassAI/Source/MassAIBehavior/Public/MassLookAtSubsystem.h
@@ -44,14 +44,14 @@ public:
 	 * Creates a new LookAt request using the provided parameters for the mass entity associated to 'ViewerActor', if any.
 	 * Method will log an error if it fails (e.g., no associated mass entity).
 	 */
-	UFUNCTION(BlueprintCallable, Category = "LookAt", meta = (ReturnDisplayName="Request Handle"))
+	UFUNCTION(BlueprintCallable, BlueprintPure = false, Category = "LookAt", meta = (ReturnDisplayName="Request Handle"))
 	FMassLookAtRequestHandle CreateRequest(AActor* ViewerActor, FMassLookAtPriority Priority, FVector TargetLocation) const;
 
 	/**
 	 * Removes given request from the active LookAt requests.
 	 * Method will log an error if it fails (e.g., invalid request handle).
 	 */
-	UFUNCTION(BlueprintCallable, Category = "LookAt")
+	UFUNCTION(BlueprintCallable, BlueprintPure = false, Category = "LookAt")
 	void DeleteRequest(FMassLookAtRequestHandle RequestHandle) const;
 
 	/**
