Fix sliding NPCs occuring when entering fail state.
On failure fallback on animate action to stop the NPC from moving.
Added NavmeshAnimate task (similar to NavmeshStand task) to trigger Animate action.

diff --git a/Engine/Plugins/AI/MassAI/Source/MassAIBehavior/Private/Tasks/MassNavMeshAnimateTask.cpp b/Engine/Plugins/AI/MassAI/Source/MassAIBehavior/Private/Tasks/MassNavMeshAnimateTask.cpp
new file mode 100644
--- /dev/null
+++ b/Engine/Plugins/AI/MassAI/Source/MassAIBehavior/Private/Tasks/MassNavMeshAnimateTask.cpp
@@ -0,0 +1,66 @@
+﻿// Copyright Epic Games, Inc. All Rights Reserved.
+
+#include "Tasks/MassNavMeshAnimateTask.h"
+#include "StateTreeExecutionContext.h"
+#include "MassNavigationFragments.h"
+#include "MassNavMeshNavigationFragments.h"
+#include "MassNavMeshNavigationUtils.h"
+#include "MassStateTreeExecutionContext.h"
+#include "MassSignalSubsystem.h"
+#include "StateTreeLinker.h"
+#include "MassCommonFragments.h"
+
+bool FMassNavMeshAnimateTask::Link(FStateTreeLinker& Linker)
+{
+	Linker.LinkExternalData(MoveTargetHandle);
+	Linker.LinkExternalData(MassSignalSubsystemHandle);
+	Linker.LinkExternalData(TransformHandle);
+
+	return true;
+}
+
+EStateTreeRunStatus FMassNavMeshAnimateTask::EnterState(FStateTreeExecutionContext& Context, const FStateTreeTransitionResult& Transition) const
+{
+	const FMassStateTreeExecutionContext& MassStateTreeContext = static_cast<FMassStateTreeExecutionContext&>(Context);
+
+	FInstanceDataType& InstanceData = Context.GetInstanceData(*this);
+
+	FMassMoveTargetFragment& MoveTarget = Context.GetExternalData(MoveTargetHandle);
+
+	const UWorld* World = Context.GetWorld();
+	checkf(World != nullptr, TEXT("A valid world is expected from the execution context."));
+
+	const FVector AgentNavLocation = Context.GetExternalData(TransformHandle).GetTransform().GetLocation();
+	MoveTarget.Center = AgentNavLocation;
+	MoveTarget.CreateNewAction(EMassMovementAction::Animate, *World);
+	
+	const bool bSuccess = UE::MassNavigation::ActivateActionAnimate(Context.GetOwner(), MassStateTreeContext.GetEntity(), MoveTarget);
+
+	if (!bSuccess)
+	{
+		return EStateTreeRunStatus::Failed;
+	}
+
+	InstanceData.Time = 0.f;
+
+	// A Duration <= 0 indicates that the task runs until a transition in the state tree stops it.
+	// Otherwise, we schedule a signal to end the task.
+	if (InstanceData.Duration > 0.f)
+	{
+		UMassSignalSubsystem& MassSignalSubsystem = Context.GetExternalData(MassSignalSubsystemHandle);
+		MassSignalSubsystem.DelaySignalEntityDeferred(MassStateTreeContext.GetMassEntityExecutionContext(),
+			UE::Mass::Signals::AnimateTaskFinished, MassStateTreeContext.GetEntity(), InstanceData.Duration);
+	}
+
+	return EStateTreeRunStatus::Running;
+}
+
+EStateTreeRunStatus FMassNavMeshAnimateTask::Tick(FStateTreeExecutionContext& Context, const float DeltaTime) const
+{
+	FInstanceDataType& InstanceData = Context.GetInstanceData(*this);
+	
+	InstanceData.Time += DeltaTime;
+	
+	return InstanceData.Duration <= 0.f ? EStateTreeRunStatus::Running :
+		(InstanceData.Time < InstanceData.Duration ? EStateTreeRunStatus::Running : EStateTreeRunStatus::Succeeded);
+}
diff --git a/Engine/Plugins/AI/MassAI/Source/MassAIBehavior/Public/MassStateTreeTypes.h b/Engine/Plugins/AI/MassAI/Source/MassAIBehavior/Public/MassStateTreeTypes.h
--- a/Engine/Plugins/AI/MassAI/Source/MassAIBehavior/Public/MassStateTreeTypes.h
+++ b/Engine/Plugins/AI/MassAI/Source/MassAIBehavior/Public/MassStateTreeTypes.h
@@ -23,6 +23,7 @@ namespace UE::Mass::Signals
 	const FName LookAtFinished = FName(TEXT("LookAtFinished"));
 	const FName NewStateTreeTaskRequired = FName(TEXT("NewStateTreeTaskRequired"));
 	const FName StandTaskFinished = FName(TEXT("StandTaskFinished"));
+	const FName AnimateTaskFinished = FName(TEXT("AnimateTaskFinished"));
 	const FName DelayedTransitionWakeup = FName(TEXT("DelayedTransitionWakeup"));
 	// @todo MassStateTree: move this to its game plugin when possible
 	const FName ContextualAnimTaskFinished = FName(TEXT("ContextualAnimTaskFinished"));
diff --git a/Engine/Plugins/AI/MassAI/Source/MassAIBehavior/Public/Tasks/MassNavMeshAnimateTask.h b/Engine/Plugins/AI/MassAI/Source/MassAIBehavior/Public/Tasks/MassNavMeshAnimateTask.h
new file mode 100644
--- /dev/null
+++ b/Engine/Plugins/AI/MassAI/Source/MassAIBehavior/Public/Tasks/MassNavMeshAnimateTask.h
@@ -0,0 +1,50 @@
+﻿// Copyright Epic Games, Inc. All Rights Reserved.
+
+#pragma once
+
+#include "MassStateTreeTypes.h"
+#include "MassNavMeshAnimateTask.generated.h"
+
+struct FStateTreeExecutionContext;
+struct FMassMoveTargetFragment;
+struct FTransformFragment;
+class UMassSignalSubsystem;
+
+USTRUCT()
+struct MASSAIBEHAVIOR_API FMassNavMeshAnimateTaskInstanceData
+{
+	GENERATED_BODY()
+
+	/** Delay before the task ends. Default (0 or any negative) will run indefinitely, so it requires a transition in the state tree to stop it. */
+	UPROPERTY(EditAnywhere, Category = Parameter)
+	float Duration = 0.f;
+
+	UPROPERTY()
+	float Time = 0.f;
+};
+
+/**
+ * Stop and animate on current navmesh location
+ */
+USTRUCT(meta = (DisplayName = "Navmesh Animate"))
+struct MASSAIBEHAVIOR_API FMassNavMeshAnimateTask : public FMassStateTreeTaskBase
+{
+	GENERATED_BODY()
+
+	using FInstanceDataType = FMassNavMeshAnimateTaskInstanceData;
+
+protected:
+	virtual bool Link(FStateTreeLinker& Linker) override;
+
+	virtual const UStruct* GetInstanceDataType() const override
+	{
+		return FInstanceDataType::StaticStruct();
+	}
+
+	virtual EStateTreeRunStatus EnterState(FStateTreeExecutionContext& Context, const FStateTreeTransitionResult& Transition) const override;
+	virtual EStateTreeRunStatus Tick(FStateTreeExecutionContext& Context, const float DeltaTime) const override;
+
+	TStateTreeExternalDataHandle<FMassMoveTargetFragment> MoveTargetHandle;
+	TStateTreeExternalDataHandle<UMassSignalSubsystem> MassSignalSubsystemHandle;
+	TStateTreeExternalDataHandle<FTransformFragment> TransformHandle;
+};
diff --git a/Engine/Plugins/AI/MassAI/Source/MassNavMeshNavigation/Private/MassNavMeshNavigationUtils.cpp b/Engine/Plugins/AI/MassAI/Source/MassNavMeshNavigation/Private/MassNavMeshNavigationUtils.cpp
--- a/Engine/Plugins/AI/MassAI/Source/MassNavMeshNavigation/Private/MassNavMeshNavigationUtils.cpp
+++ b/Engine/Plugins/AI/MassAI/Source/MassNavMeshNavigation/Private/MassNavMeshNavigationUtils.cpp
@@ -28,8 +28,28 @@ namespace UE::MassNavigation
 		InOutMoveTarget.IntentAtGoal = EMassMovementAction::Stand;
 		InOutMoveTarget.DesiredSpeed.Set(DesiredSpeed);
 		
-		UE_VLOG(Requester, LogMassNavigation, Log, TEXT("Entity [%s] successfully requested %s."), *Entity.DebugGetDescription(), *InOutMoveTarget.ToString());
+		UE_VLOG_UELOG(Requester, LogMassNavigation, Log, TEXT("Entity [%s] successfully requested %s."), *Entity.DebugGetDescription(), *InOutMoveTarget.ToString());
 
 		return true;
-	}	
+	}
+
+	bool ActivateActionAnimate(const UObject* Requester,
+						       const FMassEntityHandle Entity,
+						       FMassMoveTargetFragment& MoveTarget)
+	{
+		MoveTarget.DistanceToGoal = 0.f;
+		MoveTarget.EntityDistanceToGoal = FMassMoveTargetFragment::UnsetDistance;
+		MoveTarget.DesiredSpeed.Set(0.f);
+
+		if (!ensureMsgf(MoveTarget.GetCurrentAction() == EMassMovementAction::Animate, TEXT("Expecting action 'Animate': Invalid action %u"), MoveTarget.GetCurrentAction()))
+		{
+			return false;
+		}
+
+		MoveTarget.IntentAtGoal = EMassMovementAction::Stand;
+
+		UE_VLOG_UELOG(Requester, LogMassNavigation, Log, TEXT("Entity [%s] successfully requested %s"), *Entity.DebugGetDescription(), *MoveTarget.ToString());
+
+		return true;
+	}
 }
diff --git a/Engine/Plugins/AI/MassAI/Source/MassNavMeshNavigation/Public/MassNavMeshNavigationUtils.h b/Engine/Plugins/AI/MassAI/Source/MassNavMeshNavigation/Public/MassNavMeshNavigationUtils.h
--- a/Engine/Plugins/AI/MassAI/Source/MassNavMeshNavigation/Public/MassNavMeshNavigationUtils.h
+++ b/Engine/Plugins/AI/MassAI/Source/MassNavMeshNavigation/Public/MassNavMeshNavigationUtils.h
@@ -13,4 +13,7 @@ namespace UE::MassNavigation
 {
 	MASSNAVMESHNAVIGATION_API bool ActivateActionStand(
 		const UObject* Requester, const FMassEntityHandle Entity, const float DesiredSpeed, FMassMoveTargetFragment& InOutMoveTarget, FMassNavMeshShortPathFragment& OutShortPath);
+
+	MASSNAVMESHNAVIGATION_API bool ActivateActionAnimate(
+		const UObject* Requester, const FMassEntityHandle Entity, FMassMoveTargetFragment& MoveTarget);
 };
