[Mass] Entity description verbosity

diff --git a/Engine/Plugins/AI/MassAI/Source/MassAIDebug/Private/GameplayDebuggerCategory_Mass.cpp b/Engine/Plugins/AI/MassAI/Source/MassAIDebug/Private/GameplayDebuggerCategory_Mass.cpp
--- a/Engine/Plugins/AI/MassAI/Source/MassAIDebug/Private/GameplayDebuggerCategory_Mass.cpp
+++ b/Engine/Plugins/AI/MassAI/Source/MassAIDebug/Private/GameplayDebuggerCategory_Mass.cpp
@@ -101,7 +101,7 @@ FGameplayDebuggerCategory_Mass::FOnConsoleCommandBroadcastDelegate FGameplayDebu
 FGameplayDebuggerCategory_Mass::FOnConsoleCommandBroadcastDelegate FGameplayDebuggerCategory_Mass::OnToggleNearEntityAvoidanceBroadcast;
 FGameplayDebuggerCategory_Mass::FOnConsoleCommandBroadcastDelegate FGameplayDebuggerCategory_Mass::OnToggleNearEntityPathBroadcast;
 FGameplayDebuggerCategory_Mass::FOnConsoleCommandBroadcastDelegate FGameplayDebuggerCategory_Mass::OnToggleEntityLookAtBroadcast;
-FGameplayDebuggerCategory_Mass::FOnConsoleCommandBroadcastDelegate FGameplayDebuggerCategory_Mass::OnToggleEntityDescriptionBroadcast;
+FGameplayDebuggerCategory_Mass::FOnConsoleCommandBroadcastDelegate FGameplayDebuggerCategory_Mass::OnCycleEntityDescriptionBroadcast;
 FGameplayDebuggerCategory_Mass::FOnConsoleCommandBroadcastDelegate FGameplayDebuggerCategory_Mass::OnToggleDebugLocalEntityManagerBroadcast;
 FGameplayDebuggerCategory_Mass::FOnConsoleCommandBroadcastDelegate FGameplayDebuggerCategory_Mass::OnTogglePickedActorAsViewerBroadcast;
 FGameplayDebuggerCategory_Mass::FOnConsoleCommandBroadcastDelegate FGameplayDebuggerCategory_Mass::OnToggleDrawViewersBroadcast;
@@ -122,7 +122,6 @@ FGameplayDebuggerCategory_Mass::FGameplayDebuggerCategory_Mass()
 	bShowNearEntityAvoidance = false;
 	bShowNearEntityPath = false;
 	bShowEntityLookAt = false;
-	bShowEntityDescription = true;
 	bMarkEntityBeingDebugged = true;
 	bDebugLocalEntityManager = false;
 	bShowViewers = false;
@@ -136,7 +135,7 @@ FGameplayDebuggerCategory_Mass::FGameplayDebuggerCategory_Mass()
 	BindKeyPress(EKeys::V.GetFName(), FGameplayDebuggerInputModifier::Shift, this, &FGameplayDebuggerCategory_Mass::OnToggleNearEntityAvoidance, EGameplayDebuggerInputMode::Replicated);
 	BindKeyPress(EKeys::N.GetFName(), FGameplayDebuggerInputModifier::Shift, this, &FGameplayDebuggerCategory_Mass::OnToggleNearEntityPath, EGameplayDebuggerInputMode::Replicated);
 	BindKeyPress(EKeys::K.GetFName(), FGameplayDebuggerInputModifier::Shift, this, &FGameplayDebuggerCategory_Mass::OnToggleEntityLookAt, EGameplayDebuggerInputMode::Replicated);
-	BindKeyPress(EKeys::E.GetFName(), FGameplayDebuggerInputModifier::Shift, this, &FGameplayDebuggerCategory_Mass::OnToggleEntityDescription, EGameplayDebuggerInputMode::Replicated);
+	BindKeyPress(EKeys::E.GetFName(), FGameplayDebuggerInputModifier::Shift, this, &FGameplayDebuggerCategory_Mass::OnCycleEntityDescription, EGameplayDebuggerInputMode::Replicated);
 	ToggleDebugLocalEntityManagerInputIndex = GetNumInputHandlers();
 	BindKeyPress(EKeys::L.GetFName(), FGameplayDebuggerInputModifier::Shift, this, &FGameplayDebuggerCategory_Mass::OnToggleDebugLocalEntityManager, EGameplayDebuggerInputMode::Local);
 	BindKeyPress(EKeys::Add.GetFName(), FGameplayDebuggerInputModifier::Shift, this, &FGameplayDebuggerCategory_Mass::OnIncreaseSearchRange, EGameplayDebuggerInputMode::Replicated);
@@ -159,7 +158,7 @@ FGameplayDebuggerCategory_Mass::FGameplayDebuggerCategory_Mass()
 		ConsoleCommands.Emplace(TEXT("gdt.mass.ToggleNearEntityAvoidance"), TEXT(""), FConsoleCommandWithWorldDelegate::CreateLambda([](UWorld* InWorld) { OnToggleNearEntityAvoidanceBroadcast.Broadcast(InWorld); }));
 		ConsoleCommands.Emplace(TEXT("gdt.mass.ToggleNearEntityPath"), TEXT(""), FConsoleCommandWithWorldDelegate::CreateLambda([](UWorld* InWorld) { OnToggleNearEntityPathBroadcast.Broadcast(InWorld); }));
 		ConsoleCommands.Emplace(TEXT("gdt.mass.ToggleEntityLookAt"), TEXT(""), FConsoleCommandWithWorldDelegate::CreateLambda([](UWorld* InWorld) { OnToggleEntityLookAtBroadcast.Broadcast(InWorld); }));
-		ConsoleCommands.Emplace(TEXT("gdt.mass.ToggleEntityDescription"), TEXT(""), FConsoleCommandWithWorldDelegate::CreateLambda([](UWorld* InWorld) { OnToggleEntityDescriptionBroadcast.Broadcast(InWorld); }));
+		ConsoleCommands.Emplace(TEXT("gdt.mass.CycleEntityDescriptionVerbosity"), TEXT(""), FConsoleCommandWithWorldDelegate::CreateLambda([](UWorld* InWorld) { OnCycleEntityDescriptionBroadcast.Broadcast(InWorld); }));
 		ConsoleCommands.Emplace(TEXT("gdt.mass.ToggleDebugLocalEntityManager"), TEXT(""), FConsoleCommandWithWorldDelegate::CreateLambda([](UWorld* InWorld) { OnToggleDebugLocalEntityManagerBroadcast.Broadcast(InWorld); }));
 		ConsoleCommands.Emplace(TEXT("gdt.mass.TogglePickedActorAsViewer"), TEXT(""), FConsoleCommandWithWorldDelegate::CreateLambda([](UWorld* InWorld) { OnTogglePickedActorAsViewerBroadcast.Broadcast(InWorld); }));
 		ConsoleCommands.Emplace(TEXT("gdt.mass.ToggleDrawViewers"), TEXT(""), FConsoleCommandWithWorldDelegate::CreateLambda([](UWorld* InWorld) { OnToggleDrawViewersBroadcast.Broadcast(InWorld); }));	
@@ -175,7 +174,7 @@ FGameplayDebuggerCategory_Mass::FGameplayDebuggerCategory_Mass()
 	ConsoleCommandHandles.Add(FDelegateHandlePair(&OnToggleNearEntityAvoidanceBroadcast, OnToggleNearEntityAvoidanceBroadcast.AddLambda([this](UWorld* InWorld) { if (InWorld == GetWorldFromReplicator()) { OnToggleNearEntityAvoidance(); }})));
 	ConsoleCommandHandles.Add(FDelegateHandlePair(&OnToggleNearEntityPathBroadcast, OnToggleNearEntityPathBroadcast.AddLambda([this](UWorld* InWorld) { if (InWorld == GetWorldFromReplicator()) { OnToggleNearEntityPath(); }})));
 	ConsoleCommandHandles.Add(FDelegateHandlePair(&OnToggleEntityLookAtBroadcast, OnToggleEntityLookAtBroadcast.AddLambda([this](UWorld* InWorld) { if (InWorld == GetWorldFromReplicator()) { OnToggleEntityLookAt(); }})));
-	ConsoleCommandHandles.Add(FDelegateHandlePair(&OnToggleEntityDescriptionBroadcast, OnToggleEntityDescriptionBroadcast.AddLambda([this](UWorld* InWorld) { if (InWorld == GetWorldFromReplicator()) { OnToggleEntityDescription(); }})));
+	ConsoleCommandHandles.Add(FDelegateHandlePair(&OnCycleEntityDescriptionBroadcast, OnCycleEntityDescriptionBroadcast.AddLambda([this](UWorld* InWorld) { if (InWorld == GetWorldFromReplicator()) { OnCycleEntityDescription(); }})));
 	ConsoleCommandHandles.Add(FDelegateHandlePair(&OnToggleDebugLocalEntityManagerBroadcast, OnToggleDebugLocalEntityManagerBroadcast.AddLambda([this](UWorld* InWorld) { if (InWorld == GetWorldFromReplicator()) { OnToggleDebugLocalEntityManager(); }})));
 	ConsoleCommandHandles.Add(FDelegateHandlePair(&OnTogglePickedActorAsViewerBroadcast, OnTogglePickedActorAsViewerBroadcast.AddLambda([this](UWorld* InWorld) { if (InWorld == GetWorldFromReplicator()) { OnTogglePickedActorAsViewer(); }})));
 	ConsoleCommandHandles.Add(FDelegateHandlePair(&OnToggleDrawViewersBroadcast, OnToggleDrawViewersBroadcast.AddLambda([this](UWorld* InWorld) { if (InWorld == GetWorldFromReplicator()) { OnToggleDrawViewers(); }})));
@@ -759,7 +758,8 @@ void FGameplayDebuggerCategory_Mass::CollectData(APlayerController* OwnerPC, AAc
 					}
 					
 					// Status
-					if (bShowEntityDescription && DistanceToEntitySq < FMath::Square(MaxViewDistance * 0.5f))
+					if (EntityDescriptionVerbosity >= EEntityDescriptionVerbosity::Minimal
+						&& DistanceToEntitySq < FMath::Square(MaxViewDistance * 0.5f))
 					{
 						FString Status;
 
@@ -767,77 +767,80 @@ void FGameplayDebuggerCategory_Mass::CollectData(APlayerController* OwnerPC, AAc
 						Status += TEXT("{orange}");
 						Status += Entity.DebugGetDescription();
 
-						// LOD
-						if (bHasLOD)
+						if (EntityDescriptionVerbosity >= EEntityDescriptionVerbosity::Full)
 						{
-							Status += TEXT(" {white}LOD ");
-							switch (SimLODList[EntityIndex].LOD)
+							// LOD
+							if (bHasLOD)
 							{
-							case EMassLOD::High:
-								Status += TEXT("High");
-								break;
-							case EMassLOD::Medium:
-								Status += TEXT("Med");
-								break;
-							case EMassLOD::Low:
-								Status += TEXT("Low");
-								break;
-							case EMassLOD::Off:
-								Status += TEXT("Off");
-								break;
-							default:
-								Status += TEXT("?");
-								break;
+								Status += TEXT(" {white}LOD ");
+								switch (SimLODList[EntityIndex].LOD)
+								{
+								case EMassLOD::High:
+									Status += TEXT("High");
+									break;
+								case EMassLOD::Medium:
+									Status += TEXT("Med");
+									break;
+								case EMassLOD::Low:
+									Status += TEXT("Low");
+									break;
+								case EMassLOD::Off:
+									Status += TEXT("Off");
+									break;
+								default:
+									Status += TEXT("?");
+									break;
+								}
 							}
-						}
-						Status += TEXT("\n");
+							Status += TEXT("\n");
 
-						// StateTree
-						if (bHasStateTree)
-						{
-							// Current StateTree task
-							// Optional shared fragment is expected to be present for entities with a StateTree instance
-							const FMassStateTreeSharedFragment& SharedStateTree = Context.GetConstSharedFragment<FMassStateTreeSharedFragment>();
-							if (const UStateTree* StateTree = SharedStateTree.StateTree)
+							// StateTree
+							if (bHasStateTree)
 							{
-								if (FStateTreeInstanceData* InstanceData = MassStateTreeSubsystem->GetInstanceData(StateTreeInstanceList[EntityIndex].InstanceHandle))
-								{
-									FStateTreeReadOnlyExecutionContext StateTreeContext(MassStateTreeSubsystem, StateTree, *InstanceData);
-									Status += StateTreeContext.GetActiveStateName();
-									Status += FString::Printf(TEXT("  {yellow}%d{white}\n"), StateTreeContext.GetStateChangeCount());
-								}
-								else
+								// Current StateTree task
+								// Optional shared fragment is expected to be present for entities with a StateTree instance
+								const FMassStateTreeSharedFragment& SharedStateTree = Context.GetConstSharedFragment<FMassStateTreeSharedFragment>();
+								if (const UStateTree* StateTree = SharedStateTree.StateTree)
 								{
-									Status += TEXT("{red}<No StateTree instance>{white}\n");
+									if (FStateTreeInstanceData* InstanceData = MassStateTreeSubsystem->GetInstanceData(StateTreeInstanceList[EntityIndex].InstanceHandle))
+									{
+										FStateTreeReadOnlyExecutionContext StateTreeContext(MassStateTreeSubsystem, StateTree, *InstanceData);
+										Status += StateTreeContext.GetActiveStateName();
+										Status += FString::Printf(TEXT("  {yellow}%d{white}\n"), StateTreeContext.GetStateChangeCount());
+									}
+									else
+									{
+										Status += TEXT("{red}<No StateTree instance>{white}\n");
+									}
 								}
 							}
-						}
 
-						// Movement info
-						if (bHasMoveTarget)
-						{
-							const FMassMoveTargetFragment& MoveTarget = MoveTargetList[EntityIndex];
-							if (bHasVelocity && bHasForce)
+							// Movement info
+							if (bHasMoveTarget)
 							{
-								Status += FString::Printf(TEXT("{yellow}%s/%03d {lightgrey}Speed:{white}%.1f {lightgrey}Force:{white}%.1f\n"),
-									*UEnum::GetDisplayValueAsText(MoveTarget.GetCurrentAction()).ToString(), MoveTarget.GetCurrentActionID(), VelocityList[EntityIndex].Value.Length(), ForceList[EntityIndex].Value.Length());
-							}
-
-							Status += FString::Printf(TEXT("{pink}-> %s {white}Dist: %.1f\n"),
-								*UEnum::GetDisplayValueAsText(MoveTarget.IntentAtGoal).ToString(), MoveTarget.DistanceToGoal);
-						}
+								const FMassMoveTargetFragment& MoveTarget = MoveTargetList[EntityIndex];
+								if (bHasVelocity && bHasForce)
+								{
+									Status += FString::Printf(TEXT("{yellow}%s/%03d {lightgrey}Speed:{white}%.1f {lightgrey}Force:{white}%.1f\n"),
+										*UEnum::GetDisplayValueAsText(MoveTarget.GetCurrentAction()).ToString(), MoveTarget.GetCurrentActionID(), VelocityList[EntityIndex].Value.Length(), ForceList[EntityIndex].Value.Length());
+								}
 
-						// LookAt
-						if (bShowEntityLookAt && bHasLookAt)
-						{
-							const FMassLookAtFragment& LookAt = LookAtList[EntityIndex];
-							const double RemainingTime = LookAt.GazeDuration - (CurrentTime - LookAt.GazeStartTime);
-							Status += FString::Printf(TEXT("{turquoise}%s/%s {lightgrey}%.1f\n"),
-								*UEnum::GetDisplayValueAsText(LookAt.LookAtMode).ToString(), *UEnum::GetDisplayValueAsText(LookAt.RandomGazeMode).ToString(), RemainingTime);
+								Status += FString::Printf(TEXT("{pink}-> %s {white}Dist: %.1f\n"),
+									*UEnum::GetDisplayValueAsText(MoveTarget.IntentAtGoal).ToString(), MoveTarget.DistanceToGoal);
+							}
 
-							if (LookAtSubsystem != nullptr)
+							// LookAt
+							if (bShowEntityLookAt && bHasLookAt)
 							{
-								Status += FString::Printf(TEXT("{lightgrey}%s\n"), *LookAtSubsystem->DebugGetRequestsString(Entity));
+								const FMassLookAtFragment& LookAt = LookAtList[EntityIndex];
+								const double RemainingTime = LookAt.GazeDuration - (CurrentTime - LookAt.GazeStartTime);
+								Status += FString::Printf(TEXT("{turquoise}%s/%s {lightgrey}%.1f\n"),
+									*UEnum::GetDisplayValueAsText(LookAt.LookAtMode).ToString(), *UEnum::GetDisplayValueAsText(LookAt.RandomGazeMode).ToString(), RemainingTime);
+
+								if (LookAtSubsystem != nullptr)
+								{
+									Status += FString::Printf(TEXT("{lightgrey}%s\n"), *LookAtSubsystem->DebugGetRequestsString(Entity));
+								}
 							}
 						}
 
@@ -936,8 +939,15 @@ void FGameplayDebuggerCategory_Mass::DrawData(APlayerController* OwnerPC, FGamep
 	CanvasContext.Printf(TEXT("[{yellow}%s{white}] %s Entity Avoidance"), *GetInputHandlerDescription(6), bShowNearEntityAvoidance ? TEXT("Hide") : TEXT("Show"));
 	CanvasContext.Printf(TEXT("[{yellow}%s{white}] %s Entity Path"), *GetInputHandlerDescription(7), bShowNearEntityPath ? TEXT("Hide") : TEXT("Show"));
 	CanvasContext.Printf(TEXT("[{yellow}%s{white}] %s Entity LookAt"), *GetInputHandlerDescription(8), bShowEntityLookAt ? TEXT("Hide") : TEXT("Show"));
-	CanvasContext.Printf(TEXT("[{yellow}%s{white}] %s Entity Description"), *GetInputHandlerDescription(9), bShowEntityDescription ? TEXT("Hide") : TEXT("Show"));
 
+	switch (EntityDescriptionVerbosity)
+	{
+	case EEntityDescriptionVerbosity::Hidden: CanvasContext.Printf(TEXT("[{yellow}%s{white}] Cycle Description Verbosity (none)"), *GetInputHandlerDescription(9)); break;
+	case EEntityDescriptionVerbosity::Minimal: CanvasContext.Printf(TEXT("[{yellow}%s{white}] Cycle Description Verbosity (minimal)"), *GetInputHandlerDescription(9));	break;
+	case EEntityDescriptionVerbosity::Full: [[fallthrough]];
+	case EEntityDescriptionVerbosity::Max: CanvasContext.Printf(TEXT("[{yellow}%s{white}] Cycle Description Verbosity (full)"), *GetInputHandlerDescription(9)); break;
+	}
+	
 	if (IsCategoryLocal() && !IsCategoryAuth())
 	{
 		// we want to display this line only on clients in client-server environment.
diff --git a/Engine/Plugins/AI/MassAI/Source/MassAIDebug/Public/GameplayDebuggerCategory_Mass.h b/Engine/Plugins/AI/MassAI/Source/MassAIDebug/Public/GameplayDebuggerCategory_Mass.h
--- a/Engine/Plugins/AI/MassAI/Source/MassAIDebug/Public/GameplayDebuggerCategory_Mass.h
+++ b/Engine/Plugins/AI/MassAI/Source/MassAIDebug/Public/GameplayDebuggerCategory_Mass.h
@@ -35,27 +35,65 @@ protected:
 	
 	void SetCachedEntity(const FMassEntityHandle Entity, const FMassEntityManager& EntityManager);
 
-	void OnToggleArchetypes() { bShowArchetypes = !bShowArchetypes; }
-	void OnToggleShapes() { bShowShapes = !bShowShapes; }
-	void OnToggleAgentFragments() { bShowAgentFragments = !bShowAgentFragments; }
-	void OnPickEntity() { bPickEntity = true; }
-	void OnToggleEntityDetails() { bShowEntityDetails = !bShowEntityDetails; }
-	void OnToggleNearEntityOverview() { bShowNearEntityOverview = !bShowNearEntityOverview; }
-	void OnToggleNearEntityAvoidance() { bShowNearEntityAvoidance = !bShowNearEntityAvoidance; }
-	void OnToggleNearEntityPath() { bShowNearEntityPath = !bShowNearEntityPath; }
-	void OnToggleEntityLookAt() 
-	{ 
-		bShowEntityLookAt = !bShowEntityLookAt; 
+	void OnToggleArchetypes()
+	{
+		bShowArchetypes = !bShowArchetypes;
+	}
+
+	void OnToggleShapes()
+	{
+		bShowShapes = !bShowShapes;
+	}
+
+	void OnToggleAgentFragments()
+	{
+		bShowAgentFragments = !bShowAgentFragments;
+	}
+
+	void OnPickEntity()
+	{
+		bPickEntity = true;
+	}
+
+	void OnToggleEntityDetails()
+	{
+		bShowEntityDetails = !bShowEntityDetails;
+	}
+
+	void OnToggleNearEntityOverview()
+	{
+		bShowNearEntityOverview = !bShowNearEntityOverview;
+	}
+
+	void OnToggleNearEntityAvoidance()
+	{
+		bShowNearEntityAvoidance = !bShowNearEntityAvoidance;
 	}
-	void OnToggleEntityDescription() 
-	{ 
-		bShowEntityDescription = !bShowEntityDescription; 
+	void OnToggleNearEntityPath()
+	{
+		bShowNearEntityPath = !bShowNearEntityPath;
+	}
+
+	void OnToggleEntityLookAt()
+	{
+		bShowEntityLookAt = !bShowEntityLookAt;
+	}
+
+	void OnCycleEntityDescription()
+	{
+		EntityDescriptionVerbosity = static_cast<EEntityDescriptionVerbosity>(
+			(static_cast<uint8>(EntityDescriptionVerbosity) + 1) % static_cast<uint8>(EEntityDescriptionVerbosity::Max));
 	}
+
 	void OnToggleDebugLocalEntityManager();
 	void OnIncreaseSearchRange();
 	void OnDecreaseSearchRange();
 	void OnTogglePickedActorAsViewer();
-	void OnToggleDrawViewers() { bShowViewers = !bShowViewers; }
+	void OnToggleDrawViewers()
+	{
+		bShowViewers = !bShowViewers;
+	}
+
 	void OnClearActorViewers();
 	
 	void PickEntity(const FVector& ViewLocation, const FVector& ViewDirection, const UWorld& World, FMassEntityManager& EntityManager, const bool bLimitAngle = true);
@@ -67,8 +105,18 @@ protected:
 	void ClearCachedEntity();
 
 protected:
+	enum class EEntityDescriptionVerbosity : uint8
+	{
+		Hidden,
+		Minimal,
+		Full,
+		Max
+	};
+
 	TWeakObjectPtr<AActor> CachedDebugActor;
 	FMassEntityHandle CachedEntity;
+	EEntityDescriptionVerbosity EntityDescriptionVerbosity = EEntityDescriptionVerbosity::Minimal;
+
 	bool bShowArchetypes;
 	bool bShowShapes;
 	bool bShowAgentFragments;
@@ -78,7 +126,6 @@ protected:
 	bool bShowNearEntityAvoidance;
 	bool bShowNearEntityPath;
 	bool bShowEntityLookAt;
-	bool bShowEntityDescription;
 	bool bMarkEntityBeingDebugged;
 	bool bDebugLocalEntityManager;
 	bool bShowViewers;
@@ -108,10 +155,10 @@ protected:
 	static FOnConsoleCommandBroadcastDelegate OnToggleEntityDetailsBroadcast;
 	static FOnConsoleCommandBroadcastDelegate OnToggleNearEntityOverviewBroadcast;
 	static FOnConsoleCommandBroadcastDelegate OnToggleNearEntityAvoidanceBroadcast;
-	static FOnConsoleCommandBroadcastDelegate OnToggleNearEntityPathBroadcast; 	
+	static FOnConsoleCommandBroadcastDelegate OnToggleNearEntityPathBroadcast;
 	static FOnConsoleCommandBroadcastDelegate OnToggleEntityLookAtBroadcast;
-	static FOnConsoleCommandBroadcastDelegate OnToggleEntityDescriptionBroadcast;
-	static FOnConsoleCommandBroadcastDelegate OnToggleDebugLocalEntityManagerBroadcast; 	
+	static FOnConsoleCommandBroadcastDelegate OnCycleEntityDescriptionBroadcast;
+	static FOnConsoleCommandBroadcastDelegate OnToggleDebugLocalEntityManagerBroadcast;
 	static FOnConsoleCommandBroadcastDelegate OnTogglePickedActorAsViewerBroadcast;
 	static FOnConsoleCommandBroadcastDelegate OnToggleDrawViewersBroadcast;
 	static FOnConsoleCommandBroadcastDelegate OnClearActorViewersBroadcast;
