[MassLookAt] added CreateLookAtActorRequest BP function

diff --git a/Engine/Plugins/AI/MassAI/Config/BaseMassAI.ini b/Engine/Plugins/AI/MassAI/Config/BaseMassAI.ini
--- a/Engine/Plugins/AI/MassAI/Config/BaseMassAI.ini
+++ b/Engine/Plugins/AI/MassAI/Config/BaseMassAI.ini
@@ -1,2 +1,2 @@
 [CoreRedirects]
-; no changes past UE 5.0
++FunctionRedirects=(OldName="/Script/MassAIBehavior.MassLookAtSubsystem.CreateRequest",NewName="/Script/MassAIBehavior.MassLookAtSubsystem.CreateLookAtPositionRequest")
diff --git a/Engine/Plugins/AI/MassAI/Source/MassAIBehavior/Private/MassLookAtSubsystem.cpp b/Engine/Plugins/AI/MassAI/Source/MassAIBehavior/Private/MassLookAtSubsystem.cpp
--- a/Engine/Plugins/AI/MassAI/Source/MassAIBehavior/Private/MassLookAtSubsystem.cpp
+++ b/Engine/Plugins/AI/MassAI/Source/MassAIBehavior/Private/MassLookAtSubsystem.cpp
@@ -211,7 +211,7 @@ static FAutoConsoleCommandWithWorldArgsAndOutputDevice CmdSendLookAtPlayerReques
 
 } // UE::Mass::LookAt::Private
 
-FMassLookAtRequestHandle UMassLookAtSubsystem::CreateRequest(AActor* ViewerActor, FMassLookAtPriority Priority, const FVector TargetLocation) const
+FMassLookAtRequestHandle UMassLookAtSubsystem::CreateLookAtPositionRequest(AActor* ViewerActor, FMassLookAtPriority Priority, const FVector TargetLocation) const
 {
 	const UWorld* World = GetWorld();
 
@@ -232,7 +232,7 @@ FMassLookAtRequestHandle UMassLookAtSubsystem::CreateRequest(AActor* ViewerActor
 	const FMassEntityHandle ViewerEntity = ActorSubsystem->GetEntityHandleFromActor(ViewerActor);
 	if (!ViewerEntity.IsSet())
 	{
-		UE_VLOG_UELOG(this, LogMassBehavior, Error, TEXT("%hs failed: unable to find a MassEntity associated to the player"), __FUNCTION__);
+		UE_VLOG_UELOG(this, LogMassBehavior, Error, TEXT("%hs failed: unable to find a MassEntity associated to '%s'"), __FUNCTION__, *GetNameSafe(ViewerActor));
 		return {};
 	}
 
@@ -265,6 +265,63 @@ FMassLookAtRequestHandle UMassLookAtSubsystem::CreateRequest(AActor* ViewerActor
 	return FMassLookAtRequestHandle{RequestEntity, TargetEntity};
 }
 
+FMassLookAtRequestHandle UMassLookAtSubsystem::CreateLookAtActorRequest(AActor* ViewerActor, FMassLookAtPriority Priority, AActor* TargetActor) const
+{
+	const UWorld* World = GetWorld();
+
+	UMassActorSubsystem* ActorSubsystem = World->GetSubsystem<UMassActorSubsystem>();
+	if (ActorSubsystem == nullptr)
+	{
+		UE_VLOG_UELOG(this, LogMassBehavior, Error, TEXT("%hs failed: unable to find MassActorSubsystem"), __FUNCTION__);
+		return {};
+	}
+
+	UMassEntitySubsystem* EntitySubsystem = World->GetSubsystem<UMassEntitySubsystem>();
+	if (EntitySubsystem == nullptr)
+	{
+		UE_VLOG_UELOG(this, LogMassBehavior, Error, TEXT("%hs failed: unable to find UMassEntitySubsystem"), __FUNCTION__);
+		return {};
+	}
+
+	const FMassEntityHandle ViewerEntity = ActorSubsystem->GetEntityHandleFromActor(ViewerActor);
+	if (!ViewerEntity.IsSet())
+	{
+		UE_VLOG_UELOG(this, LogMassBehavior, Error, TEXT("%hs failed: unable to find a MassEntity associated to '%s'"), __FUNCTION__, *GetNameSafe(ViewerActor));
+		return {};
+	}
+
+	const FMassEntityHandle TargetEntity = ActorSubsystem->GetEntityHandleFromActor(TargetActor);
+	if (!TargetEntity.IsSet())
+	{
+		UE_VLOG_UELOG(this, LogMassBehavior, Error, TEXT("%hs failed: unable to find a MassEntity associated to '%s'"), __FUNCTION__, *GetNameSafe(TargetActor));
+		return {};
+	}
+
+	// @todo: once available, consider using incoming EntityBuilder improvements for this whole creation step
+	FMassEntityManager& EntityManager = EntitySubsystem->GetMutableEntityManager();
+	FMassEntityHandle RequestEntity = EntityManager.ReserveEntity();
+
+	// Push command to create a new entity representing a LookAt target along with a request to look at it
+	const TSharedPtr<FMassCommandBuffer> CommandBuffer = MakeShareable(new FMassCommandBuffer());
+	CommandBuffer->PushCommand<FMassDeferredCreateCommand>(
+		[
+			TargetArchetype = TargetArchetype
+			, RequestArchetype = RequestArchetype
+			, RequestEntity
+			, ViewerEntity
+			, TargetEntity
+			, Priority
+		]
+		(FMassEntityManager& InEntityManager)
+		{
+			UE::Mass::LookAt::Private::CreateRequests(InEntityManager, RequestArchetype, { ViewerEntity }, { RequestEntity }, TargetEntity, Priority);
+		});
+
+	UE::Mass::LookAt::Private::FlushCommands(EntityManager, CommandBuffer);
+
+	return FMassLookAtRequestHandle{RequestEntity, TargetEntity};
+}
+
 void UMassLookAtSubsystem::DeleteRequest(const FMassLookAtRequestHandle RequestHandle) const
 {
 	// Simple validation when none of the entities are set since it is probably due to a bad data setup.
diff --git a/Engine/Plugins/AI/MassAI/Source/MassAIBehavior/Public/MassLookAtSubsystem.h b/Engine/Plugins/AI/MassAI/Source/MassAIBehavior/Public/MassLookAtSubsystem.h
--- a/Engine/Plugins/AI/MassAI/Source/MassAIBehavior/Public/MassLookAtSubsystem.h
+++ b/Engine/Plugins/AI/MassAI/Source/MassAIBehavior/Public/MassLookAtSubsystem.h
@@ -45,7 +45,14 @@ public:
 	 * Method will log an error if it fails (e.g., no associated mass entity).
 	 */
 	UFUNCTION(BlueprintCallable, BlueprintPure = false, Category = "LookAt", meta = (ReturnDisplayName="Request Handle"))
-	FMassLookAtRequestHandle CreateRequest(AActor* ViewerActor, FMassLookAtPriority Priority, FVector TargetLocation) const;
+	FMassLookAtRequestHandle CreateLookAtPositionRequest(AActor* ViewerActor, FMassLookAtPriority Priority, FVector TargetLocation) const;
+
+	/**
+	 * Creates a new LookAt request using the provided parameters for the mass entity associated to 'ViewerActor', if any.
+	 * Method will log an error if it fails (e.g., no associated mass entity for viewer or target actors).
+	 */
+	UFUNCTION(BlueprintCallable, BlueprintPure = false, Category = "LookAt", meta = (ReturnDisplayName="Request Handle"))
+	FMassLookAtRequestHandle CreateLookAtActorRequest(AActor* ViewerActor, FMassLookAtPriority Priority, AActor* TargetActor) const;
 
 	/**
 	 * Removes given request from the active LookAt requests.
