Clarify UMassSteerToMoveTargetProcessor debug display

diff --git a/Engine/Plugins/AI/MassAI/Source/MassNavigation/Private/MassSteeringProcessors.cpp b/Engine/Plugins/AI/MassAI/Source/MassNavigation/Private/MassSteeringProcessors.cpp
--- a/Engine/Plugins/AI/MassAI/Source/MassNavigation/Private/MassSteeringProcessors.cpp
+++ b/Engine/Plugins/AI/MassAI/Source/MassNavigation/Private/MassSteeringProcessors.cpp
@@ -368,14 +368,14 @@ void UMassSteerToMoveTargetProcessor::Execute(FMassEntityManager& EntityManager,
 					UE_VLOG_SEGMENT_THICK(LogOwner, LogMassNavigation, Log, MoveTargetCenter, CurrentLocation + ZOffset, LightEntityColor, /*Thickness*/1, TEXT(""));
 				}
 
-				// Display Desired Velocity
+				// Display DesiredMovement DesiredVelocity
 				UE_VLOG_SEGMENT_THICK(LogOwner, LogMassNavigation, Log, CurrentLocation + ZOffset, CurrentLocation + ZOffset + DesiredMovement.DesiredVelocity, FColor::Yellow, /*Thickness*/4,
-					TEXT("DesiredVelocity %.1f"), DesiredMovement.DesiredVelocity.Length());
+					TEXT("Mvt DesiredVelocity %.1f"), DesiredMovement.DesiredVelocity.Length());
 
-				// Display DesiredVelocity
+				// Display Steering internal DesiredVelocity
 				UE_VLOG_SEGMENT_THICK(LogOwner, LogMassNavigation, Log, CurrentLocation + ZOffset + FVector(0,0,2),
 					CurrentLocation + Steering.DesiredVelocity + ZOffset + FVector(0,0,2), FColor::Orange, /*Thickness*/4,
-					TEXT("DesiredVelocity %.1f"), Steering.DesiredVelocity.Length());
+					TEXT("Steering DesiredVelocity %.1f"), Steering.DesiredVelocity.Length());
 
 				// Display Force
 				UE_VLOG_SEGMENT_THICK(LogOwner, LogMassNavigation, Log, CurrentLocation + ZOffset, CurrentLocation + Force.Value + ZOffset, FColor::Emerald, /*Thickness*/4,
