[MassLookAt] Added custom interpolation speed

diff --git a/Engine/Plugins/AI/MassAI/Source/MassAIBehavior/Private/MassLookAtSubsystem.cpp b/Engine/Plugins/AI/MassAI/Source/MassAIBehavior/Private/MassLookAtSubsystem.cpp
--- a/Engine/Plugins/AI/MassAI/Source/MassAIBehavior/Private/MassLookAtSubsystem.cpp
+++ b/Engine/Plugins/AI/MassAI/Source/MassAIBehavior/Private/MassLookAtSubsystem.cpp
@@ -30,7 +30,9 @@ bool CreateRequests(FMassEntityManager& InEntityManager
 	, const TConstArrayView<FMassEntityHandle> Viewers
 	, const TConstArrayView<FMassEntityHandle> Requests
 	, const FMassEntityHandle TargetEntity
-	, const FMassLookAtPriority Priority)
+	, const FMassLookAtPriority Priority
+	, const EMassLookAtInterpolationSpeed InterpolationSpeed
+	, const float CustomInterpolationSpeed = UE::Mass::LookAt::DefaultCustomInterpolationSpeed)
 {
 	if (Viewers.IsEmpty())
 	{
@@ -43,8 +45,16 @@ bool CreateRequests(FMassEntityManager& InEntityManager
 
 	for (int Index = 0; Index < Requests.Num(); ++Index)
 	{
-		InEntityManager.SetEntityFragmentValues(Requests[Index]
-			, { FInstancedStruct::Make(FMassLookAtRequestFragment{Viewers[Index], Priority, EMassLookAtMode::LookAtEntity, TargetEntity}) });
+		InEntityManager.SetEntityFragmentValues(Requests[Index], { FInstancedStruct::Make(
+			FMassLookAtRequestFragment
+			{
+				Viewers[Index]
+				, Priority
+				, EMassLookAtMode::LookAtEntity
+				, TargetEntity
+				, InterpolationSpeed
+				, CustomInterpolationSpeed
+			})});
 
 		UE_VLOG_UELOG(InEntityManager.GetOwner(), LogMassBehavior, Log, TEXT("Created LookAt Request '%s', Target '%s', Priority=%s")
 			, *LexToString(Requests[Index])
@@ -192,7 +202,8 @@ static FAutoConsoleCommandWithWorldArgsAndOutputDevice CmdSendLookAtPlayerReques
 							, Viewers
 							, InEntityManager.BatchReserveEntities(Viewers.Num(), DebugRequests)
 							, PlayerEntity
-							, FMassLookAtPriority{ PriorityLevel });
+							, FMassLookAtPriority{ PriorityLevel }
+							, EMassLookAtInterpolationSpeed::Regular);
 					});
 			}
 			else
@@ -211,7 +222,12 @@ static FAutoConsoleCommandWithWorldArgsAndOutputDevice CmdSendLookAtPlayerReques
 
 } // UE::Mass::LookAt::Private
 
-FMassLookAtRequestHandle UMassLookAtSubsystem::CreateLookAtPositionRequest(AActor* ViewerActor, FMassLookAtPriority Priority, const FVector TargetLocation) const
+FMassLookAtRequestHandle UMassLookAtSubsystem::CreateLookAtPositionRequest(
+	AActor* ViewerActor
+	, const FMassLookAtPriority Priority
+	, const FVector TargetLocation
+	, const EMassLookAtInterpolationSpeed InterpolationSpeed
+	, const float CustomInterpolationSpeed) const
 {
 	const UWorld* World = GetWorld();
 
@@ -253,11 +269,22 @@ FMassLookAtRequestHandle UMassLookAtSubsystem::CreateLookAtPositionRequest(AActo
 			, TargetEntity
 			, TargetTransform
 			, Priority
+			, InterpolationSpeed
+			, CustomInterpolationSpeed
 		]
 		(FMassEntityManager& InEntityManager)
 		{
 			UE::Mass::LookAt::Private::CreateTargets(InEntityManager, TargetArchetype, { TargetEntity }, { TargetTransform }, Priority);
-			UE::Mass::LookAt::Private::CreateRequests(InEntityManager, RequestArchetype, { ViewerEntity }, { RequestEntity }, TargetEntity, Priority);
+
+			UE::Mass::LookAt::Private::CreateRequests(
+				InEntityManager
+				, RequestArchetype
+				, { ViewerEntity }
+				, { RequestEntity }
+				, TargetEntity
+				, Priority
+				, InterpolationSpeed
+				, CustomInterpolationSpeed);
 		});
 
 	UE::Mass::LookAt::Private::FlushCommands(EntityManager, CommandBuffer);
@@ -265,7 +292,12 @@ FMassLookAtRequestHandle UMassLookAtSubsystem::CreateLookAtPositionRequest(AActo
 	return FMassLookAtRequestHandle{RequestEntity, TargetEntity};
 }
 
-FMassLookAtRequestHandle UMassLookAtSubsystem::CreateLookAtActorRequest(AActor* ViewerActor, FMassLookAtPriority Priority, AActor* TargetActor) const
+FMassLookAtRequestHandle UMassLookAtSubsystem::CreateLookAtActorRequest(
+	AActor* ViewerActor
+	, const FMassLookAtPriority Priority
+	, AActor* TargetActor
+	, const EMassLookAtInterpolationSpeed InterpolationSpeed
+	, const float CustomInterpolationSpeed) const
 {
 	const UWorld* World = GetWorld();
 
@@ -305,16 +337,25 @@ FMassLookAtRequestHandle UMassLookAtSubsystem::CreateLookAtActorRequest(AActor*
 	const TSharedPtr<FMassCommandBuffer> CommandBuffer = MakeShareable(new FMassCommandBuffer());
 	CommandBuffer->PushCommand<FMassDeferredCreateCommand>(
 		[
-			TargetArchetype = TargetArchetype
-			, RequestArchetype = RequestArchetype
+			RequestArchetype = RequestArchetype
 			, RequestEntity
 			, ViewerEntity
 			, TargetEntity
 			, Priority
+			, InterpolationSpeed
+			, CustomInterpolationSpeed
 		]
 		(FMassEntityManager& InEntityManager)
 		{
-			UE::Mass::LookAt::Private::CreateRequests(InEntityManager, RequestArchetype, { ViewerEntity }, { RequestEntity }, TargetEntity, Priority);
+			UE::Mass::LookAt::Private::CreateRequests(
+				InEntityManager
+				, RequestArchetype
+				, { ViewerEntity }
+				, { RequestEntity }
+				, TargetEntity
+				, Priority
+				, InterpolationSpeed
+				, CustomInterpolationSpeed);
 		});
 
 	UE::Mass::LookAt::Private::FlushCommands(EntityManager, CommandBuffer);
@@ -526,6 +567,7 @@ void UMassLookAtSubsystem::UpdateLookAts(const FMassExecutionContext& Context, c
 						// Only update main LookAt information and don't modify gaze related one
 						const FMassLookAtRequestFragment& Request = Pair.Value;
 						LookAtFragment->InterpolationSpeed = Request.InterpolationSpeed;
+						LookAtFragment->CustomInterpolationSpeed = Request.CustomInterpolationSpeed;
 						LookAtFragment->LookAtMode = Request.LookAtMode;
 						LookAtFragment->TrackedEntity = Request.TargetEntity;
 					}
diff --git a/Engine/Plugins/AI/MassAI/Source/MassAIBehavior/Private/Tasks/MassLookAtTask.cpp b/Engine/Plugins/AI/MassAI/Source/MassAIBehavior/Private/Tasks/MassLookAtTask.cpp
--- a/Engine/Plugins/AI/MassAI/Source/MassAIBehavior/Private/Tasks/MassLookAtTask.cpp
+++ b/Engine/Plugins/AI/MassAI/Source/MassAIBehavior/Private/Tasks/MassLookAtTask.cpp
@@ -34,7 +34,12 @@ EStateTreeRunStatus FMassLookAtTask::EnterState(FStateTreeExecutionContext& Cont
 
 	MassStateTreeContext.GetMassEntityExecutionContext().Defer()
 		.PushCommand<FMassDeferredCreateCommand>(
-			[&Context, SelfEntity = MassStateTreeContext.GetEntity(), &InstanceData, Mode = LookAtMode, Priority = Priority, InterpolationSpeed = InterpolationSpeed]
+			[&Context, SelfEntity = MassStateTreeContext.GetEntity()
+			, &InstanceData
+			, Mode = LookAtMode
+			, Priority = Priority
+			, InterpolationSpeed = InterpolationSpeed
+			, CustomInterpolationSpeed = CustomInterpolationSpeed]
 			(FMassEntityManager& InEntityManager)
 			{
 				const FMassArchetypeCompositionDescriptor Composition{FMassFragmentBitSet(*FMassLookAtRequestFragment::StaticStruct())};
@@ -47,8 +52,15 @@ EStateTreeRunStatus FMassLookAtTask::EnterState(FStateTreeExecutionContext& Cont
 				if (Requests.Num() == NumRequestsToCreate)
 				{
 					InstanceData.Request = Requests.Top();
-					FMassLookAtRequestFragment RequestFragment{ SelfEntity, Priority, Mode, InstanceData.TargetEntity };
-					RequestFragment.InterpolationSpeed = InterpolationSpeed;
+					FMassLookAtRequestFragment RequestFragment
+					{
+						SelfEntity
+						, Priority
+						, Mode
+						, InstanceData.TargetEntity
+						, InterpolationSpeed
+						, CustomInterpolationSpeed
+					};
 
 					// When using 'LookAtEntity' we validate the target entity to use it,
 					// or we use the default 'LookForward' as fallback.
@@ -57,7 +69,8 @@ EStateTreeRunStatus FMassLookAtTask::EnterState(FStateTreeExecutionContext& Cont
 						if (!InstanceData.TargetEntity.IsSet())
 						{
 							MASSBEHAVIOR_LOG(Error, TEXT("Failed LookAt: invalid target entity"));
-							RequestFragment = { SelfEntity, Priority, EMassLookAtMode::LookForward };
+							RequestFragment.LookAtMode = EMassLookAtMode::LookForward;
+							RequestFragment.TargetEntity = {};
 						}
 					}
 
diff --git a/Engine/Plugins/AI/MassAI/Source/MassAIBehavior/Public/MassLookAtFragments.h b/Engine/Plugins/AI/MassAI/Source/MassAIBehavior/Public/MassLookAtFragments.h
--- a/Engine/Plugins/AI/MassAI/Source/MassAIBehavior/Public/MassLookAtFragments.h
+++ b/Engine/Plugins/AI/MassAI/Source/MassAIBehavior/Public/MassLookAtFragments.h
@@ -117,6 +117,13 @@ struct MASSAIBEHAVIOR_API FMassLookAtFragment : public FMassFragment
 	UPROPERTY(Transient)
 	EMassLookAtInterpolationSpeed InterpolationSpeed = EMassLookAtInterpolationSpeed::Regular;
 
+	/**
+	 * Primary look at custom interpolation speed used when 'InterpolationSpeed = EMassLookAtInterpolationSpeed::Custom'
+	 * (not used by the LookAt processor but can be forwarded to the animation system).
+	 */
+	UPROPERTY(Transient)
+	float CustomInterpolationSpeed = UE::Mass::LookAt::DefaultCustomInterpolationSpeed;
+
 	/** Gaze look at mode. */
 	UPROPERTY(Transient)
 	EMassLookAtGazeMode RandomGazeMode = EMassLookAtGazeMode::None;
@@ -159,10 +166,18 @@ struct MASSAIBEHAVIOR_API FMassLookAtRequestFragment : public FMassFragment
 	{
 	}
 
-	FMassLookAtRequestFragment(const FMassEntityHandle ViewerEntity, const FMassLookAtPriority Priority, const EMassLookAtMode Mode, const FMassEntityHandle Target)
+	FMassLookAtRequestFragment(
+		const FMassEntityHandle ViewerEntity
+		, const FMassLookAtPriority Priority
+		, const EMassLookAtMode Mode
+		, const FMassEntityHandle Target
+		, const EMassLookAtInterpolationSpeed InterpolationSpeed
+		, const float CustomInterpolationSpeed)
 		: FMassLookAtRequestFragment(ViewerEntity, Priority, Mode)
 	{
 		TargetEntity = Target;
+		this->InterpolationSpeed = InterpolationSpeed;
+		this->CustomInterpolationSpeed = CustomInterpolationSpeed;
 	}
 
 	UPROPERTY(Transient)
@@ -180,6 +195,9 @@ struct MASSAIBEHAVIOR_API FMassLookAtRequestFragment : public FMassFragment
 	UPROPERTY(Transient)
 	EMassLookAtInterpolationSpeed InterpolationSpeed = EMassLookAtInterpolationSpeed::Regular;
 
+	UPROPERTY(Transient)
+	float CustomInterpolationSpeed = UE::Mass::LookAt::DefaultCustomInterpolationSpeed;
+
 	friend FString LexToString(const FMassLookAtRequestFragment& Fragment)
 	{
 		if (Fragment.LookAtMode == EMassLookAtMode::LookAtEntity)
diff --git a/Engine/Plugins/AI/MassAI/Source/MassAIBehavior/Public/MassLookAtSubsystem.h b/Engine/Plugins/AI/MassAI/Source/MassAIBehavior/Public/MassLookAtSubsystem.h
--- a/Engine/Plugins/AI/MassAI/Source/MassAIBehavior/Public/MassLookAtSubsystem.h
+++ b/Engine/Plugins/AI/MassAI/Source/MassAIBehavior/Public/MassLookAtSubsystem.h
@@ -45,14 +45,22 @@ public:
 	 * Method will log an error if it fails (e.g., no associated mass entity).
 	 */
 	UFUNCTION(BlueprintCallable, BlueprintPure = false, Category = "LookAt", meta = (ReturnDisplayName="Request Handle"))
-	FMassLookAtRequestHandle CreateLookAtPositionRequest(AActor* ViewerActor, FMassLookAtPriority Priority, FVector TargetLocation) const;
+	FMassLookAtRequestHandle CreateLookAtPositionRequest(AActor* ViewerActor
+		, FMassLookAtPriority Priority
+		, FVector TargetLocation
+		, EMassLookAtInterpolationSpeed InterpolationSpeed = EMassLookAtInterpolationSpeed::Regular
+		, float CustomInterpolationSpeed = 1.5f) const;
 
 	/**
 	 * Creates a new LookAt request using the provided parameters for the mass entity associated to 'ViewerActor', if any.
 	 * Method will log an error if it fails (e.g., no associated mass entity for viewer or target actors).
 	 */
 	UFUNCTION(BlueprintCallable, BlueprintPure = false, Category = "LookAt", meta = (ReturnDisplayName="Request Handle"))
-	FMassLookAtRequestHandle CreateLookAtActorRequest(AActor* ViewerActor, FMassLookAtPriority Priority, AActor* TargetActor) const;
+	FMassLookAtRequestHandle CreateLookAtActorRequest(AActor* ViewerActor
+		, FMassLookAtPriority Priority
+		, AActor* TargetActor
+		, EMassLookAtInterpolationSpeed InterpolationSpeed = EMassLookAtInterpolationSpeed::Regular
+		, float CustomInterpolationSpeed = 1.5f) const;
 
 	/**
 	 * Removes given request from the active LookAt requests.
diff --git a/Engine/Plugins/AI/MassAI/Source/MassAIBehavior/Public/MassLookAtTypes.h b/Engine/Plugins/AI/MassAI/Source/MassAIBehavior/Public/MassLookAtTypes.h
--- a/Engine/Plugins/AI/MassAI/Source/MassAIBehavior/Public/MassLookAtTypes.h
+++ b/Engine/Plugins/AI/MassAI/Source/MassAIBehavior/Public/MassLookAtTypes.h
@@ -9,6 +9,7 @@
 
 namespace UE::Mass::LookAt
 {
+	constexpr float DefaultCustomInterpolationSpeed = 1.5f;
 	constexpr int32 HashGridLevelsOfHierarchy = 2;
 	constexpr int32 HashGridRatioBetweenLevels = 4;
 	constexpr int32 HashGridResultsSoftLimit = 16;
@@ -48,7 +49,8 @@ enum class EMassLookAtInterpolationSpeed : uint8
 	Instant,
 	Fast,
 	Regular,
-	Slow
+	Slow,
+	Custom
 };
 
 /**
diff --git a/Engine/Plugins/AI/MassAI/Source/MassAIBehavior/Public/Tasks/MassLookAtTask.h b/Engine/Plugins/AI/MassAI/Source/MassAIBehavior/Public/Tasks/MassLookAtTask.h
--- a/Engine/Plugins/AI/MassAI/Source/MassAIBehavior/Public/Tasks/MassLookAtTask.h
+++ b/Engine/Plugins/AI/MassAI/Source/MassAIBehavior/Public/Tasks/MassLookAtTask.h
@@ -68,6 +68,13 @@ protected:
 	UPROPERTY(EditAnywhere, Category = Parameter)
 	EMassLookAtInterpolationSpeed InterpolationSpeed = EMassLookAtInterpolationSpeed::Regular;
 
+	/**
+	 * Look at custom interpolation speed used when 'InterpolationSpeed = EMassLookAtInterpolationSpeed::Custom'
+	 * (not used by the LookAt processor but can be forwarded to the animation system).
+	 */
+	UPROPERTY(EditAnywhere, Category = Parameter, meta = (EditCondition="InterpolationSpeed == EMassLookAtInterpolationSpeed::Custom", EditConditionHides))
+	float CustomInterpolationSpeed = UE::Mass::LookAt::DefaultCustomInterpolationSpeed;
+
 	/** Random gaze Mode */
 	UPROPERTY(EditAnywhere, Category = Parameter)
 	EMassLookAtGazeMode RandomGazeMode = EMassLookAtGazeMode::None;
