[MassLookAt] allow LookAtTask to be optional if fragment is missing. Useful to reuse similar StateTree on different archetypes.

diff --git a/Engine/Plugins/AI/MassAI/Source/MassAIBehavior/Private/Tasks/MassLookAtTask.cpp b/Engine/Plugins/AI/MassAI/Source/MassAIBehavior/Private/Tasks/MassLookAtTask.cpp
--- a/Engine/Plugins/AI/MassAI/Source/MassAIBehavior/Private/Tasks/MassLookAtTask.cpp
+++ b/Engine/Plugins/AI/MassAI/Source/MassAIBehavior/Private/Tasks/MassLookAtTask.cpp
@@ -30,65 +30,73 @@ EStateTreeRunStatus FMassLookAtTask::EnterState(FStateTreeExecutionContext& Cont
 	InstanceData.Time = 0.f;
 	
 	const FMassStateTreeExecutionContext& MassStateTreeContext = static_cast<FMassStateTreeExecutionContext&>(Context);
-	FMassLookAtFragment& LookAtFragment = MassStateTreeContext.GetExternalData(LookAtHandle);
-
-	MassStateTreeContext.GetMassEntityExecutionContext().Defer()
-		.PushCommand<FMassDeferredCreateCommand>(
-			[&Context, SelfEntity = MassStateTreeContext.GetEntity()
-			, &InstanceData
-			, Mode = LookAtMode
-			, Priority = Priority
-			, InterpolationSpeed = InterpolationSpeed
-			, CustomInterpolationSpeed = CustomInterpolationSpeed]
-			(FMassEntityManager& InEntityManager)
-			{
-				const FMassArchetypeCompositionDescriptor Composition{FMassFragmentBitSet(*FMassLookAtRequestFragment::StaticStruct())};
-
-				constexpr int32 NumRequestsToCreate = 1;
-				TArray<FMassEntityHandle> Requests;
-				const TSharedRef<FMassEntityManager::FEntityCreationContext> CreationContext =
-					InEntityManager.BatchCreateEntities(InEntityManager.CreateArchetype(Composition), NumRequestsToCreate, Requests);
-
-				if (Requests.Num() == NumRequestsToCreate)
+	FMassLookAtFragment* LookAtFragment = MassStateTreeContext.GetExternalDataPtr(LookAtHandle);
+	if (LookAtFragment != nullptr)
+	{
+		MassStateTreeContext.GetMassEntityExecutionContext().Defer()
+			.PushCommand<FMassDeferredCreateCommand>(
+				[&Context, SelfEntity = MassStateTreeContext.GetEntity()
+				, &InstanceData
+				, Mode = LookAtMode
+				, Priority = Priority
+				, InterpolationSpeed = InterpolationSpeed
+				, CustomInterpolationSpeed = CustomInterpolationSpeed]
+				(FMassEntityManager& InEntityManager)
 				{
-					InstanceData.Request = Requests.Top();
-					FMassLookAtRequestFragment RequestFragment
-					{
-						SelfEntity
-						, Priority
-						, Mode
-						, InstanceData.TargetEntity
-						, InterpolationSpeed
-						, CustomInterpolationSpeed
-					};
-
-					// When using 'LookAtEntity' we validate the target entity to use it,
-					// or we use the default 'LookForward' as fallback.
-					if (Mode == EMassLookAtMode::LookAtEntity)
+					const FMassArchetypeCompositionDescriptor Composition{FMassFragmentBitSet(*FMassLookAtRequestFragment::StaticStruct())};
+
+					constexpr int32 NumRequestsToCreate = 1;
+					TArray<FMassEntityHandle> Requests;
+					const TSharedRef<FMassEntityManager::FEntityCreationContext> CreationContext =
+						InEntityManager.BatchCreateEntities(InEntityManager.CreateArchetype(Composition), NumRequestsToCreate, Requests);
+
+					if (Requests.Num() == NumRequestsToCreate)
 					{
-						if (!InstanceData.TargetEntity.IsSet())
+						InstanceData.Request = Requests.Top();
+						FMassLookAtRequestFragment RequestFragment
+						{
+							SelfEntity
+							, Priority
+							, Mode
+							, InstanceData.TargetEntity
+							, InterpolationSpeed
+							, CustomInterpolationSpeed
+						};
+
+						// When using 'LookAtEntity' we validate the target entity to use it,
+						// or we use the default 'LookForward' as fallback.
+						if (Mode == EMassLookAtMode::LookAtEntity)
 						{
-							MASSBEHAVIOR_LOG(Error, TEXT("Failed LookAt: invalid target entity"));
-							RequestFragment.LookAtMode = EMassLookAtMode::LookForward;
-							RequestFragment.TargetEntity = {};
+							if (!InstanceData.TargetEntity.IsSet())
+							{
+								MASSBEHAVIOR_LOG(Error, TEXT("Failed LookAt: invalid target entity"));
+								RequestFragment.LookAtMode = EMassLookAtMode::LookForward;
+								RequestFragment.TargetEntity = {};
+							}
 						}
-					}
 
-					InEntityManager.SetEntityFragmentValues(InstanceData.Request, { FInstancedStruct::Make(RequestFragment) });
-				}
-			});
-
-	LookAtFragment.RandomGazeMode = RandomGazeMode;
-	LookAtFragment.RandomGazeYawVariation = RandomGazeYawVariation;
-	LookAtFragment.RandomGazePitchVariation = RandomGazePitchVariation;
-	LookAtFragment.bRandomGazeEntities = bRandomGazeEntities;
+						InEntityManager.SetEntityFragmentValues(InstanceData.Request, { FInstancedStruct::Make(RequestFragment) });
+					}
+				});
+
+		LookAtFragment->RandomGazeMode = RandomGazeMode;
+		LookAtFragment->RandomGazeYawVariation = RandomGazeYawVariation;
+		LookAtFragment->RandomGazePitchVariation = RandomGazePitchVariation;
+		LookAtFragment->bRandomGazeEntities = bRandomGazeEntities;
+
+		// A Duration <= 0 indicates that the task runs until a transition in the state tree stops it.
+		// Otherwise we schedule a signal to end the task.
+		if (InstanceData.Duration > 0.0f)
+		{
+			UMassSignalSubsystem& MassSignalSubsystem = MassStateTreeContext.GetExternalData(MassSignalSubsystemHandle);
+			MassSignalSubsystem.DelaySignalEntityDeferred(MassStateTreeContext.GetMassEntityExecutionContext(), UE::Mass::Signals::LookAtFinished, MassStateTreeContext.GetEntity(), InstanceData.Duration);
+		}
+	}
 
-	// A Duration <= 0 indicates that the task runs until a transition in the state tree stops it.
-	// Otherwise we schedule a signal to end the task.
-	if (InstanceData.Duration > 0.0f)
+	// LookAt are considered optional so we are returning an immediate success when fragment is missing and a specific duration was specified.
+	if (LookAtFragment == nullptr && InstanceData.Duration > 0.0f)
 	{
-		UMassSignalSubsystem& MassSignalSubsystem = MassStateTreeContext.GetExternalData(MassSignalSubsystemHandle);
-		MassSignalSubsystem.DelaySignalEntityDeferred(MassStateTreeContext.GetMassEntityExecutionContext(), UE::Mass::Signals::LookAtFinished, MassStateTreeContext.GetEntity(), InstanceData.Duration);
+		return EStateTreeRunStatus::Succeeded;
 	}
 
 	return EStateTreeRunStatus::Running;
@@ -99,20 +107,22 @@ void FMassLookAtTask::ExitState(FStateTreeExecutionContext& Context, const FStat
 	const FMassStateTreeExecutionContext& MassStateTreeContext = static_cast<FMassStateTreeExecutionContext&>(Context);
 
 	// Reset gaze directly since it is used only by the task and not by the requests
-	FMassLookAtFragment& LookAtFragment = MassStateTreeContext.GetExternalData(LookAtHandle);
-	LookAtFragment.ResetGaze();
-
-	// Destroy any associated request for the main LookAt
-	FInstanceDataType& InstanceData = Context.GetInstanceData(*this);
-	if (InstanceData.Request.IsSet())
+	if (FMassLookAtFragment* LookAtFragment = MassStateTreeContext.GetExternalDataPtr(LookAtHandle))
 	{
-		MassStateTreeContext.GetMassEntityExecutionContext().Defer()
-			.PushCommand<FMassDeferredDestroyCommand>([Request = InstanceData.Request](FMassEntityManager& InEntityManager)
-			{
-				InEntityManager.DestroyEntity(Request);
-			});
+		LookAtFragment->ResetGaze();
+
+		// Destroy any associated request for the main LookAt
+		FInstanceDataType& InstanceData = Context.GetInstanceData(*this);
+		if (InstanceData.Request.IsSet())
+		{
+			MassStateTreeContext.GetMassEntityExecutionContext().Defer()
+				.PushCommand<FMassDeferredDestroyCommand>([Request = InstanceData.Request](FMassEntityManager& InEntityManager)
+				{
+					InEntityManager.DestroyEntity(Request);
+				});
 
-		InstanceData.Request.Reset();
+			InstanceData.Request.Reset();
+		}
 	}
 }
 
diff --git a/Engine/Plugins/AI/MassAI/Source/MassAIBehavior/Public/Tasks/MassLookAtTask.h b/Engine/Plugins/AI/MassAI/Source/MassAIBehavior/Public/Tasks/MassLookAtTask.h
--- a/Engine/Plugins/AI/MassAI/Source/MassAIBehavior/Public/Tasks/MassLookAtTask.h
+++ b/Engine/Plugins/AI/MassAI/Source/MassAIBehavior/Public/Tasks/MassLookAtTask.h
@@ -54,7 +54,7 @@ protected:
 	virtual void GetDependencies(UE::MassBehavior::FStateTreeDependencyBuilder& Builder) const override;
 
 	TStateTreeExternalDataHandle<UMassSignalSubsystem> MassSignalSubsystemHandle;
-	TStateTreeExternalDataHandle<FMassLookAtFragment> LookAtHandle;
+	TStateTreeExternalDataHandle<FMassLookAtFragment, EStateTreeExternalDataRequirement::Optional> LookAtHandle;
 
 	/** Look At Priority */
 	UPROPERTY(EditAnywhere, Category = Parameter)
