Addition of extra navigation information to MassGameplayDebugger.

diff --git a/Engine/Plugins/AI/MassAI/Source/MassAIDebug/Private/GameplayDebuggerCategory_Mass.cpp b/Engine/Plugins/AI/MassAI/Source/MassAIDebug/Private/GameplayDebuggerCategory_Mass.cpp
--- a/Engine/Plugins/AI/MassAI/Source/MassAIDebug/Private/GameplayDebuggerCategory_Mass.cpp
+++ b/Engine/Plugins/AI/MassAI/Source/MassAIDebug/Private/GameplayDebuggerCategory_Mass.cpp
@@ -565,6 +565,7 @@ void FGameplayDebuggerCategory_Mass::CollectData(APlayerController* OwnerPC, AAc
 				const bool bHasStateTree = (StateTreeInstanceList.Num() > 0);
 				const bool bHasSteering = (SteeringList.Num() > 0);
 				const bool bHasVelocity = (VelocityList.Num() > 0);
+				const bool bHasNavMeshShortPaths = (NavMeshShortPathList.Num() > 0);
 
 				const UGameplayDebuggerUserSettings* Settings = GetDefault<UGameplayDebuggerUserSettings>();
 				const FVector::FReal MaxViewDistance = Settings->MaxViewDistance;
@@ -714,7 +715,7 @@ void FGameplayDebuggerCategory_Mass::CollectData(APlayerController* OwnerPC, AAc
 							}
 						}
 
-						if (!NavMeshShortPathList.IsEmpty())
+						if (bHasNavMeshShortPaths)
 						{
 							const FMassNavMeshShortPathFragment& ShortPath = NavMeshShortPathList[EntityIndex];
 
@@ -829,6 +830,21 @@ void FGameplayDebuggerCategory_Mass::CollectData(APlayerController* OwnerPC, AAc
 									*UEnum::GetDisplayValueAsText(MoveTarget.IntentAtGoal).ToString(), MoveTarget.DistanceToGoal);
 							}
 
+							if (bShowNearEntityPath && bHasNavMeshShortPaths)
+							{
+								// Display more movement info
+								const FMassNavMeshShortPathFragment ShortPath = NavMeshShortPathList[EntityIndex];
+								if (ShortPath.bInitialized)
+								{
+									Status += FString::Printf(TEXT("ShortPath: %i pts, progress: %0.f\n%s (%s) %s\n"),
+										ShortPath.NumPoints,
+										ShortPath.MoveTargetProgressDistance,
+										ShortPath.bDone ? TEXT("{green}done{yellow}") : TEXT("{yellow}in progress"),
+										ShortPath.bPartialResult ? TEXT("{yellow}partial") : TEXT("{yellow}final"),
+										ShortPath.bDone && !ShortPath.bPartialResult ? TEXT("{green}PATH COMPLETED") : TEXT(""));
+								}
+							}
+
 							// LookAt
 							if (bShowEntityLookAt && bHasLookAt)
 							{
@@ -861,12 +877,14 @@ void FGameplayDebuggerCategory_Mass::CollectData(APlayerController* OwnerPC, AAc
 			FMassEntityQuery EntityColliderQuery(EntityManager.AsShared());
 			EntityColliderQuery.AddRequirement<FMassAvoidanceColliderFragment>(EMassFragmentAccess::ReadOnly);
 			EntityColliderQuery.AddRequirement<FTransformFragment>(EMassFragmentAccess::ReadOnly);
+			EntityColliderQuery.AddRequirement<FAgentRadiusFragment>(EMassFragmentAccess::ReadOnly);
 			FMassExecutionContext Context(EntityManager, 0.f);
 			EntityColliderQuery.ForEachEntityChunk(Context, [this, ViewLocation, ViewDirection](const FMassExecutionContext& Context)
 			{
 				const int32 NumEntities = Context.GetNumEntities();
 				const TConstArrayView<FTransformFragment> TransformList = Context.GetFragmentView<FTransformFragment>();
 				const TConstArrayView<FMassAvoidanceColliderFragment> CollidersList = Context.GetFragmentView<FMassAvoidanceColliderFragment>();
+				const TConstArrayView<FAgentRadiusFragment> RadiiList = Context.GetFragmentView<FAgentRadiusFragment>();
 
 				for (int32 EntityIndex = 0; EntityIndex < NumEntities; ++EntityIndex)
 				{
@@ -894,6 +912,12 @@ void FGameplayDebuggerCategory_Mass::CollectData(APlayerController* OwnerPC, AAc
 						AddShape(FGameplayDebuggerShape::MakeCircle(BasePos + Pill.HalfLength * EntityForward, FVector::UpVector, Pill.Radius, FColor::Blue));
 						AddShape(FGameplayDebuggerShape::MakeCircle(BasePos - Pill.HalfLength * EntityForward, FVector::UpVector, Pill.Radius, FColor::Blue));
 					}
+					else
+					{
+						// Fallback on radius
+						const FAgentRadiusFragment& RadiusFragment = RadiiList[EntityIndex];
+						AddShape(FGameplayDebuggerShape::MakeCircle(BasePos + FVector(0,0,2), FVector::UpVector, RadiusFragment.Radius, FColor::Orange));
+					}
 				}
 			});
 		}
diff --git a/Engine/Plugins/AI/MassAI/Source/MassNavMeshNavigation/Private/MassNavMeshNavigationFragments.cpp b/Engine/Plugins/AI/MassAI/Source/MassNavMeshNavigation/Private/MassNavMeshNavigationFragments.cpp
--- a/Engine/Plugins/AI/MassAI/Source/MassNavMeshNavigation/Private/MassNavMeshNavigationFragments.cpp
+++ b/Engine/Plugins/AI/MassAI/Source/MassNavMeshNavigation/Private/MassNavMeshNavigationFragments.cpp
@@ -67,5 +67,7 @@ bool FMassNavMeshShortPathFragment::RequestShortPath(const TSharedPtr<FNavCorrid
 	// Update ProgressDistance to account for leading points. 
 	MoveTargetProgressDistance = Points[InNumLeadingPoints].Distance.Get();
 	
+	bInitialized = true;
+	
 	return true;
 }
diff --git a/Engine/Plugins/AI/MassAI/Source/MassNavMeshNavigation/Public/MassNavMeshNavigationFragments.h b/Engine/Plugins/AI/MassAI/Source/MassNavMeshNavigation/Public/MassNavMeshNavigationFragments.h
--- a/Engine/Plugins/AI/MassAI/Source/MassNavMeshNavigation/Public/MassNavMeshNavigationFragments.h
+++ b/Engine/Plugins/AI/MassAI/Source/MassNavMeshNavigation/Public/MassNavMeshNavigationFragments.h
@@ -55,6 +55,7 @@ struct MASSNAVMESHNAVIGATION_API FMassNavMeshShortPathFragment : public FMassFra
 		EndOfPathIntent = EMassMovementAction::Stand;
 		bPartialResult = false;
 		bDone = false;
+		bInitialized = false;
 	}
 
 	/** Returns true if we are done progressing on the short path. */
@@ -86,6 +87,9 @@ struct MASSNAVMESHNAVIGATION_API FMassNavMeshShortPathFragment : public FMassFra
 
 	/** True when path follow is completed. */
 	uint8 bDone : 1 = false;
+
+	/** True when the path has been initalized. */
+	uint8 bInitialized : 1 = false;
 };
 
 /** Current navmesh path */
