[MassLookAt] updated CreateLookAtActorRequest to fallback on a static target location if the target actor doesn't have an associated mass entity.

diff --git a/Engine/Plugins/AI/MassAI/Source/MassAIBehavior/Private/MassLookAtSubsystem.cpp b/Engine/Plugins/AI/MassAI/Source/MassAIBehavior/Private/MassLookAtSubsystem.cpp
--- a/Engine/Plugins/AI/MassAI/Source/MassAIBehavior/Private/MassLookAtSubsystem.cpp
+++ b/Engine/Plugins/AI/MassAI/Source/MassAIBehavior/Private/MassLookAtSubsystem.cpp
@@ -301,6 +301,12 @@ FMassLookAtRequestHandle UMassLookAtSubsystem::CreateLookAtActorRequest(
 {
 	const UWorld* World = GetWorld();
 
+	if (TargetActor == nullptr)
+	{
+		UE_VLOG_UELOG(this, LogMassBehavior, Log, TEXT("%hs failed: invalid target actor"), __FUNCTION__);
+		return {};
+	}
+
 	UMassActorSubsystem* ActorSubsystem = World->GetSubsystem<UMassActorSubsystem>();
 	if (ActorSubsystem == nullptr)
 	{
@@ -308,6 +314,13 @@ FMassLookAtRequestHandle UMassLookAtSubsystem::CreateLookAtActorRequest(
 		return {};
 	}
 
+	const FMassEntityHandle TargetEntity = ActorSubsystem->GetEntityHandleFromActor(TargetActor);
+	if (!TargetEntity.IsSet())
+	{
+		UE_VLOG_UELOG(this, LogMassBehavior, Log, TEXT("%hs: using static target location since no MassEntity is associated to '%s'"), __FUNCTION__, *GetNameSafe(TargetActor));
+		return CreateLookAtPositionRequest(ViewerActor, Priority, TargetActor->GetActorLocation(), InterpolationSpeed, CustomInterpolationSpeed);
+	}
+
 	UMassEntitySubsystem* EntitySubsystem = World->GetSubsystem<UMassEntitySubsystem>();
 	if (EntitySubsystem == nullptr)
 	{
@@ -322,13 +335,6 @@ FMassLookAtRequestHandle UMassLookAtSubsystem::CreateLookAtActorRequest(
 		return {};
 	}
 
-	const FMassEntityHandle TargetEntity = ActorSubsystem->GetEntityHandleFromActor(TargetActor);
-	if (!TargetEntity.IsSet())
-	{
-		UE_VLOG_UELOG(this, LogMassBehavior, Error, TEXT("%hs failed: unable to find a MassEntity associated to '%s'"), __FUNCTION__, *GetNameSafe(TargetActor));
-		return {};
-	}
-
 	// @todo: once available, consider using incoming EntityBuilder improvements for this whole creation step
 	FMassEntityManager& EntityManager = EntitySubsystem->GetMutableEntityManager();
 	FMassEntityHandle RequestEntity = EntityManager.ReserveEntity();
diff --git a/Engine/Plugins/AI/MassAI/Source/MassAIBehavior/Public/MassLookAtSubsystem.h b/Engine/Plugins/AI/MassAI/Source/MassAIBehavior/Public/MassLookAtSubsystem.h
--- a/Engine/Plugins/AI/MassAI/Source/MassAIBehavior/Public/MassLookAtSubsystem.h
+++ b/Engine/Plugins/AI/MassAI/Source/MassAIBehavior/Public/MassLookAtSubsystem.h
@@ -42,6 +42,11 @@ public:
 
 	/**
 	 * Creates a new LookAt request using the provided parameters for the mass entity associated to 'ViewerActor', if any.
+	 * @param ViewerActor Actor associated to the mass entity that needs to perform the LookAt
+	 * @param Priority Priority assigned to the request that will influence selected request when multiple requests are sent to the same entity
+	 * @param TargetLocation Static location used as target
+	 * @param InterpolationSpeed Optional parameter to use predefined interpolation speed
+	 * @param CustomInterpolationSpeed Optional parameter to specify a custom interpolation speed when 'InterpolationSpeed' is set to EMassLookAtInterpolationSpeed::Custom
 	 * Method will log an error if it fails (e.g., no associated mass entity).
 	 */
 	UFUNCTION(BlueprintCallable, BlueprintPure = false, Category = "LookAt", meta = (ReturnDisplayName="Request Handle"))
@@ -53,7 +58,13 @@ public:
 
 	/**
 	 * Creates a new LookAt request using the provided parameters for the mass entity associated to 'ViewerActor', if any.
-	 * Method will log an error if it fails (e.g., no associated mass entity for viewer or target actors).
+	 * @param ViewerActor Actor associated to the mass entity that needs to perform the LookAt
+	 * @param Priority Priority assigned to the request that will influence selected request when multiple requests are sent to the same entity
+	 * @param TargetActor Actor used as reference for the target location. If the actor has an associated mass entity then the target
+	 * location will be updated during the request lifetime. Otherwise, the actor location will be used to create a LookAt using a static target location
+	 * @param InterpolationSpeed Optional parameter to use predefined interpolation speed
+	 * @param CustomInterpolationSpeed Optional parameter to specify a custom interpolation speed when 'InterpolationSpeed' is set to EMassLookAtInterpolationSpeed::Custom
+	 * Method will log an error if it fails (e.g., no associated mass entity for viewer)
 	 */
 	UFUNCTION(BlueprintCallable, BlueprintPure = false, Category = "LookAt", meta = (ReturnDisplayName="Request Handle"))
 	FMassLookAtRequestHandle CreateLookAtActorRequest(AActor* ViewerActor
