[MassLookAt] filter out targets outside the query box since hashgrid might return intersecting entities that we don't want to consider

diff --git a/Engine/Plugins/AI/MassAI/Source/MassAIBehavior/Private/MassLookAtProcessors.cpp b/Engine/Plugins/AI/MassAI/Source/MassAIBehavior/Private/MassLookAtProcessors.cpp
--- a/Engine/Plugins/AI/MassAI/Source/MassAIBehavior/Private/MassLookAtProcessors.cpp
+++ b/Engine/Plugins/AI/MassAI/Source/MassAIBehavior/Private/MassLookAtProcessors.cpp
@@ -317,6 +317,11 @@ void UMassLookAtProcessor::FindNewGazeTarget(const UMassNavigationSubsystem& Mas
 				const FMassLookAtTargetFragment& TargetFragment = EntityManager.GetFragmentDataChecked<FMassLookAtTargetFragment>(CandidateEntity);
 
 				const FVector TargetLocation = TargetTransform.GetTransform().TransformPosition(TargetFragment.Offset);
+				if (!QueryBox.IsInside(TargetLocation))
+				{
+					continue;
+				}
+
 				FVector Direction = (TargetLocation - Location).GetSafeNormal();
 				Direction = Transform.InverseTransformVector(Direction);
 
