[MassLookAt] only delete target entity when it is a temporary entity created for the request

diff --git a/Engine/Plugins/AI/MassAI/Source/MassAIBehavior/Private/MassLookAtSubsystem.cpp b/Engine/Plugins/AI/MassAI/Source/MassAIBehavior/Private/MassLookAtSubsystem.cpp
--- a/Engine/Plugins/AI/MassAI/Source/MassAIBehavior/Private/MassLookAtSubsystem.cpp
+++ b/Engine/Plugins/AI/MassAI/Source/MassAIBehavior/Private/MassLookAtSubsystem.cpp
@@ -289,7 +289,7 @@ FMassLookAtRequestHandle UMassLookAtSubsystem::CreateLookAtPositionRequest(
 
 	UE::Mass::LookAt::Private::FlushCommands(EntityManager, CommandBuffer);
 
-	return FMassLookAtRequestHandle{RequestEntity, TargetEntity};
+	return FMassLookAtRequestHandle{RequestEntity, TargetEntity, /*bTargetEntityOwnedByRequest*/true};
 }
 
 FMassLookAtRequestHandle UMassLookAtSubsystem::CreateLookAtActorRequest(
@@ -366,7 +366,7 @@ FMassLookAtRequestHandle UMassLookAtSubsystem::CreateLookAtActorRequest(
 
 	UE::Mass::LookAt::Private::FlushCommands(EntityManager, CommandBuffer);
 
-	return FMassLookAtRequestHandle{RequestEntity, TargetEntity};
+	return FMassLookAtRequestHandle{RequestEntity, TargetEntity, /*bTargetEntityOwnedByRequest*/false};
 }
 
 void UMassLookAtSubsystem::DeleteRequest(const FMassLookAtRequestHandle RequestHandle) const
@@ -397,7 +397,14 @@ void UMassLookAtSubsystem::DeleteRequest(const FMassLookAtRequestHandle RequestH
 	const TSharedPtr<FMassCommandBuffer> CommandBuffer = MakeShareable(new FMassCommandBuffer());
 	CommandBuffer->PushCommand<FMassDeferredDestroyCommand>([RequestHandle](FMassEntityManager& InEntityManager)
 		{
-			InEntityManager.BatchDestroyEntities({ RequestHandle.Request, RequestHandle.Target });
+			if (RequestHandle.bTargetEntityOwnedByRequest)
+			{
+				InEntityManager.BatchDestroyEntities({ RequestHandle.Request, RequestHandle.Target });
+			}
+			else
+			{
+				InEntityManager.DestroyEntity(RequestHandle.Request);
+			}
 		});
 
 	UE::Mass::LookAt::Private::FlushCommands(EntitySubsystem->GetMutableEntityManager(), CommandBuffer);
diff --git a/Engine/Plugins/AI/MassAI/Source/MassAIBehavior/Public/MassLookAtSubsystem.h b/Engine/Plugins/AI/MassAI/Source/MassAIBehavior/Public/MassLookAtSubsystem.h
--- a/Engine/Plugins/AI/MassAI/Source/MassAIBehavior/Public/MassLookAtSubsystem.h
+++ b/Engine/Plugins/AI/MassAI/Source/MassAIBehavior/Public/MassLookAtSubsystem.h
@@ -15,11 +15,24 @@ struct FMassLookAtRequestHandle
 {
 	GENERATED_BODY()
 
+	FMassLookAtRequestHandle() = default;
+
+	FMassLookAtRequestHandle(const FMassEntityHandle& Request, const FMassEntityHandle& Target, const bool bTargetEntityOwnedByRequest)
+		: Request(Request)
+		, Target(Target)
+		, bTargetEntityOwnedByRequest(bTargetEntityOwnedByRequest)
+	{
+	}
+
 	UPROPERTY()
 	FMassEntityHandle Request;
 
 	UPROPERTY()
 	FMassEntityHandle Target;
+
+	/** Indicates if the target entity was created for the request and must be destroyed when deleting the request (e.g. CreateLookAtPositionRequest) */
+	UPROPERTY()
+	bool bTargetEntityOwnedByRequest = false;
 };
 
 /**
