Addition of path boundaries to mass gameplay debugger
Additional logs in navigation flow

diff --git a/Engine/Plugins/AI/MassAI/Source/MassAIBehavior/Private/Tasks/MassNavMeshPathFollowTask.cpp b/Engine/Plugins/AI/MassAI/Source/MassAIBehavior/Private/Tasks/MassNavMeshPathFollowTask.cpp
--- a/Engine/Plugins/AI/MassAI/Source/MassAIBehavior/Private/Tasks/MassNavMeshPathFollowTask.cpp
+++ b/Engine/Plugins/AI/MassAI/Source/MassAIBehavior/Private/Tasks/MassNavMeshPathFollowTask.cpp
@@ -80,6 +80,8 @@ bool FMassNavMeshPathFollowTask::RequestPath(FStateTreeExecutionContext& Context
 			FMassNavMeshCachedPathFragment& CachedPathFragment = Context.GetExternalData(CachedPathHandle);
 			CachedPathFragment.NavPath = Result.Path;
 
+			CachedPathFragment.PathSource = EMassNavigationPathSource::NavMesh;
+
 			// Build corridor
 			CachedPathFragment.Corridor = MakeShared<FNavCorridor>();
 			const FSharedConstNavQueryFilter NavQueryFilter = Query.QueryFilter ? Query.QueryFilter : NavData->GetDefaultQueryFilter();
@@ -88,6 +90,12 @@ bool FMassNavMeshPathFollowTask::RequestPath(FStateTreeExecutionContext& Context
 			CorridorParams.SetFromWidth(InstanceData.CorridorWidth);
 			CorridorParams.PathOffsetFromBoundaries = InstanceData.OffsetFromBoundaries;
 
+			MASSBEHAVIOR_LOG(Verbose, TEXT("FMassNavMeshPathFollowTask::RequestPath"));
+			MASSBEHAVIOR_LOG(Verbose, TEXT("Start: %s, End: %s"),
+				*CachedPathFragment.NavPath->GetPathPoints()[0].Location.ToCompactString(),
+				*CachedPathFragment.NavPath->GetPathPoints().Last().Location.ToCompactString());
+			MASSBEHAVIOR_LOG(Verbose, TEXT("Corridor params - %s"), *CorridorParams.ToString());
+			
 			CachedPathFragment.Corridor->BuildFromPath(*CachedPathFragment.NavPath, NavQueryFilter, CorridorParams);
 
 			// Update short path
diff --git a/Engine/Plugins/AI/MassAI/Source/MassAIDebug/Private/GameplayDebuggerCategory_Mass.cpp b/Engine/Plugins/AI/MassAI/Source/MassAIDebug/Private/GameplayDebuggerCategory_Mass.cpp
--- a/Engine/Plugins/AI/MassAI/Source/MassAIDebug/Private/GameplayDebuggerCategory_Mass.cpp
+++ b/Engine/Plugins/AI/MassAI/Source/MassAIDebug/Private/GameplayDebuggerCategory_Mass.cpp
@@ -515,6 +515,7 @@ void FGameplayDebuggerCategory_Mass::CollectData(APlayerController* OwnerPC, AAc
 		EntityQuery.AddRequirement<FMassStandingSteeringFragment>(EMassFragmentAccess::ReadOnly, EMassFragmentPresence::Optional);
 		EntityQuery.AddRequirement<FMassGhostLocationFragment>(EMassFragmentAccess::ReadOnly, EMassFragmentPresence::Optional);
 		EntityQuery.AddRequirement<FMassVelocityFragment>(EMassFragmentAccess::ReadOnly, EMassFragmentPresence::Optional);
+		EntityQuery.AddRequirement<FMassDesiredMovementFragment>(EMassFragmentAccess::ReadOnly, EMassFragmentPresence::Optional);
 		EntityQuery.AddRequirement<FMassForceFragment>(EMassFragmentAccess::ReadOnly, EMassFragmentPresence::Optional);
 		EntityQuery.AddRequirement<FMassMoveTargetFragment>(EMassFragmentAccess::ReadOnly, EMassFragmentPresence::Optional);
 		EntityQuery.AddRequirement<FMassLookAtFragment>(EMassFragmentAccess::ReadOnly, EMassFragmentPresence::Optional);
@@ -522,6 +523,7 @@ void FGameplayDebuggerCategory_Mass::CollectData(APlayerController* OwnerPC, AAc
 		EntityQuery.AddRequirement<FMassSimulationLODFragment>(EMassFragmentAccess::ReadOnly, EMassFragmentPresence::Optional);
 		EntityQuery.AddRequirement<FMassZoneGraphShortPathFragment>(EMassFragmentAccess::ReadOnly, EMassFragmentPresence::Optional);
 		EntityQuery.AddRequirement<FMassNavMeshShortPathFragment>(EMassFragmentAccess::ReadOnly, EMassFragmentPresence::Optional);
+		EntityQuery.AddRequirement<FMassNavMeshCachedPathFragment>(EMassFragmentAccess::ReadOnly, EMassFragmentPresence::Optional);
 		EntityQuery.AddRequirement<FMassSmartObjectUserFragment>(EMassFragmentAccess::ReadOnly, EMassFragmentPresence::Optional);
 
 		const double CurrentTime = World->GetTimeSeconds();
@@ -545,6 +547,7 @@ void FGameplayDebuggerCategory_Mass::CollectData(APlayerController* OwnerPC, AAc
 				const TConstArrayView<FMassStandingSteeringFragment> StandingSteeringList = Context.GetFragmentView<FMassStandingSteeringFragment>();
 				const TConstArrayView<FMassGhostLocationFragment> GhostList = Context.GetFragmentView<FMassGhostLocationFragment>();
 				const TConstArrayView<FMassVelocityFragment> VelocityList = Context.GetFragmentView<FMassVelocityFragment>();
+				const TConstArrayView<FMassDesiredMovementFragment> DesiredMovementList = Context.GetFragmentView<FMassDesiredMovementFragment>();
 				const TConstArrayView<FMassForceFragment> ForceList = Context.GetFragmentView<FMassForceFragment>();
 				const TConstArrayView<FMassMoveTargetFragment> MoveTargetList = Context.GetFragmentView<FMassMoveTargetFragment>();
 				const TConstArrayView<FMassLookAtFragment> LookAtList = Context.GetFragmentView<FMassLookAtFragment>();
@@ -552,6 +555,7 @@ void FGameplayDebuggerCategory_Mass::CollectData(APlayerController* OwnerPC, AAc
 				const TConstArrayView<FMassSimulationLODFragment> SimLODList = Context.GetFragmentView<FMassSimulationLODFragment>();
 				const TConstArrayView<FMassZoneGraphShortPathFragment> ZoneGraphShortPathList = Context.GetFragmentView<FMassZoneGraphShortPathFragment>();
 				const TConstArrayView<FMassNavMeshShortPathFragment> NavMeshShortPathList = Context.GetFragmentView<FMassNavMeshShortPathFragment>();
+				const TConstArrayView<FMassNavMeshCachedPathFragment> NavMeshPathList = Context.GetFragmentView<FMassNavMeshCachedPathFragment>();
 				const TConstArrayView<FMassSmartObjectUserFragment> SOUserList = Context.GetFragmentView<FMassSmartObjectUserFragment>();
 
 				const bool bHasForce = (ForceList.Num() > 0);
@@ -565,7 +569,9 @@ void FGameplayDebuggerCategory_Mass::CollectData(APlayerController* OwnerPC, AAc
 				const bool bHasStateTree = (StateTreeInstanceList.Num() > 0);
 				const bool bHasSteering = (SteeringList.Num() > 0);
 				const bool bHasVelocity = (VelocityList.Num() > 0);
+				const bool bHasDesiredMovement = (DesiredMovementList.Num() > 0);
 				const bool bHasNavMeshShortPaths = (NavMeshShortPathList.Num() > 0);
+				const bool bHasNavMeshPaths = (NavMeshPathList.Num() > 0);
 
 				const UGameplayDebuggerUserSettings* Settings = GetDefault<UGameplayDebuggerUserSettings>();
 				const FVector::FReal MaxViewDistance = Settings->MaxViewDistance;
@@ -631,17 +637,17 @@ void FGameplayDebuggerCategory_Mass::CollectData(APlayerController* OwnerPC, AAc
 						constexpr float LookArrowHeadSize = 20.f;
 						constexpr float LookArrowThickness = 2.f;
 
-						BasePos = EntityLocation;
+						FVector LookAtBasePos = EntityLocation;
 
 						// Apply the offset for the entity looking at something
 						if (LookAtTargetList.Num())
 						{
-							BasePos += LookAtTargetList[EntityIndex].Offset;
+							LookAtBasePos += LookAtTargetList[EntityIndex].Offset;
 						}
 						else
 						{
 							const FVector TargetOffset = GetDefault<UMassLookAtSettings>()->GetDefaultTargetLocationOffset();
-							BasePos += TargetOffset.IsZero() ? FVector(0, 0, 160) : TargetOffset;
+							LookAtBasePos += TargetOffset.IsZero() ? FVector(0, 0, 160) : TargetOffset;
 						}
 
 						const FMassLookAtFragment& LookAt = LookAtList[EntityIndex];
@@ -652,10 +658,10 @@ void FGameplayDebuggerCategory_Mass::CollectData(APlayerController* OwnerPC, AAc
 							// Target location is expected to include its own offset
 							FVector TargetPosition = LookAt.MainTargetLocation;
 							AddShape(FGameplayDebuggerShape::MakeCircle(TargetPosition, FVector::UpVector, bHasRadius ? RadiusList[EntityIndex].Radius : DebugShapeRadius, FColor::Red));
-							AddShape(FGameplayDebuggerShape::MakeArrow(BasePos, TargetPosition, TargetArrowHeadSize, TargetArrowThickness, FColor::Red));
+							AddShape(FGameplayDebuggerShape::MakeArrow(LookAtBasePos, TargetPosition, TargetArrowHeadSize, TargetArrowThickness, FColor::Red));
 
-							const FVector::FReal TargetDistance = FMath::Max(LookArrowLength, FVector::DotProduct(WorldLookDirection, TargetPosition - BasePos));
-							AddShape(FGameplayDebuggerShape::MakeArrow(BasePos, BasePos + WorldLookDirection * TargetDistance, LookArrowHeadSize, LookArrowThickness, FColorList::LightGrey));
+							const FVector::FReal TargetDistance = FMath::Max(LookArrowLength, FVector::DotProduct(WorldLookDirection, TargetPosition - LookAtBasePos));
+							AddShape(FGameplayDebuggerShape::MakeArrow(LookAtBasePos, LookAtBasePos + WorldLookDirection * TargetDistance, LookArrowHeadSize, LookArrowThickness, FColorList::LightGrey));
 							bLookArrowDrawn = true;
 						}
 
@@ -664,12 +670,12 @@ void FGameplayDebuggerCategory_Mass::CollectData(APlayerController* OwnerPC, AAc
 							// Target location is expected to include its own offset
 							FVector TargetPosition = LookAt.GazeTargetLocation;
 							AddShape(FGameplayDebuggerShape::MakeCircle(TargetPosition, FVector::UpVector, bHasRadius ? RadiusList[EntityIndex].Radius : DebugShapeRadius, FColorList::Goldenrod));
-							AddShape(FGameplayDebuggerShape::MakeArrow(BasePos, TargetPosition, TargetArrowHeadSize, TargetArrowThickness, FColorList::Goldenrod));
+							AddShape(FGameplayDebuggerShape::MakeArrow(LookAtBasePos, TargetPosition, TargetArrowHeadSize, TargetArrowThickness, FColorList::Goldenrod));
 						}
 
 						if (!bLookArrowDrawn)
 						{
-							AddShape(FGameplayDebuggerShape::MakeArrow(BasePos, BasePos + WorldLookDirection * LookArrowLength, LookArrowHeadSize, LookArrowThickness, FColor::Turquoise));
+							AddShape(FGameplayDebuggerShape::MakeArrow(LookAtBasePos, LookAtBasePos + WorldLookDirection * LookArrowLength, LookArrowHeadSize, LookArrowThickness, FColor::Turquoise));
 						}
 					}
 
@@ -685,6 +691,14 @@ void FGameplayDebuggerCategory_Mass::CollectData(APlayerController* OwnerPC, AAc
 						}
 					}
 
+					// Desired velocity
+					if (bHasDesiredMovement)
+					{
+						BasePos += ZDeltaOffset;
+						const FVector& DesiredVelocity = DesiredMovementList[EntityIndex].DesiredVelocity;
+						AddShape(FGameplayDebuggerShape::MakeArrow(BasePos, BasePos + DesiredVelocity, 10.0f, 2.0f, FColor::Blue));
+					}
+					
 					// Path
 					if (bShowNearEntityPath && bHasMoveTarget && bHasRadius)
 					{
@@ -715,15 +729,24 @@ void FGameplayDebuggerCategory_Mass::CollectData(APlayerController* OwnerPC, AAc
 							}
 						}
 
-						if (bHasNavMeshShortPaths)
+						if (bHasNavMeshShortPaths && Entity == CachedEntity)
 						{
 							const FMassNavMeshShortPathFragment& ShortPath = NavMeshShortPathList[EntityIndex];
 
+							FColor CorridorColor = FColorList::Yellow;
+							if (bHasNavMeshPaths)
+							{
+								const FMassNavMeshCachedPathFragment& Path = NavMeshPathList[EntityIndex];
+								CorridorColor = Path.PathSource == EMassNavigationPathSource::NavMesh ? FColorList::Cyan : (Path.PathSource == EMassNavigationPathSource::Spline) ? FColorList::Blue : FColorList::Black;
+							}
+
 							for (uint8 PointIndex = 0; PointIndex < ShortPath.NumPoints - 1; PointIndex++)
 							{
 								const FMassNavMeshPathPoint& CurrPoint = ShortPath.Points[PointIndex];
 								const FMassNavMeshPathPoint& NextPoint = ShortPath.Points[PointIndex + 1];
 								AddShape(FGameplayDebuggerShape::MakeSegment(CurrPoint.Position + ZBaseOffset, NextPoint.Position + ZBaseOffset, /*Thickness*/3.f, FColorList::GreenYellow));
+								AddShape(FGameplayDebuggerShape::MakeSegment(CurrPoint.Left + ZBaseOffset, NextPoint.Left + ZBaseOffset, /*Thickness*/2.f, CorridorColor));
+								AddShape(FGameplayDebuggerShape::MakeSegment(CurrPoint.Right + ZBaseOffset, NextPoint.Right + ZBaseOffset, /*Thickness*/2.f, CorridorColor));
 							}
 					
 							for (uint8 PointIndex = 0; PointIndex < ShortPath.NumPoints; PointIndex++)
diff --git a/Engine/Plugins/AI/MassAI/Source/MassNavMeshNavigation/Public/MassNavMeshNavigationFragments.h b/Engine/Plugins/AI/MassAI/Source/MassNavMeshNavigation/Public/MassNavMeshNavigationFragments.h
--- a/Engine/Plugins/AI/MassAI/Source/MassNavMeshNavigation/Public/MassNavMeshNavigationFragments.h
+++ b/Engine/Plugins/AI/MassAI/Source/MassNavMeshNavigation/Public/MassNavMeshNavigationFragments.h
@@ -94,6 +94,13 @@ struct FMassNavMeshShortPathFragment : public FMassFragment
 	uint8 bInitialized : 1 = false;
 };
 
+enum class EMassNavigationPathSource : uint8
+{
+	NavMesh,
+	Spline,
+	Unset
+};
+
 /** Current navmesh path */
 USTRUCT()
 struct FMassNavMeshCachedPathFragment : public FMassFragment
@@ -109,6 +116,9 @@ struct FMassNavMeshCachedPathFragment : public FMassFragment
 	/** Index used to keep track of progression on the navmesh path. */
 	uint16 NavPathNextStartIndex = 0;
 
+	/** Source of the path, navmesh or else. */
+	EMassNavigationPathSource PathSource = EMassNavigationPathSource::Unset;
+
 	/** Number of points before the starting point. */
 	static constexpr uint8 NumLeadingPoints = 1;
 };
