Update NPC desired movement based on current velocity when in animate mode for root motion driven NPCs.

#tests devtest Nikodem

diff --git a/Engine/Plugins/AI/MassAI/Source/MassNavigation/Private/MassSteeringProcessors.cpp b/Engine/Plugins/AI/MassAI/Source/MassNavigation/Private/MassSteeringProcessors.cpp
--- a/Engine/Plugins/AI/MassAI/Source/MassNavigation/Private/MassSteeringProcessors.cpp
+++ b/Engine/Plugins/AI/MassAI/Source/MassNavigation/Private/MassSteeringProcessors.cpp
@@ -64,6 +64,7 @@ UMassSteerToMoveTargetProcessor::UMassSteerToMoveTargetProcessor()
 void UMassSteerToMoveTargetProcessor::ConfigureQueries(const TSharedRef<FMassEntityManager>& EntityManager)
 {
 	EntityQuery.AddRequirement<FTransformFragment>(EMassFragmentAccess::ReadOnly);
+	EntityQuery.AddRequirement<FMassVelocityFragment>(EMassFragmentAccess::ReadOnly);
 	EntityQuery.AddRequirement<FMassMoveTargetFragment>(EMassFragmentAccess::ReadWrite);
 	EntityQuery.AddRequirement<FMassSteeringFragment>(EMassFragmentAccess::ReadWrite);
 	EntityQuery.AddRequirement<FMassStandingSteeringFragment>(EMassFragmentAccess::ReadWrite);
@@ -80,6 +81,7 @@ void UMassSteerToMoveTargetProcessor::ConfigureQueries(const TSharedRef<FMassEnt
 
 	// No need for Off LOD to do steering, applying move target directly
 	EntityQuery.AddTagRequirement<FMassOffLODTag>(EMassFragmentPresence::None);
+	EntityQuery.AddTagRequirement<FMassCodeDrivenMovementTag>(EMassFragmentPresence::Optional);
 }
 
 void UMassSteerToMoveTargetProcessor::Execute(FMassEntityManager& EntityManager, FMassExecutionContext& Context)
@@ -90,6 +92,7 @@ void UMassSteerToMoveTargetProcessor::Execute(FMassEntityManager& EntityManager,
 	EntityQuery.ForEachEntityChunk(Context, [this, World](FMassExecutionContext& Context)
 	{
 		const TConstArrayView<FTransformFragment> TransformList = Context.GetFragmentView<FTransformFragment>();
+		const TConstArrayView<FMassVelocityFragment> VelocityList = Context.GetFragmentView<FMassVelocityFragment>();
 		const TArrayView<FMassMoveTargetFragment> MoveTargetList = Context.GetMutableFragmentView<FMassMoveTargetFragment>();
 		const TArrayView<FMassDesiredMovementFragment> MovementList = Context.GetMutableFragmentView<FMassDesiredMovementFragment>();
 		const TArrayView<FMassForceFragment> ForceList = Context.GetMutableFragmentView<FMassForceFragment>();
@@ -99,6 +102,7 @@ void UMassSteerToMoveTargetProcessor::Execute(FMassEntityManager& EntityManager,
 		const FMassMovementParameters& MovementParams = Context.GetConstSharedFragment<FMassMovementParameters>();
 		const FMassMovingSteeringParameters& MovingSteeringParams = Context.GetConstSharedFragment<FMassMovingSteeringParameters>();
 		const FMassStandingSteeringParameters& StandingSteeringParams = Context.GetConstSharedFragment<FMassStandingSteeringParameters>();
+		bool bIsCodeDriven = Context.DoesArchetypeHaveTag<FMassCodeDrivenMovementTag>();
 
 		const FVector::FReal SteerK = 1. / MovingSteeringParams.ReactionTime;
 		const float DeltaTime = Context.GetDeltaTimeSeconds();
@@ -106,6 +110,7 @@ void UMassSteerToMoveTargetProcessor::Execute(FMassEntityManager& EntityManager,
 		for (FMassExecutionContext::FEntityIterator EntityIt = Context.CreateEntityIterator(); EntityIt; ++EntityIt)
 		{
 			const FTransformFragment& TransformFragment = TransformList[EntityIt];
+			const FMassVelocityFragment& VelocityFragment = VelocityList[EntityIt];
 			FMassSteeringFragment& Steering = SteeringList[EntityIt];
 			FMassStandingSteeringFragment& StandingSteering = StandingSteeringList[EntityIt];
 			FMassGhostLocationFragment& Ghost = GhostList[EntityIt];
@@ -330,11 +335,25 @@ void UMassSteerToMoveTargetProcessor::Execute(FMassEntityManager& EntityManager,
 			}
 			else if (MoveTarget.GetCurrentAction() == EMassMovementAction::Animate)
 			{
-				// Stop all movement when animating.
 				Steering.Reset();
 				MoveTarget.bSteeringFallingBehind = false;
 				Force.Value = FVector::ZeroVector;
-				DesiredMovement.DesiredVelocity = FVector::ZeroVector;
+
+				if (bIsCodeDriven)
+				{
+					// Stop all movement when animating.
+					DesiredMovement.DesiredVelocity = FVector::ZeroVector;
+				}
+				else
+				{
+					// Animation driven, sync desired movement to animation
+					Steering.DesiredVelocity = VelocityFragment.Value;
+					DesiredMovement.DesiredVelocity = VelocityFragment.Value;
+					if (!DesiredMovement.DesiredVelocity.IsNearlyZero())
+					{
+						DesiredMovement.DesiredFacing = FQuat(DesiredMovement.DesiredVelocity.Rotation());
+					}
+				}
 			}
 
 #if WITH_MASSGAMEPLAY_DEBUG
diff --git a/Engine/Plugins/Runtime/MassGameplay/Source/MassMovement/Private/MassMovementProcessors.cpp b/Engine/Plugins/Runtime/MassGameplay/Source/MassMovement/Private/MassMovementProcessors.cpp
--- a/Engine/Plugins/Runtime/MassGameplay/Source/MassMovement/Private/MassMovementProcessors.cpp
+++ b/Engine/Plugins/Runtime/MassGameplay/Source/MassMovement/Private/MassMovementProcessors.cpp
@@ -59,6 +59,9 @@ void UMassApplyForceProcessor::Execute(FMassEntityManager& EntityManager, FMassE
 
 UMassApplyMovementProcessor::UMassApplyMovementProcessor()
 	: EntityQuery(*this)
+#if WITH_MASSGAMEPLAY_DEBUG
+	, DebugEntityQuery(*this)
+#endif
 {
 	ExecutionFlags = (int32)EProcessorExecutionFlags::AllNetModes;
 	ExecutionOrder.ExecuteInGroup = UE::Mass::ProcessorGroupNames::Movement;
@@ -72,6 +75,11 @@ void UMassApplyMovementProcessor::ConfigureQueries(const TSharedRef<FMassEntityM
 	EntityQuery.AddRequirement<FMassDesiredMovementFragment>(EMassFragmentAccess::ReadOnly);
 	EntityQuery.AddTagRequirement<FMassOffLODTag>(EMassFragmentPresence::None);
 	EntityQuery.AddTagRequirement<FMassCodeDrivenMovementTag>(EMassFragmentPresence::All);
+
+#if WITH_MASSGAMEPLAY_DEBUG
+	// We have a separate processor because we want the debug entity query to run always, regardless of the FMassCodeDrivenMovementTag 
+	DebugEntityQuery.AddRequirement<FMassVelocityFragment>(EMassFragmentAccess::ReadWrite);
+#endif
 }
 
 void UMassApplyMovementProcessor::Execute(FMassEntityManager& EntityManager,
@@ -102,15 +110,26 @@ void UMassApplyMovementProcessor::Execute(FMassEntityManager& EntityManager,
 			{
 				Velocity.Value = FVector::ZeroVector;
 			}
-
-			// Keep as "expected value" for next frame.
-			Velocity.DebugPreviousValue = Velocity.Value;
 #endif // WITH_MASSGAMEPLAY_DEBUG
 
 			FVector CurrentLocation = CurrentTransform.GetLocation();
 			CurrentLocation += Velocity.Value * DeltaTime;
 			CurrentTransform.SetTranslation(CurrentLocation);
+		}
+	});
 
+#if WITH_MASSGAMEPLAY_DEBUG
+	DebugEntityQuery.ForEachEntityChunk(Context, [this, DeltaTime](FMassExecutionContext& Context)
+	{
+		const TArrayView<FMassVelocityFragment> VelocityList = Context.GetMutableFragmentView<FMassVelocityFragment>();
+
+		for (FMassExecutionContext::FEntityIterator EntityIt = Context.CreateEntityIterator(); EntityIt; ++EntityIt)
+		{
+			FMassVelocityFragment& Velocity = VelocityList[EntityIt];
+
+			// Keep as "expected value" for next frame.
+			Velocity.DebugPreviousValue = Velocity.Value;
 		}
 	});
+#endif // WITH_MASSGAMEPLAY_DEBUG
 }
diff --git a/Engine/Plugins/Runtime/MassGameplay/Source/MassMovement/Public/Movement/MassMovementProcessors.h b/Engine/Plugins/Runtime/MassGameplay/Source/MassMovement/Public/Movement/MassMovementProcessors.h
--- a/Engine/Plugins/Runtime/MassGameplay/Source/MassMovement/Public/Movement/MassMovementProcessors.h
+++ b/Engine/Plugins/Runtime/MassGameplay/Source/MassMovement/Public/Movement/MassMovementProcessors.h
@@ -3,6 +3,7 @@
 #pragma once
 
 #include "MassProcessor.h"
+#include "MassCommonTypes.h" // For WITH_MASSGAMEPLAY_DEBUG
 #include "MassMovementProcessors.generated.h"
 
 #define UE_API MASSMOVEMENT_API
@@ -46,6 +47,10 @@ protected:
 
 private:
 	FMassEntityQuery EntityQuery;
+
+#if WITH_MASSGAMEPLAY_DEBUG
+	FMassEntityQuery DebugEntityQuery;
+#endif
 };
 
 #undef UE_API
