[MassGameplay] added TMaps to UMassLookAtSubsystem to alleviate costs of linear searching through PerViewerRequests and RegisteredRequests, which can get pretty large.

Additionally, minor change, changed how UMassLookAtRequestDeinitializer::Execute gathers the affected entities.

diff --git a/Engine/Plugins/AI/MassAI/Source/MassAIBehavior/Private/MassLookAtProcessors.cpp b/Engine/Plugins/AI/MassAI/Source/MassAIBehavior/Private/MassLookAtProcessors.cpp
--- a/Engine/Plugins/AI/MassAI/Source/MassAIBehavior/Private/MassLookAtProcessors.cpp
+++ b/Engine/Plugins/AI/MassAI/Source/MassAIBehavior/Private/MassLookAtProcessors.cpp
@@ -755,11 +755,7 @@ void UMassLookAtRequestDeinitializer::Execute(FMassEntityManager& EntityManager,
 
 	Query.ForEachEntityChunk(Context, [&Requests](FMassExecutionContext& Context)
 	{
-		Requests.Reserve(Requests.Num() + Context.GetNumEntities());
-		for (FMassExecutionContext::FEntityIterator EntityIt = Context.CreateEntityIterator(); EntityIt; ++EntityIt)
-		{
-			Requests.Add(Context.GetEntity(EntityIt));
-		}
+		Requests.Append(Context.GetEntities());
 	});
 
 	LookAtSubsystem.UnregisterRequests(Context, Requests);
diff --git a/Engine/Plugins/AI/MassAI/Source/MassAIBehavior/Private/MassLookAtSubsystem.cpp b/Engine/Plugins/AI/MassAI/Source/MassAIBehavior/Private/MassLookAtSubsystem.cpp
--- a/Engine/Plugins/AI/MassAI/Source/MassAIBehavior/Private/MassLookAtSubsystem.cpp
+++ b/Engine/Plugins/AI/MassAI/Source/MassAIBehavior/Private/MassLookAtSubsystem.cpp
@@ -443,6 +443,7 @@ void UMassLookAtSubsystem::RegisterRequests(const FMassExecutionContext& InConte
 
 	{
 		UE_MT_SCOPED_WRITE_ACCESS(RequestsAccessDetector);
+		TRACE_CPUPROFILER_EVENT_SCOPE_STR("MassLookAt_RegisterRequests")
 
 		RegisteredRequests.Reserve(RegisteredRequests.Num() + (InRequests.Num() - ActiveRequestsFreeList.Num()));
 
@@ -466,11 +467,9 @@ void UMassLookAtSubsystem::RegisterRequests(const FMassExecutionContext& InConte
 				NewRequestIndex = RegisteredRequests.Add(NewRequest);
 			}
 
-			int32 ViewerDataIndex = PerViewerRequests.IndexOfByPredicate([ViewerEntity = NewRequest.Parameters.ViewerEntity](const FViewerRequest& ViewerRequest)
-			{
-				return ViewerEntity == ViewerRequest.Viewer;
-			});
+			RequestHandleToIndexMap.Add(NewRequest.RequestHandle, NewRequestIndex);
 
+			int32& ViewerDataIndex = ViewerHandleToIndexMap.FindOrAdd(NewRequest.Parameters.ViewerEntity, INDEX_NONE);
 			if (ViewerDataIndex == INDEX_NONE)
 			{
 				ViewerDataIndex = PerViewerRequests.Num();
@@ -492,37 +491,27 @@ void UMassLookAtSubsystem::UnregisterRequests(const FMassExecutionContext& InCon
 
 	{
 		UE_MT_SCOPED_WRITE_ACCESS(RequestsAccessDetector);
+		TRACE_CPUPROFILER_EVENT_SCOPE_STR("MassLookAt_UnregisterRequests")
 
 		for (const FMassEntityHandle& RemovedRequest : InRequests)
 		{
 			int32 InvalidatedIndex = INDEX_NONE;
-			FMassEntityHandle ViewerEntity;
-
-			for (int32 ActiveRequestIndex = 0; ActiveRequestIndex < RegisteredRequests.Num(); ++ActiveRequestIndex)
-			{
-				if (RegisteredRequests[ActiveRequestIndex].RequestHandle == RemovedRequest)
-				{
-					InvalidatedIndex = ActiveRequestIndex;
-
-					// Invalidate entry and add to the free list
-					ViewerEntity = RegisteredRequests[ActiveRequestIndex].Parameters.ViewerEntity;
-					RegisteredRequests[ActiveRequestIndex].RequestHandle = FMassEntityHandle();
-					ActiveRequestsFreeList.Push(ActiveRequestIndex);
-					break;
-				}
-			}
+			RequestHandleToIndexMap.RemoveAndCopyValue(RemovedRequest, InvalidatedIndex);
 
 			if (ensureMsgf(InvalidatedIndex != INDEX_NONE, TEXT("Trying to remove a request that was never registered")))
 			{
-				const int32 ViewerDataIndex = PerViewerRequests.IndexOfByPredicate([ViewerEntity](const FViewerRequest& ViewerRequest)
-					{
-						return ViewerEntity == ViewerRequest.Viewer;
-					});
+				// Invalidate entry and add to the free list
+				const FMassEntityHandle ViewerEntity = RegisteredRequests[InvalidatedIndex].Parameters.ViewerEntity;
+				RegisteredRequests[InvalidatedIndex].RequestHandle.Reset();
+				ActiveRequestsFreeList.Push(InvalidatedIndex);
+
+				const int32* ViewerDataIndex = ViewerHandleToIndexMap.Find(ViewerEntity);
 
-				if (ensureMsgf(ViewerDataIndex != INDEX_NONE, TEXT("Unable to find the per viewer data; looks like a bookkeeping issue")))
+				if (ensureMsgf(ViewerDataIndex, TEXT("Unable to find the per viewer data; looks like a bookkeeping issue")))
 				{
-					DirtyViewers.Add(ViewerDataIndex);
-					PerViewerRequests[ViewerDataIndex].RequestIndices.Remove(InvalidatedIndex);
+					checkf(*ViewerDataIndex != INDEX_NONE, TEXT("INDEX_NONEs are not expected to be stored in ViewerHandleToIndexMap"));
+					DirtyViewers.Add(*ViewerDataIndex);
+					PerViewerRequests[*ViewerDataIndex].RequestIndices.Remove(InvalidatedIndex);
 				}
 			}
 		}
@@ -533,6 +522,8 @@ void UMassLookAtSubsystem::UnregisterRequests(const FMassExecutionContext& InCon
 
 void UMassLookAtSubsystem::UpdateLookAts(const FMassExecutionContext& Context, const TConstArrayView<int32> DirtyViewers)
 {
+	TRACE_CPUPROFILER_EVENT_SCOPE_STR("MassLookAt_UpdateLookAts")
+
 	TArray<TPair<FMassEntityHandle, FMassLookAtRequestFragment>> Updated;
 	for (const int32 DirtyViewerIndex : DirtyViewers)
 	{
@@ -594,15 +585,12 @@ FString UMassLookAtSubsystem::DebugGetRequestsString(const FMassEntityHandle InE
 {
 	UE_MT_SCOPED_READ_ACCESS(RequestsAccessDetector);
 
-	const int32 ViewerDataIndex = PerViewerRequests.IndexOfByPredicate([InEntity](const FViewerRequest& ViewerRequest)
-		{
-			return InEntity == ViewerRequest.Viewer;
-		});
+	const int32* ViewerDataIndex = ViewerHandleToIndexMap.Find(InEntity);
 
 	TStringBuilder<256> Builder;
-	if (ViewerDataIndex != INDEX_NONE)
+	if (*ViewerDataIndex && *ViewerDataIndex != INDEX_NONE)
 	{
-		for (const int32 RequestIndex : PerViewerRequests[ViewerDataIndex].RequestIndices)
+		for (const int32 RequestIndex : PerViewerRequests[*ViewerDataIndex].RequestIndices)
 		{
 			if (Builder.Len() > 0)
 			{
@@ -615,4 +603,4 @@ FString UMassLookAtSubsystem::DebugGetRequestsString(const FMassEntityHandle InE
 	}
 	return Builder.ToString();
 }
-#endif // WITH_MASSGAMEPLAY_DEBUG
+#endif // WITH_MASSGAMEPLAY_DEBUG
\ No newline at end of file
diff --git a/Engine/Plugins/AI/MassAI/Source/MassAIBehavior/Public/MassLookAtSubsystem.h b/Engine/Plugins/AI/MassAI/Source/MassAIBehavior/Public/MassLookAtSubsystem.h
--- a/Engine/Plugins/AI/MassAI/Source/MassAIBehavior/Public/MassLookAtSubsystem.h
+++ b/Engine/Plugins/AI/MassAI/Source/MassAIBehavior/Public/MassLookAtSubsystem.h
@@ -226,12 +226,21 @@ private:
 	/** List of all currently registered requests as registered by external systems */
 	TArray<FRequest> RegisteredRequests;
 
+	/** Used for lookup in RegisteredRequests, since that array can get large, and linear search would kill performance rather quickly. */
+	TMap<FMassEntityHandle, int32> RequestHandleToIndexMap;
+
 	/** List of available indices in ActiveRequests (to preserve stable indices when unregistering requests) */
 	TArray<int32> ActiveRequestsFreeList;
 
-	/** Per viewer entity representation of all the active requests */
+	/**
+	 * Per viewer entity representation of all the active requests
+	 * @todo we never remove elements from this array, which means it can end up being HUGE. Needs fixing.
+	 */
 	TArray<FViewerRequest> PerViewerRequests;
 
+	/** Used for lookup in PerViewerRequests, since that array can get large, and linear search would kill performance rather quickly. */
+	TMap<FMassEntityHandle, int32> ViewerHandleToIndexMap;
+
 	/** Multithread access detector to detect threading issues with any list of requests */
 	UE_MT_DECLARE_RW_ACCESS_DETECTOR(RequestsAccessDetector);
 
