Added text to MassNavigation::Debug methods

diff --git a/Engine/Plugins/AI/MassAI/Source/MassNavigation/Private/MassAvoidanceProcessors.cpp b/Engine/Plugins/AI/MassAI/Source/MassNavigation/Private/MassAvoidanceProcessors.cpp
--- a/Engine/Plugins/AI/MassAI/Source/MassNavigation/Private/MassAvoidanceProcessors.cpp
+++ b/Engine/Plugins/AI/MassAI/Source/MassNavigation/Private/MassAvoidanceProcessors.cpp
@@ -239,9 +239,9 @@ namespace UE::MassAvoidance
 		}
 	}
 
-	static void DebugDrawForce(const FDebugContext& Context, const FVector& Start, const FVector& End, const FColor& Color, const float Thickness)
+	static void DebugDrawForce(const FDebugContext& Context, const FVector& Start, const FVector& End, const FColor& Color, const float Thickness, const FString& Text = FString())
 	{
-		DebugDrawArrow(Context, Start, End, Color, /*HeadSize*/4.f, Thickness);
+		DebugDrawArrow(Context, Start, End, Color, /*HeadSize*/4.f, Thickness, Text);
 	}
 
 	static void DebugDrawSummedForce(const FDebugContext& Context, const FVector& Start, const FVector& End, const FColor& Color)
diff --git a/Engine/Plugins/AI/MassAI/Source/MassNavigation/Private/MassNavigationDebug.cpp b/Engine/Plugins/AI/MassAI/Source/MassNavigation/Private/MassNavigationDebug.cpp
--- a/Engine/Plugins/AI/MassAI/Source/MassNavigation/Private/MassNavigationDebug.cpp
+++ b/Engine/Plugins/AI/MassAI/Source/MassNavigation/Private/MassNavigationDebug.cpp
@@ -47,7 +47,8 @@ namespace UE::MassNavigation::Debug
 		}
 	}
 
-	void DebugDrawArrow(const FDebugContext& Context, const FVector& Start, const FVector& End, const FColor& Color, const float HeadSize /*= 8.f*/, const float Thickness /*= 1.5f*/)
+	void DebugDrawArrow(const FDebugContext& Context, const FVector& Start, const FVector& End, const FColor& Color,
+		const float HeadSize /*= 8.f*/, const float Thickness /*= 1.5f*/, const FString& Text /*= FString()*/)
 	{
 		if (!Context.ShouldLogEntity())
 		{
@@ -62,7 +63,7 @@ namespace UE::MassNavigation::Debug
 		const FVector Perp = FVector::CrossProduct(UnitV, FVector::UpVector);
 		const FVector Left = Perp - (Pointyness*UnitV);
 		const FVector Right = -Perp - (Pointyness*UnitV);
-		UE_VLOG_SEGMENT_THICK(LogOwner, Context.Category, Log, Start, End, Color, (int16)Thickness, TEXT(""));
+		UE_VLOG_SEGMENT_THICK(LogOwner, Context.Category, Log, Start, End, Color, (int16)Thickness, TEXT("%s"), *Text);
 		UE_VLOG_SEGMENT_THICK(LogOwner, Context.Category, Log, End, End + HeadSize * Left, Color, (int16)Thickness, TEXT(""));
 		UE_VLOG_SEGMENT_THICK(LogOwner, Context.Category, Log, End, End + HeadSize * Right, Color, (int16)Thickness, TEXT(""));
 
@@ -74,7 +75,7 @@ namespace UE::MassNavigation::Debug
 		}
 	}
 
-	void DebugDrawSphere(const FDebugContext& Context, const FVector& Center, const FVector::FReal InRadius, const FColor& Color)
+	void DebugDrawSphere(const FDebugContext& Context, const FVector& Center, const FVector::FReal InRadius, const FColor& Color, const FString& Text /*= FString()*/)
 	{
 		if (!Context.ShouldLogEntity())
 		{
@@ -82,7 +83,7 @@ namespace UE::MassNavigation::Debug
 		}
 
 		const float Radius = static_cast<float>(InRadius);
-		UE_VLOG_LOCATION(Context.GetLogOwner(), Context.Category, Log, Center, Radius, Color, TEXT(""));
+		UE_VLOG_LOCATION(Context.GetLogOwner(), Context.Category, Log, Center, Radius, Color, TEXT("%s"), *Text);
 
 		if (UseDrawDebugHelper() && Context.World)
 		{
@@ -90,14 +91,14 @@ namespace UE::MassNavigation::Debug
 		}
 	}
 
-	void DebugDrawBox(const FDebugContext& Context, const FBox& Box, const FColor& Color)
+	void DebugDrawBox(const FDebugContext& Context, const FBox& Box, const FColor& Color, const FString& Text /*= FString()*/)
 	{
 		if (!Context.ShouldLogEntity())
 		{
 			return;
 		}
 
-		UE_VLOG_BOX(Context.GetLogOwner(), Context.Category, Log, Box, Color, TEXT(""));
+		UE_VLOG_BOX(Context.GetLogOwner(), Context.Category, Log, Box, Color, TEXT("%s"), *Text);
 		
 		if (UseDrawDebugHelper() && Context.World)
 		{
diff --git a/Engine/Plugins/AI/MassAI/Source/MassNavigation/Public/MassNavigationDebug.h b/Engine/Plugins/AI/MassAI/Source/MassNavigation/Public/MassNavigationDebug.h
--- a/Engine/Plugins/AI/MassAI/Source/MassNavigation/Public/MassNavigationDebug.h
+++ b/Engine/Plugins/AI/MassAI/Source/MassNavigation/Public/MassNavigationDebug.h
@@ -50,12 +50,23 @@ namespace UE::MassNavigation::Debug
 	
 	MASSNAVIGATION_API bool UseDrawDebugHelper();
 	MASSNAVIGATION_API bool DebugIsSelected(const FMassEntityHandle Entity);
-	MASSNAVIGATION_API void DebugDrawLine(const FDebugContext& Context, const FVector& Start, const FVector& End, const FColor& Color, const float Thickness = 0.f, const bool bPersistent = false, const FString& Text = FString());
-	MASSNAVIGATION_API void DebugDrawArrow(const FDebugContext& Context, const FVector& Start, const FVector& End, const FColor& Color, const float HeadSize = 8.f, const float Thickness = 1.5f);
-	MASSNAVIGATION_API void DebugDrawSphere(const FDebugContext& Context, const FVector& Center, const FVector::FReal InRadius, const FColor& Color);
-	MASSNAVIGATION_API void DebugDrawBox(const FDebugContext& Context, const FBox& Box, const FColor& Color);
-	MASSNAVIGATION_API void DebugDrawCylinder(const FDebugContext& Context, const FVector& Bottom, const FVector& Top, const FVector::FReal InRadius, const FColor& Color, const FString& Text = FString());
-	MASSNAVIGATION_API void DebugDrawCircle(const FDebugContext& Context, const FVector& Bottom, const FVector::FReal InRadius, const FColor& Color, const FString& Text = FString());
+
+	MASSNAVIGATION_API void DebugDrawLine(const FDebugContext& Context, const FVector& Start, const FVector& End, const FColor& Color,
+		const float Thickness = 0.f, const bool bPersistent = false, const FString& Text = FString());
+
+	MASSNAVIGATION_API void DebugDrawArrow(const FDebugContext& Context, const FVector& Start, const FVector& End, const FColor& Color,
+		const float HeadSize = 8.f, const float Thickness = 1.5f, const FString& Text = FString());
+
+	MASSNAVIGATION_API void DebugDrawSphere(const FDebugContext& Context, const FVector& Center, const FVector::FReal InRadius, const FColor& Color,
+		const FString& Text = FString());
+
+	MASSNAVIGATION_API void DebugDrawBox(const FDebugContext& Context, const FBox& Box, const FColor& Color, const FString& Text = FString());
+
+	MASSNAVIGATION_API void DebugDrawCylinder(const FDebugContext& Context, const FVector& Bottom, const FVector& Top, const FVector::FReal InRadius,
+		const FColor& Color, const FString& Text = FString());
+
+	MASSNAVIGATION_API void DebugDrawCircle(const FDebugContext& Context, const FVector& Bottom, const FVector::FReal InRadius, const FColor& Color,
+		const FString& Text = FString());
 	
 	inline FColor MixColors(const FColor ColorA, const FColor ColorB)
 	{
