[MassLookAt] fixed main and gaze directions that were not supposed to consider the target offset.

diff --git a/Engine/Plugins/AI/MassAI/Source/MassAIBehavior/Private/MassLookAtProcessors.cpp b/Engine/Plugins/AI/MassAI/Source/MassAIBehavior/Private/MassLookAtProcessors.cpp
--- a/Engine/Plugins/AI/MassAI/Source/MassAIBehavior/Private/MassLookAtProcessors.cpp
+++ b/Engine/Plugins/AI/MassAI/Source/MassAIBehavior/Private/MassLookAtProcessors.cpp
@@ -403,12 +403,10 @@ void UMassLookAtProcessor::UpdateLookAtTrackedEntity(const FMassEntityManager& E
 		const FMassLookAtTargetFragment* TargetFragment = EntityManager.GetFragmentDataPtr<FMassLookAtTargetFragment>(LookAt.TrackedEntity);
 
 		const FVector AgentPosition = Transform.GetLocation();
-		const FVector TargetLocation = TargetFragment != nullptr
-			? TargetTransform.GetTransform().TransformPosition(TargetFragment->Offset)
-			: TargetTransform.GetTransform().GetLocation();
-		const FVector NewGlobalDirection = (TargetLocation - AgentPosition).GetSafeNormal();
+		const FVector TargetBaseLocation = TargetTransform.GetTransform().GetLocation();
+		const FVector NewGlobalDirection = (TargetBaseLocation - AgentPosition).GetSafeNormal();
 		LookAt.Direction = Transform.InverseTransformVector(NewGlobalDirection);
-		LookAt.MainTargetLocation = TargetLocation;
+		LookAt.MainTargetLocation = TargetFragment != nullptr ? TargetTransform.GetTransform().TransformPosition(TargetFragment->Offset) : TargetBaseLocation;
 
 #if WITH_MASSGAMEPLAY_DEBUG
 		if (bDisplayDebug)
@@ -431,12 +429,10 @@ bool UMassLookAtProcessor::UpdateGazeTrackedEntity(const FMassEntityManager& Ent
 		const FMassLookAtTargetFragment* TargetFragment = EntityManager.GetFragmentDataPtr<FMassLookAtTargetFragment>(LookAt.GazeTrackedEntity);
 
 		const FVector AgentPosition = Transform.GetLocation();
-		const FVector TargetLocation = TargetFragment != nullptr
-			? TargetTransform.GetTransform().TransformPosition(TargetFragment->Offset)
-			: TargetTransform.GetTransform().GetLocation();
-		const FVector NewGlobalDirection = (TargetLocation - AgentPosition).GetSafeNormal();
+		const FVector TargetBaseLocation = TargetTransform.GetTransform().GetLocation();
+		const FVector NewGlobalDirection = (TargetBaseLocation - AgentPosition).GetSafeNormal();
 		LookAt.GazeDirection = Transform.InverseTransformVector(NewGlobalDirection);
-		LookAt.GazeTargetLocation = TargetLocation;
+		LookAt.GazeTargetLocation = TargetFragment != nullptr ? TargetTransform.GetTransform().TransformPosition(TargetFragment->Offset) : TargetBaseLocation;
 
 		bHasTarget = true;
 
