[MassLookAt] systemic lookat are not using the request systems which is reserved to overrides. This will reduce a lot the number of created/destroyed MassEntities in large simulations

diff --git a/Engine/Plugins/AI/MassAI/Source/MassAIBehavior/Private/MassLookAtSubsystem.cpp b/Engine/Plugins/AI/MassAI/Source/MassAIBehavior/Private/MassLookAtSubsystem.cpp
--- a/Engine/Plugins/AI/MassAI/Source/MassAIBehavior/Private/MassLookAtSubsystem.cpp
+++ b/Engine/Plugins/AI/MassAI/Source/MassAIBehavior/Private/MassLookAtSubsystem.cpp
@@ -6,7 +6,8 @@
 #include "MassAIBehaviorTypes.h"
 #include "MassEntitySubsystem.h"
 #include "MassEntityView.h"
-#include "Debugger/StateTreeDebuggerTypes.h"
+#include "MassSignalSubsystem.h"
+#include "MassStateTreeTypes.h"
 
 namespace UE::Mass::LookAt::Private
 {
@@ -32,7 +33,7 @@ bool CreateRequests(FMassEntityManager& InEntityManager
 	, const FMassEntityHandle TargetEntity
 	, const FMassLookAtPriority Priority
 	, const EMassLookAtInterpolationSpeed InterpolationSpeed
-	, const float CustomInterpolationSpeed = UE::Mass::LookAt::DefaultCustomInterpolationSpeed)
+	, const float CustomInterpolationSpeed = DefaultCustomInterpolationSpeed)
 {
 	if (Viewers.IsEmpty())
 	{
@@ -562,19 +563,87 @@ void UMassLookAtSubsystem::UpdateLookAts(const FMassExecutionContext& Context, c
 
 	Context.Defer().PushCommand<FMassDeferredSetCommand>([Updated = MoveTemp(Updated)](const FMassEntityManager& Manager)
 		{
+			TArray<FMassEntityHandle> EntitiesToSignal;
+			EntitiesToSignal.Reserve(Updated.Num());
 			for (const TPair<FMassEntityHandle, FMassLookAtRequestFragment>& Pair : Updated)
 			{
-				if (Manager.IsEntityValid(Pair.Key))
+				const FMassEntityHandle Entity = Pair.Key;
+				if (FMassLookAtFragment* LookAtFragment = Manager.IsEntityValid(Entity) ? Manager.GetFragmentDataPtr<FMassLookAtFragment>(Entity) : nullptr)
 				{
-					if (FMassLookAtFragment* LookAtFragment = Manager.GetFragmentDataPtr<FMassLookAtFragment>(Pair.Key))
+					const FMassLookAtRequestFragment& Request = Pair.Value;
+
+					// Default request is used to clear the override
+					const bool bIsClearOverrideRequest = !Request.ViewerEntity.IsSet();
+
+					switch (LookAtFragment->OverrideState)
 					{
-						// Only update main LookAt information and don't modify gaze related one
-						const FMassLookAtRequestFragment& Request = Pair.Value;
-						LookAtFragment->InterpolationSpeed = Request.InterpolationSpeed;
-						LookAtFragment->CustomInterpolationSpeed = Request.CustomInterpolationSpeed;
-						LookAtFragment->LookAtMode = Request.LookAtMode;
-						LookAtFragment->TrackedEntity = Request.TargetEntity;
+					case FMassLookAtFragment::EOverrideState::AllDisabled:
+						if (bIsClearOverrideRequest)
+						{
+							// Already disabled no need to change the fragment
+							continue;
+						}
+
+						LookAtFragment->OverrideState = FMassLookAtFragment::EOverrideState::ActiveOverrideOnly;
+						break;
+
+					case FMassLookAtFragment::EOverrideState::ActiveOverrideOnly:
+						if (bIsClearOverrideRequest)
+						{
+							LookAtFragment->OverrideState = FMassLookAtFragment::EOverrideState::AllDisabled;
+						}
+
+						// In both cases we set values from the request to apply the override or to clear the active one
+						break;
+
+					case FMassLookAtFragment::EOverrideState::ActiveSystemicOnly:
+						if (bIsClearOverrideRequest)
+						{
+							// Override already cleared, do not change the fragment (systemic is active)
+							continue;
+						}
+
+						// Switch to 'OverridenSystemic' state and apply request parameters
+						LookAtFragment->OverrideState = FMassLookAtFragment::EOverrideState::OverridenSystemic;
+						break;
+
+					case FMassLookAtFragment::EOverrideState::OverridenSystemic:
+						if (bIsClearOverrideRequest)
+						{
+							// Mark as pending reactivation and signal the entity to wake up the task to retry activation of the systemic LookAt
+							LookAtFragment->OverrideState = FMassLookAtFragment::EOverrideState::PendingSystemicReactivation;
+							EntitiesToSignal.Add(Pair.Key);
+							continue;
+						}
+						// Stay in 'OverridenSystemic' state and apply new request parameters
+						break;
+
+					case FMassLookAtFragment::EOverrideState::PendingSystemicReactivation:
+						if (bIsClearOverrideRequest)
+						{
+							// Already pending reactivation and signaled, no need to signal again
+							continue;
+						}
+
+						// Switch back to 'OverridenSystemic' so task will not apply its values when processing the signal
+						LookAtFragment->OverrideState = FMassLookAtFragment::EOverrideState::OverridenSystemic;
+						break;
 					}
+
+					// Only update main LookAt information and don't modify gaze related one
+					LookAtFragment->InterpolationSpeed = Request.InterpolationSpeed;
+					LookAtFragment->CustomInterpolationSpeed = Request.CustomInterpolationSpeed;
+					LookAtFragment->LookAtMode = Request.LookAtMode;
+					LookAtFragment->TrackedEntity = Request.TargetEntity;
+				}
+			}
+
+			// Signal all entities
+			if (EntitiesToSignal.Num())
+			{
+				if (UMassSignalSubsystem* SignalSubsystem = Manager.GetWorld()->GetSubsystem<UMassSignalSubsystem>())
+				{
+					SignalSubsystem->SignalEntities(UE::Mass::Signals::LookAtFinished, EntitiesToSignal);
 				}
 			}
 		});
diff --git a/Engine/Plugins/AI/MassAI/Source/MassAIBehavior/Private/Tasks/MassLookAtTask.cpp b/Engine/Plugins/AI/MassAI/Source/MassAIBehavior/Private/Tasks/MassLookAtTask.cpp
--- a/Engine/Plugins/AI/MassAI/Source/MassAIBehavior/Private/Tasks/MassLookAtTask.cpp
+++ b/Engine/Plugins/AI/MassAI/Source/MassAIBehavior/Private/Tasks/MassLookAtTask.cpp
@@ -23,6 +23,39 @@ void FMassLookAtTask::GetDependencies(UE::MassBehavior::FStateTreeDependencyBuil
 	Builder.AddReadWrite<UMassSignalSubsystem>();
 }
 
+bool FMassLookAtTask::TryActivateSystemicLookAt(const FStateTreeExecutionContext& Context, const FInstanceDataType& InstanceData, FMassLookAtFragment& Fragment) const
+{
+	// We can activate systemic LookAt in the following cases:
+	// - nothing is currently active
+	// - a systemic one is active (the leaf task has priority preserving the original implementation)
+	// - the last override got removed while a systemic task is still running
+	if (Fragment.OverrideState == FMassLookAtFragment::EOverrideState::AllDisabled
+		|| Fragment.OverrideState == FMassLookAtFragment::EOverrideState::ActiveSystemicOnly
+		|| Fragment.OverrideState == FMassLookAtFragment::EOverrideState::PendingSystemicReactivation)
+	{
+		Fragment.InterpolationSpeed = InterpolationSpeed;
+		Fragment.CustomInterpolationSpeed = CustomInterpolationSpeed;
+		Fragment.LookAtMode = LookAtMode;
+		Fragment.TrackedEntity = InstanceData.TargetEntity;
+		Fragment.OverrideState = FMassLookAtFragment::EOverrideState::ActiveSystemicOnly;
+
+		// When using 'LookAtEntity' we validate the target entity to use it,
+		// or we use the default 'LookForward' as fallback.
+		if (LookAtMode == EMassLookAtMode::LookAtEntity)
+		{
+			if (!InstanceData.TargetEntity.IsSet())
+			{
+				MASSBEHAVIOR_LOG(Error, TEXT("Failed LookAt: invalid target entity"));
+				Fragment.LookAtMode = EMassLookAtMode::LookForward;
+				Fragment.TrackedEntity = {};
+			}
+		}
+		return true;
+	}
+
+	return false;
+}
+
 EStateTreeRunStatus FMassLookAtTask::EnterState(FStateTreeExecutionContext& Context, const FStateTreeTransitionResult& Transition) const
 {
 	FInstanceDataType& InstanceData = Context.GetInstanceData(*this);
@@ -33,51 +66,7 @@ EStateTreeRunStatus FMassLookAtTask::EnterState(FStateTreeExecutionContext& Cont
 	FMassLookAtFragment* LookAtFragment = MassStateTreeContext.GetExternalDataPtr(LookAtHandle);
 	if (LookAtFragment != nullptr)
 	{
-		MassStateTreeContext.GetMassEntityExecutionContext().Defer()
-			.PushCommand<FMassDeferredCreateCommand>(
-				[&Context, SelfEntity = MassStateTreeContext.GetEntity()
-				, &InstanceData
-				, Mode = LookAtMode
-				, Priority = Priority
-				, InterpolationSpeed = InterpolationSpeed
-				, CustomInterpolationSpeed = CustomInterpolationSpeed]
-				(FMassEntityManager& InEntityManager)
-				{
-					const FMassArchetypeCompositionDescriptor Composition{FMassFragmentBitSet(*FMassLookAtRequestFragment::StaticStruct())};
-
-					constexpr int32 NumRequestsToCreate = 1;
-					TArray<FMassEntityHandle> Requests;
-					const TSharedRef<FMassEntityManager::FEntityCreationContext> CreationContext =
-						InEntityManager.BatchCreateEntities(InEntityManager.CreateArchetype(Composition), NumRequestsToCreate, Requests);
-
-					if (Requests.Num() == NumRequestsToCreate)
-					{
-						InstanceData.Request = Requests.Top();
-						FMassLookAtRequestFragment RequestFragment
-						{
-							SelfEntity
-							, Priority
-							, Mode
-							, InstanceData.TargetEntity
-							, InterpolationSpeed
-							, CustomInterpolationSpeed
-						};
-
-						// When using 'LookAtEntity' we validate the target entity to use it,
-						// or we use the default 'LookForward' as fallback.
-						if (Mode == EMassLookAtMode::LookAtEntity)
-						{
-							if (!InstanceData.TargetEntity.IsSet())
-							{
-								MASSBEHAVIOR_LOG(Error, TEXT("Failed LookAt: invalid target entity"));
-								RequestFragment.LookAtMode = EMassLookAtMode::LookForward;
-								RequestFragment.TargetEntity = {};
-							}
-						}
-
-						InEntityManager.SetEntityFragmentValues(InstanceData.Request, { FInstancedStruct::Make(RequestFragment) });
-					}
-				});
+		TryActivateSystemicLookAt(Context, InstanceData, *LookAtFragment);
 
 		LookAtFragment->RandomGazeMode = RandomGazeMode;
 		LookAtFragment->RandomGazeYawVariation = RandomGazeYawVariation;
@@ -106,31 +95,33 @@ void FMassLookAtTask::ExitState(FStateTreeExecutionContext& Context, const FStat
 {
 	const FMassStateTreeExecutionContext& MassStateTreeContext = static_cast<FMassStateTreeExecutionContext&>(Context);
 
-	// Reset gaze directly since it is used only by the task and not by the requests
 	if (FMassLookAtFragment* LookAtFragment = MassStateTreeContext.GetExternalDataPtr(LookAtHandle))
 	{
-		LookAtFragment->ResetGaze();
-
-		// Destroy any associated request for the main LookAt
-		FInstanceDataType& InstanceData = Context.GetInstanceData(*this);
-		if (InstanceData.Request.IsSet())
-		{
-			MassStateTreeContext.GetMassEntityExecutionContext().Defer()
-				.PushCommand<FMassDeferredDestroyCommand>([Request = InstanceData.Request](FMassEntityManager& InEntityManager)
-				{
-					InEntityManager.DestroyEntity(Request);
-				});
-
-			InstanceData.Request.Reset();
-		}
+		LookAtFragment->ResetSystemicLookAt();
 	}
 }
 
 EStateTreeRunStatus FMassLookAtTask::Tick(FStateTreeExecutionContext& Context, const float DeltaTime) const
 {
 	FInstanceDataType& InstanceData = Context.GetInstanceData(*this);
-	
+
 	InstanceData.Time += DeltaTime;
-	
-	return InstanceData.Duration <= 0.0f ? EStateTreeRunStatus::Running : (InstanceData.Time < InstanceData.Duration ? EStateTreeRunStatus::Running : EStateTreeRunStatus::Succeeded);
+
+	const EStateTreeRunStatus Status = InstanceData.Duration <= 0.0f ? EStateTreeRunStatus::Running : (InstanceData.Time < InstanceData.Duration ? EStateTreeRunStatus::Running : EStateTreeRunStatus::Succeeded);
+
+	// We might get notified that a LookAt override finished so we can try to activate again
+	if (Status == EStateTreeRunStatus::Running)
+	{
+		const FMassStateTreeExecutionContext& MassStateTreeContext = static_cast<FMassStateTreeExecutionContext&>(Context);
+		if (FMassLookAtFragment* Fragment = MassStateTreeContext.GetExternalDataPtr(LookAtHandle); Fragment != nullptr)
+		{
+			// Only care about reactivation in the Tick
+			if (Fragment->OverrideState == FMassLookAtFragment::EOverrideState::PendingSystemicReactivation)
+			{
+				TryActivateSystemicLookAt(Context, InstanceData, *Fragment);
+			}
+		}
+	}
+
+	return Status;
 }
diff --git a/Engine/Plugins/AI/MassAI/Source/MassAIBehavior/Public/MassLookAtFragments.h b/Engine/Plugins/AI/MassAI/Source/MassAIBehavior/Public/MassLookAtFragments.h
--- a/Engine/Plugins/AI/MassAI/Source/MassAIBehavior/Public/MassLookAtFragments.h
+++ b/Engine/Plugins/AI/MassAI/Source/MassAIBehavior/Public/MassLookAtFragments.h
@@ -44,11 +44,6 @@ struct FMassLookAtFragment : public FMassFragment
 {
 	GENERATED_BODY()
 
-	FMassLookAtFragment()
-		: bRandomGazeEntities(false)
-	{
-	}
-
 	/** Resets gaze related properties in the fragment to their default values */
 	void ResetGaze()
 	{
@@ -66,8 +61,13 @@ struct FMassLookAtFragment : public FMassFragment
 	/** Resets all properties in the fragment to their default values */
 	void ResetAll()
 	{
+		ResetMainLookAt();
 		ResetGaze();
+	}
 
+	/** Resets main LookAt related properties in the fragment to their default values */
+	void ResetMainLookAt()
+	{
 		MainTargetLocation = FVector::ZeroVector;
 		Direction = FVector::ForwardVector;
 		TrackedEntity.Reset();
@@ -75,6 +75,30 @@ struct FMassLookAtFragment : public FMassFragment
 		LookAtMode = EMassLookAtMode::LookForward;
 	}
 
+	/** Resets gaze and main LookAt (if override not active) related properties in the fragment to their default values */
+	void ResetSystemicLookAt()
+	{
+		if (OverrideState == EOverrideState::ActiveSystemicOnly
+			|| OverrideState == EOverrideState::OverridenSystemic
+			|| OverrideState == EOverrideState::PendingSystemicReactivation)
+		{
+			ResetGaze();
+
+			// Only reset main LookAt properties if no override is active.
+			// In that case it will be handled by the LookAtSubsystem when updating the override requests.
+			if (OverrideState == EOverrideState::ActiveSystemicOnly
+				|| OverrideState == EOverrideState::PendingSystemicReactivation)
+			{
+				ResetMainLookAt();
+				OverrideState = EOverrideState::AllDisabled;
+			}
+			else
+			{
+				OverrideState = EOverrideState::ActiveOverrideOnly;
+			}
+		}
+	}
+
 	/** Current look at main target location. */
 	UPROPERTY(Transient)
 	FVector MainTargetLocation = FVector::ZeroVector;
@@ -138,9 +162,33 @@ struct FMassLookAtFragment : public FMassFragment
 	UPROPERTY(Transient)
 	uint8 RandomGazePitchVariation = 0;
 
-	/** Tru if random gaze can also pick interesting entities to look at. */
+	/** Whether random gaze can also pick interesting entities to look at. */
 	UPROPERTY(Transient)
-	uint8 bRandomGazeEntities : 1;
+	uint8 bRandomGazeEntities : 1 = false;
+
+	enum class EOverrideState : uint8
+	{
+		/** No active task and no active override */
+		AllDisabled,
+
+		/** No active task, only an active override */
+		ActiveOverrideOnly,
+
+		/** Active task only, no active override */
+		ActiveSystemicOnly,
+
+		/** Active task overriden */
+		OverridenSystemic,
+
+		/**
+		 * Intermediate state used from 'OverridenSystemic' after removing last override request
+		 * to allow active task to reapply its values and transition to 'ActiveSystemicOnly'
+		 */
+		PendingSystemicReactivation
+	};
+
+	/** Current state based on the systemic task and active overrides */
+	EOverrideState OverrideState : 7 = EOverrideState::AllDisabled;
 };
 
 /**
diff --git a/Engine/Plugins/AI/MassAI/Source/MassAIBehavior/Public/Tasks/MassLookAtTask.h b/Engine/Plugins/AI/MassAI/Source/MassAIBehavior/Public/Tasks/MassLookAtTask.h
--- a/Engine/Plugins/AI/MassAI/Source/MassAIBehavior/Public/Tasks/MassLookAtTask.h
+++ b/Engine/Plugins/AI/MassAI/Source/MassAIBehavior/Public/Tasks/MassLookAtTask.h
@@ -34,10 +34,6 @@ struct FMassLookAtTaskInstanceData
 	/** Accumulated time used to stop task if duration is set */
 	UPROPERTY()
 	float Time = 0.f;
-
-	/** Request created for the main LookAt */
-	UPROPERTY()
-	FMassEntityHandle Request;
 };
 
 USTRUCT(meta = (DisplayName = "Mass LookAt Task"))
@@ -49,12 +45,18 @@ struct FMassLookAtTask : public FMassStateTreeTaskBase
 	
 protected:
 	UE_API virtual bool Link(FStateTreeLinker& Linker) override;
-	virtual const UStruct* GetInstanceDataType() const override { return FInstanceDataType::StaticStruct(); }
+	virtual const UStruct* GetInstanceDataType() const override
+	{
+		return FInstanceDataType::StaticStruct();
+	}
+
 	UE_API virtual EStateTreeRunStatus EnterState(FStateTreeExecutionContext& Context, const FStateTreeTransitionResult& Transition) const override;
 	UE_API virtual void ExitState(FStateTreeExecutionContext& Context, const FStateTreeTransitionResult& Transition) const override;
 	UE_API virtual EStateTreeRunStatus Tick(FStateTreeExecutionContext& Context, const float DeltaTime) const override;
 	UE_API virtual void GetDependencies(UE::MassBehavior::FStateTreeDependencyBuilder& Builder) const override;
 
+	bool TryActivateSystemicLookAt(const FStateTreeExecutionContext& Context, const FInstanceDataType& InstanceData, FMassLookAtFragment& Fragment) const;
+
 	TStateTreeExternalDataHandle<UMassSignalSubsystem> MassSignalSubsystemHandle;
 	TStateTreeExternalDataHandle<FMassLookAtFragment, EStateTreeExternalDataRequirement::Optional> LookAtHandle;
 
