[MassLookAt] crash fix

diff --git a/Engine/Plugins/AI/MassAI/Source/MassAIBehavior/Private/MassLookAtSubsystem.cpp b/Engine/Plugins/AI/MassAI/Source/MassAIBehavior/Private/MassLookAtSubsystem.cpp
--- a/Engine/Plugins/AI/MassAI/Source/MassAIBehavior/Private/MassLookAtSubsystem.cpp
+++ b/Engine/Plugins/AI/MassAI/Source/MassAIBehavior/Private/MassLookAtSubsystem.cpp
@@ -657,7 +657,7 @@ FString UMassLookAtSubsystem::DebugGetRequestsString(const FMassEntityHandle InE
 	const int32* ViewerDataIndex = ViewerHandleToIndexMap.Find(InEntity);
 
 	TStringBuilder<256> Builder;
-	if (*ViewerDataIndex && *ViewerDataIndex != INDEX_NONE)
+	if (ViewerDataIndex && *ViewerDataIndex != INDEX_NONE)
 	{
 		for (const int32 RequestIndex : PerViewerRequests[*ViewerDataIndex].RequestIndices)
 		{
