[MassAI] forced StateTreeDebug processor to run on game thread to solve threading issue when writing entity details in the debugger

diff --git a/Engine/Plugins/AI/MassAI/Source/MassAIDebug/Private/MassDebugStateTreeProcessor.cpp b/Engine/Plugins/AI/MassAI/Source/MassAIDebug/Private/MassDebugStateTreeProcessor.cpp
--- a/Engine/Plugins/AI/MassAI/Source/MassAIDebug/Private/MassDebugStateTreeProcessor.cpp
+++ b/Engine/Plugins/AI/MassAI/Source/MassAIDebug/Private/MassDebugStateTreeProcessor.cpp
@@ -19,6 +19,9 @@ UMassDebugStateTreeProcessor::UMassDebugStateTreeProcessor()
 {
 	ExecutionOrder.ExecuteInGroup = UE::Mass::ProcessorGroupNames::Behavior;
 	ExecutionOrder.ExecuteAfter.Add(TEXT("MassStateTreeProcessor"));
+
+	// Run on game thread to avoid race condition for AppendSelectedEntityInfo
+	bRequiresGameThreadExecution = true;
 }
 
 void UMassDebugStateTreeProcessor::ConfigureQueries(const TSharedRef<FMassEntityManager>& EntityManager)
