Mass State Tree: Added base classes for Mass State Tree conditions and property functions to support Mass dependencies. Updated Mass State Tree schema so it allows only Mass and common State Tree structs.
#rnx

diff --git a/Engine/Plugins/AI/MassAI/Source/MassAIBehavior/Private/MassStateTreeSchema.cpp b/Engine/Plugins/AI/MassAI/Source/MassAIBehavior/Private/MassStateTreeSchema.cpp
--- a/Engine/Plugins/AI/MassAI/Source/MassAIBehavior/Private/MassStateTreeSchema.cpp
+++ b/Engine/Plugins/AI/MassAI/Source/MassAIBehavior/Private/MassStateTreeSchema.cpp
@@ -17,12 +17,15 @@
 
 bool UMassStateTreeSchema::IsStructAllowed(const UScriptStruct* InScriptStruct) const
 {
-	// Only allow Mass evals and tasks,and common conditions.
+	// Only allow Mass and common structs.
 	return InScriptStruct->IsChildOf(FMassStateTreeEvaluatorBase::StaticStruct())
-			|| InScriptStruct->IsChildOf(FStateTreeEvaluatorCommonBase::StaticStruct())
 			|| InScriptStruct->IsChildOf(FMassStateTreeTaskBase::StaticStruct())
-			|| InScriptStruct->IsChildOf(FStateTreeConditionBase::StaticStruct())
-			|| InScriptStruct->IsChildOf(FStateTreeConsiderationBase::StaticStruct());
+			|| InScriptStruct->IsChildOf(FMassStateTreeConditionBase::StaticStruct())
+			|| InScriptStruct->IsChildOf(FMassStateTreePropertyFunctionBase::StaticStruct())
+			|| InScriptStruct->IsChildOf(FStateTreeEvaluatorCommonBase::StaticStruct())
+			|| InScriptStruct->IsChildOf(FStateTreeConditionCommonBase::StaticStruct())
+			|| InScriptStruct->IsChildOf(FStateTreeConsiderationCommonBase::StaticStruct())
+			|| InScriptStruct->IsChildOf(FStateTreePropertyFunctionCommonBase::StaticStruct());
 }
 
 bool UMassStateTreeSchema::IsExternalItemAllowed(const UStruct& InStruct) const
@@ -51,6 +54,14 @@ bool UMassStateTreeSchema::Link(FStateTreeLinker& Linker)
 				{
 					Task->GetDependencies(Builder);
 				}
+				else if (const FMassStateTreeConditionBase* Condition = Node.GetPtr<const FMassStateTreeConditionBase>())
+				{
+					Condition->GetDependencies(Builder);
+				}
+				else if (const FMassStateTreePropertyFunctionBase* PropertyFunction = Node.GetPtr<const FMassStateTreePropertyFunctionBase>())
+				{
+					PropertyFunction->GetDependencies(Builder);
+				}
 			}
 		};
 
diff --git a/Engine/Plugins/AI/MassAI/Source/MassAIBehavior/Public/Conditions/ZoneGraphTagConditions.h b/Engine/Plugins/AI/MassAI/Source/MassAIBehavior/Public/Conditions/ZoneGraphTagConditions.h
--- a/Engine/Plugins/AI/MassAI/Source/MassAIBehavior/Public/Conditions/ZoneGraphTagConditions.h
+++ b/Engine/Plugins/AI/MassAI/Source/MassAIBehavior/Public/Conditions/ZoneGraphTagConditions.h
@@ -3,7 +3,7 @@
 #pragma once
 
 #include "CoreMinimal.h"
-#include "StateTreeConditionBase.h"
+#include "MassStateTreeTypes.h"
 #include "ZoneGraphTypes.h"
 #include "ZoneGraphTagConditions.generated.h"
 
@@ -23,7 +23,7 @@ struct FZoneGraphTagFilterConditionInstanceData
 };
 
 USTRUCT(DisplayName="ZoneGraphTagFilter Compare")
-struct FZoneGraphTagFilterCondition : public FStateTreeConditionBase
+struct FZoneGraphTagFilterCondition : public FMassStateTreeConditionBase
 {
 	GENERATED_BODY()
 
@@ -58,7 +58,7 @@ struct FZoneGraphTagMaskConditionInstanceData
 };
 
 USTRUCT(DisplayName="ZoneGraphTagMask Compare")
-struct FZoneGraphTagMaskCondition : public FStateTreeConditionBase
+struct FZoneGraphTagMaskCondition : public FMassStateTreeConditionBase
 {
 	GENERATED_BODY()
 
@@ -93,7 +93,7 @@ struct FZoneGraphTagConditionInstanceData
 };
 
 USTRUCT(DisplayName="ZoneGraphTag Compare")
-struct FZoneGraphTagCondition : public FStateTreeConditionBase
+struct FZoneGraphTagCondition : public FMassStateTreeConditionBase
 {
 	GENERATED_BODY()
 
diff --git a/Engine/Plugins/AI/MassAI/Source/MassAIBehavior/Public/MassStateTreeTypes.h b/Engine/Plugins/AI/MassAI/Source/MassAIBehavior/Public/MassStateTreeTypes.h
--- a/Engine/Plugins/AI/MassAI/Source/MassAIBehavior/Public/MassStateTreeTypes.h
+++ b/Engine/Plugins/AI/MassAI/Source/MassAIBehavior/Public/MassStateTreeTypes.h
@@ -2,7 +2,9 @@
 
 #pragma once
 
+#include "StateTreeConditionBase.h"
 #include "StateTreeEvaluatorBase.h" 
+#include "StateTreePropertyFunctionBase.h"
 #include "StateTreeTaskBase.h"
 #if UE_ENABLE_INCLUDE_ORDER_DEPRECATED_IN_5_6
 #include "CoreMinimal.h"
@@ -32,14 +34,14 @@ namespace UE::Mass::Signals
 /**
  * Base struct for all Mass StateTree Evaluators.
  */
-USTRUCT(meta = (DisplayName = "Mass Evaluator Base"))
+USTRUCT(meta = (Hidden, DisplayName = "Mass Evaluator Base"))
 struct FMassStateTreeEvaluatorBase : public FStateTreeEvaluatorBase
 {
 	GENERATED_BODY()
 
 	/**
 	 * Appends this evaluator's Mass dependencies to the given Builder.
-	 * This is done once once for every evaluator instance, when the state tree asset is loaded or compiled.
+	 * This is done once for every evaluator instance, when the state tree asset is loaded or compiled.
 	 */
 	virtual void GetDependencies(UE::MassBehavior::FStateTreeDependencyBuilder& Builder) const
 	{
@@ -49,14 +51,48 @@ struct FMassStateTreeEvaluatorBase : public FStateTreeEvaluatorBase
 /**
  * Base struct for all Mass StateTree Tasks.
  */
-USTRUCT(meta = (DisplayName = "Mass Task Base"))
+USTRUCT(meta = (Hidden, DisplayName = "Mass Task Base"))
 struct FMassStateTreeTaskBase : public FStateTreeTaskBase
 {
 	GENERATED_BODY()
 
 	/**
 	 * Appends this task's Mass dependencies to the given Builder.
-	 * This is done once once for every evaluator instance, when the state tree asset is loaded or compiled.
+	 * This is done once for every task instance, when the state tree asset is loaded or compiled.
+	 */
+	virtual void GetDependencies(UE::MassBehavior::FStateTreeDependencyBuilder& Builder) const
+	{
+	}
+};
+
+/**
+ * Base struct for all Mass StateTree Conditions.
+ */
+USTRUCT(meta = (Hidden, DisplayName = "Mass Condition Base"))
+struct FMassStateTreeConditionBase : public FStateTreeConditionBase
+{
+	GENERATED_BODY()
+
+	/**
+	 * Appends this condition's Mass dependencies to the given Builder.
+	 * This is done once for every condition instance, when the state tree asset is loaded or compiled.
+	 */
+	virtual void GetDependencies(UE::MassBehavior::FStateTreeDependencyBuilder& Builder) const
+	{
+	}
+};
+
+/**
+ * Base struct for all Mass StateTree Property Functions.
+ */
+USTRUCT(meta = (Hidden, DisplayName = "Mass Property Function Base"))
+struct FMassStateTreePropertyFunctionBase : public FStateTreePropertyFunctionBase
+{
+	GENERATED_BODY()
+
+	/**
+	 * Appends this property function's Mass dependencies to the given Builder.
+	 * This is done once for every property function instance, when the state tree asset is loaded or compiled.
 	 */
 	virtual void GetDependencies(UE::MassBehavior::FStateTreeDependencyBuilder& Builder) const
 	{
