[MassStateTree] added thread access detectors

diff --git a/Engine/Plugins/AI/MassAI/Source/MassAIBehavior/Private/MassStateTreeSubsystem.cpp b/Engine/Plugins/AI/MassAI/Source/MassAIBehavior/Private/MassStateTreeSubsystem.cpp
--- a/Engine/Plugins/AI/MassAI/Source/MassAIBehavior/Private/MassStateTreeSubsystem.cpp
+++ b/Engine/Plugins/AI/MassAI/Source/MassAIBehavior/Private/MassStateTreeSubsystem.cpp
@@ -46,7 +46,8 @@ FMassStateTreeInstanceHandle UMassStateTreeSubsystem::AllocateInstanceData(const
 	{
 		return FMassStateTreeInstanceHandle();
 	}
-	
+
+	UE_MT_SCOPED_WRITE_ACCESS(InstanceDataMTDetector);
 	int32 Index = 0;
 	if (InstanceDataFreelist.Num() > 0)
 	{
@@ -77,6 +78,7 @@ void UMassStateTreeSubsystem::FreeInstanceData(const FMassStateTreeInstanceHandl
 		return;
 	}
 
+	UE_MT_SCOPED_WRITE_ACCESS(InstanceDataMTDetector);
 	FMassStateTreeInstanceDataItem& Item = InstanceDataArray[Handle.GetIndex()];
 	Item.InstanceData.Reset();
 	Item.Generation++;
diff --git a/Engine/Plugins/AI/MassAI/Source/MassAIBehavior/Public/MassStateTreeSubsystem.h b/Engine/Plugins/AI/MassAI/Source/MassAIBehavior/Public/MassStateTreeSubsystem.h
--- a/Engine/Plugins/AI/MassAI/Source/MassAIBehavior/Public/MassStateTreeSubsystem.h
+++ b/Engine/Plugins/AI/MassAI/Source/MassAIBehavior/Public/MassStateTreeSubsystem.h
@@ -63,12 +63,14 @@ public:
 	/** @return Pointer to instance data held by the handle, or nullptr if handle is not valid. */
 	FStateTreeInstanceData* GetInstanceData(const FMassStateTreeInstanceHandle Handle)
 	{
+		UE_MT_SCOPED_READ_ACCESS(InstanceDataMTDetector);
 		return IsValidHandle(Handle) ? &InstanceDataArray[Handle.GetIndex()].InstanceData : nullptr;
 	}
 
 	/** @return True if the handle points to active instance data. */
 	bool IsValidHandle(const FMassStateTreeInstanceHandle Handle) const
 	{
+		UE_MT_SCOPED_READ_ACCESS(InstanceDataMTDetector);
 		return InstanceDataArray.IsValidIndex(Handle.GetIndex()) && InstanceDataArray[Handle.GetIndex()].Generation == Handle.GetGeneration();
 	}
 	
@@ -85,6 +87,14 @@ protected:
 	UPROPERTY(Transient)
 	TArray<FMassStateTreeInstanceDataItem> InstanceDataArray;
 
+	/**
+	 * Multithread access detector to prevent attempts to allocate new instance data
+	 * from multiple threads, which is not supported by the current implementation.
+	 * @see AllocateInstanceData, FreeInstanceData, CreateProcessorForStateTree
+	 * @todo Instance data creation needs to be refactored in order to allow parallelization of StateTree activations
+	 */
+	UE_MT_DECLARE_RW_ACCESS_DETECTOR(InstanceDataMTDetector);
+
 	/**
 	 * The relevant Entity Manager. Needed to build processing requirements for dynamic processors.
 	 */
