[MassStateTree] replaced recently aded MTDetector by the transactionally safe version

diff --git a/Engine/Plugins/AI/MassAI/Source/MassAIBehavior/Public/MassStateTreeSubsystem.h b/Engine/Plugins/AI/MassAI/Source/MassAIBehavior/Public/MassStateTreeSubsystem.h
--- a/Engine/Plugins/AI/MassAI/Source/MassAIBehavior/Public/MassStateTreeSubsystem.h
+++ b/Engine/Plugins/AI/MassAI/Source/MassAIBehavior/Public/MassStateTreeSubsystem.h
@@ -5,7 +5,10 @@
 #include "MassStateTreeTypes.h"
 #include "MassSubsystemBase.h"
 #include "StateTreeExecutionContext.h"
+#if UE_ENABLE_INCLUDE_ORDER_DEPRECATED_IN_5_7
 #include "MassExternalSubsystemTraits.h"
+#endif // UE_ENABLE_INCLUDE_ORDER_DEPRECATED_IN_5_7
+#include "Misc/MTTransactionallySafeAccessDetector.h"
 #include "UObject/ObjectKey.h"
 #include "MassStateTreeSubsystem.generated.h"
 
@@ -93,7 +96,7 @@ protected:
 	 * @see AllocateInstanceData, FreeInstanceData, CreateProcessorForStateTree
 	 * @todo Instance data creation needs to be refactored in order to allow parallelization of StateTree activations
 	 */
-	UE_MT_DECLARE_RW_ACCESS_DETECTOR(InstanceDataMTDetector);
+	UE_MT_DECLARE_TS_RW_ACCESS_DETECTOR(InstanceDataMTDetector);
 
 	/**
 	 * The relevant Entity Manager. Needed to build processing requirements for dynamic processors.
