[MassStateTree] fix for potential threading issue when gathering entities to signal during parallel updates
Created dedicated code path to avoid using threadsafe container in a code path used with AutoRTFM

diff --git a/Engine/Plugins/AI/MassAI/Source/MassAIBehavior/Private/MassStateTreeProcessors.cpp b/Engine/Plugins/AI/MassAI/Source/MassAIBehavior/Private/MassStateTreeProcessors.cpp
--- a/Engine/Plugins/AI/MassAI/Source/MassAIBehavior/Private/MassStateTreeProcessors.cpp
+++ b/Engine/Plugins/AI/MassAI/Source/MassAIBehavior/Private/MassStateTreeProcessors.cpp
@@ -1,24 +1,24 @@
 // Copyright Epic Games, Inc. All Rights Reserved.
 
 #include "MassStateTreeProcessors.h"
-#include "StateTree.h"
-#include "MassStateTreeExecutionContext.h"
+#include "Containers/ConsumeAllMpmcQueue.h"
+#include "Engine/World.h"
+#include "Experimental/ConcurrentLinearAllocator.h"
+#include "MassBehaviorSettings.h"
+#include "MassComponentHitTypes.h"
+#include "MassDebugger.h"
 #include "MassExecutionContext.h"
 #include "MassNavigationTypes.h"
+#include "MassSignalSubsystem.h"
 #include "MassSimulationLOD.h"
-#include "MassComponentHitTypes.h"
 #include "MassSmartObjectTypes.h"
-#include "MassZoneGraphAnnotationTypes.h"
-#include "MassStateTreeSubsystem.h"
+#include "MassStateTreeExecutionContext.h"
 #include "MassStateTreeFragments.h"
-#include "MassSignalSubsystem.h"
+#include "MassStateTreeSubsystem.h"
+#include "MassZoneGraphAnnotationTypes.h"
 #include "ProfilingDebugging/CsvProfiler.h"
-#include "Engine/World.h"
-#include "MassBehaviorSettings.h"
+#include "StateTree.h"
 #include "VisualLogger/VisualLogger.h"
-#include "MassDebugger.h"
-#include "MassStateTreeSubsystem.h"
-
 
 CSV_DEFINE_CATEGORY(StateTreeProcessor, true);
 
@@ -64,9 +64,30 @@ void ForEachEntityInChunk(
 		}
 	}
 }
-
 } // UE::MassBehavior
 
+namespace UE::MassStateTree
+{
+struct FBlockAllocationTag : FDefaultBlockAllocationTag
+{
+	static constexpr const char* TagName = "MassStateTreeLinear";
+	using Allocator = FAlignedAllocator;
+};
+
+struct FLinearQueueAllocator
+{
+	FORCEINLINE static void* Malloc(const SIZE_T Size, const uint32 Alignment)
+	{
+		return TConcurrentLinearAllocator<FBlockAllocationTag>::Malloc(Size, Alignment);
+	}
+
+	FORCEINLINE static void Free(void* Pointer)
+	{
+		TConcurrentLinearAllocator<FBlockAllocationTag>::Free(Pointer);
+	}
+};
+
+} // UE::MassStateTree
 
 //----------------------------------------------------------------------//
 // UMassStateTreeFragmentDestructor
@@ -194,8 +215,7 @@ void UMassStateTreeActivationProcessor::Execute(FMassEntityManager& EntityManage
 					StateTreeFragment.LastUpdateTimeInSeconds = TimeInSeconds;
 				});
 
-			// Adding a tag on each entities to remember we have sent the state tree initialization signal
-			EntitiesToSignal.Reserve(EntitiesToSignal.Num() + NumEntities);
+			// Adding a tag on each entity to remember we have sent the state tree initialization signal
 			for (FMassExecutionContext::FEntityIterator EntityIt = Context.CreateEntityIterator(); EntityIt; ++EntityIt)
 			{
 				const FMassStateTreeInstanceFragment& StateTreeInstance = StateTreeInstanceList[EntityIt];
@@ -287,9 +307,7 @@ void UMassStateTreeProcessor::SignalEntities(FMassEntityManager& EntityManager,
 
 	const double TimeInSeconds = EntityManager.GetWorld()->GetTimeSeconds();
 
-	TArray<FMassEntityHandle> EntitiesToSignal;
-
-	auto TickFunc = [TimeInSeconds, &EntitiesToSignal, &EntitySignals](FMassExecutionContext& ExecutionContext)
+	auto TickFunc = [TimeInSeconds, &EntitySignals](FMassExecutionContext& ExecutionContext, TFunctionRef<void(FMassEntityHandle)> AddEntityToSignalFunc)
 		{
 			// Keep stats regarding the amount of tree instances ticked per frame
 			CSV_CUSTOM_STAT(StateTreeProcessor, NumTickedStateTree, ExecutionContext.GetNumEntities(), ECsvCustomStatOp::Accumulate);
@@ -297,7 +315,7 @@ void UMassStateTreeProcessor::SignalEntities(FMassEntityManager& EntityManager,
 			UMassStateTreeSubsystem& MassStateTreeSubsystem = ExecutionContext.GetMutableSubsystemChecked<UMassStateTreeSubsystem>();
 
 			UE::MassBehavior::ForEachEntityInChunk(ExecutionContext, MassStateTreeSubsystem,
-				[TimeInSeconds, &EntitiesToSignal, &EntitySignals, &MassStateTreeSubsystem]
+				[TimeInSeconds, AddEntityToSignalFunc, &EntitySignals, &MassStateTreeSubsystem]
 				(FMassStateTreeExecutionContext& StateTreeExecutionContext, FMassStateTreeInstanceFragment& StateTreeFragment)
 				{
 					// Compute adjusted delta time
@@ -335,22 +353,42 @@ void UMassStateTreeProcessor::SignalEntities(FMassEntityManager& EntityManager,
 						// Could not find new state yet, try again next frame
 						if (StateTreeExecutionContext.GetLastTickStatus() != EStateTreeRunStatus::Running)
 						{
-							EntitiesToSignal.Add(StateTreeExecutionContext.GetEntity());
+							AddEntityToSignalFunc(StateTreeExecutionContext.GetEntity());
 						}
 					}
 				});
 		};
 
+	TArray<FMassEntityHandle, FNonconcurrentLinearArrayAllocator> EntitiesToSignal;
 	if (bProcessEntitiesInParallel)
 	{
-		EntityQuery.ParallelForEachEntityChunk(Context, TickFunc);
+		UE::TConsumeAllMpmcQueue<FMassEntityHandle, UE::MassStateTree::FLinearQueueAllocator> EntitiesToSignalQueue;
+
+		EntityQuery.ParallelForEachEntityChunk(Context, [TickFunc, &EntitiesToSignalQueue](FMassExecutionContext& ExecutionContext)
+			{
+				TickFunc(ExecutionContext, [&EntitiesToSignalQueue](FMassEntityHandle Entity)
+					{
+						EntitiesToSignalQueue.ProduceItem(Entity);
+					});
+			});
+
+		EntitiesToSignalQueue.ConsumeAllFifo([&EntitiesToSignal](FMassEntityHandle&& Entity)
+			{
+				EntitiesToSignal.Add(MoveTemp(Entity));
+			});
 	}
 	else
 	{
-		EntityQuery.ForEachEntityChunk(Context, TickFunc);
+		EntityQuery.ForEachEntityChunk(Context, [TickFunc, &EntitiesToSignal](FMassExecutionContext& ExecutionContext)
+			{
+				TickFunc(ExecutionContext, [&EntitiesToSignal](FMassEntityHandle Entity)
+					{
+						EntitiesToSignal.Add(Entity);
+					});
+			});
 	}
 
-	if (EntitiesToSignal.Num())
+	if (EntitiesToSignal.Num() > 0)
 	{
 		SignalSubsystem.SignalEntities(UE::Mass::Signals::NewStateTreeTaskRequired, EntitiesToSignal);
 	}
