Integrating avoidance changes
(from other depot: 8710728, 8743466, 8763847, 8775723)
Fix to apply max speed in UMassApplyForceProcessor
Added debug to display standing entities in mass avoidance
Pathfollow -> Fix to keep the steering force when near the end of a path (we want to apply the fade only on the avoidance forces)
Stronger avoidance forces when getting out of environement boudaries
Added tag to toggle preventing backward while avoiding

diff --git a/Engine/Plugins/AI/MassAI/Source/MassNavigation/Private/MassAvoidanceProcessors.cpp b/Engine/Plugins/AI/MassAI/Source/MassNavigation/Private/MassAvoidanceProcessors.cpp
--- a/Engine/Plugins/AI/MassAI/Source/MassNavigation/Private/MassAvoidanceProcessors.cpp
+++ b/Engine/Plugins/AI/MassAI/Source/MassNavigation/Private/MassAvoidanceProcessors.cpp
@@ -177,7 +177,9 @@ namespace UE::MassAvoidance
 
 	// Agents colors
 	static constexpr FColor AgentsColor = Amber;
+	static constexpr FColor AgentsNotMovingColor = Blue;
 	static constexpr FColor AgentAvoidForceColor = Orange;
+	static constexpr FColor AgentsNotMovingSeparationColor = Blue;
 	static constexpr FColor AgentSeparationForceColor = OrangeRed;
 	
 	// Obstacles colors
@@ -249,6 +251,24 @@ namespace UE::MassAvoidance
 		DebugDrawArrow(Context, Start + FVector(0.,0.,1.), End + FVector(0., 0., 1.), Color, /*HeadSize*/8.f, SummedForcesThickness);
 	}
 
+	static void DebugDrawHitPosition(const FDebugContext& Context, const FVector& Location, const FVector& Velocity,
+		const FVector::FReal CPA, const FVector::FReal Radius, const FColor Color)
+	{
+		if (!Context.ShouldLogEntity())
+		{
+			return;
+		}
+
+		const FVector HitPosition = Location + (Velocity * CPA);
+		const FVector LeftOffset = Radius * UE::MassNavigation::GetLeftDirection(Velocity.GetSafeNormal(), FVector::UpVector);
+		UE::MassAvoidance::DebugDrawLine(Context, Location + UE::MassAvoidance::DebugAgentHeightOffset + LeftOffset,
+			HitPosition + UE::MassAvoidance::DebugAgentHeightOffset + LeftOffset, Color, /*Thickness*/1.5f);
+		UE::MassAvoidance::DebugDrawLine(Context, Location + UE::MassAvoidance::DebugAgentHeightOffset - LeftOffset,
+			HitPosition + UE::MassAvoidance::DebugAgentHeightOffset - LeftOffset, Color, /*Thickness*/1.5f);
+		UE::MassAvoidance::DebugDrawCylinder(Context, HitPosition,
+			HitPosition + UE::MassAvoidance::DebugAgentHeightOffset, Radius, Color);
+	}
+
 #endif // WITH_MASSGAMEPLAY_DEBUG
 
 } // namespace UE::MassAvoidance
@@ -269,6 +289,7 @@ UMassMovingAvoidanceProcessor::UMassMovingAvoidanceProcessor()
 void UMassMovingAvoidanceProcessor::ConfigureQueries(const TSharedRef<FMassEntityManager>& EntityManager)
 {
 	EntityQuery.AddRequirement<FMassForceFragment>(EMassFragmentAccess::ReadWrite);
+	EntityQuery.AddRequirement<FMassDesiredMovementFragment>(EMassFragmentAccess::ReadWrite, EMassFragmentPresence::Optional);
 	EntityQuery.AddRequirement<FMassNavigationEdgesFragment>(EMassFragmentAccess::ReadOnly);
 	EntityQuery.AddRequirement<FMassMoveTargetFragment>(EMassFragmentAccess::ReadOnly);
 	EntityQuery.AddRequirement<FTransformFragment>(EMassFragmentAccess::ReadOnly);
@@ -278,6 +299,7 @@ void UMassMovingAvoidanceProcessor::ConfigureQueries(const TSharedRef<FMassEntit
 	EntityQuery.AddTagRequirement<FMassMediumLODTag>(EMassFragmentPresence::None);
 	EntityQuery.AddTagRequirement<FMassLowLODTag>(EMassFragmentPresence::None);
 	EntityQuery.AddTagRequirement<FMassOffLODTag>(EMassFragmentPresence::None);
+	EntityQuery.AddTagRequirement<FMassAvoidancePreventBackwardMovementTag>(EMassFragmentPresence::Optional);
 	EntityQuery.AddConstSharedRequirement<FMassMovingAvoidanceParameters>(EMassFragmentPresence::All);
 	EntityQuery.AddConstSharedRequirement<FMassMovementParameters>(EMassFragmentPresence::All);
 
@@ -309,6 +331,8 @@ void UMassMovingAvoidanceProcessor::Execute(FMassEntityManager& EntityManager, F
 		const double CurrentTime = World->GetTimeSeconds();
 		
 		const TArrayView<FMassForceFragment> ForceList = Context.GetMutableFragmentView<FMassForceFragment>();
+		const TArrayView<FMassDesiredMovementFragment> DesiredMovementList = Context.GetMutableFragmentView<FMassDesiredMovementFragment>();
+		const bool bHasDesiredMovementFragments = !DesiredMovementList.IsEmpty();
 		const TConstArrayView<FMassNavigationEdgesFragment> NavEdgesList = Context.GetFragmentView<FMassNavigationEdgesFragment>();
 		const TConstArrayView<FTransformFragment> LocationList = Context.GetFragmentView<FTransformFragment>();
 		const TConstArrayView<FMassVelocityFragment> VelocityList = Context.GetFragmentView<FMassVelocityFragment>();
@@ -317,8 +341,10 @@ void UMassMovingAvoidanceProcessor::Execute(FMassEntityManager& EntityManager, F
 		const TConstArrayView<FMassMoveTargetFragment> MoveTargetList = Context.GetFragmentView<FMassMoveTargetFragment>();
 		const FMassMovingAvoidanceParameters& MovingAvoidanceParams = Context.GetConstSharedFragment<FMassMovingAvoidanceParameters>();
 		const FMassMovementParameters& MovementParams = Context.GetConstSharedFragment<FMassMovementParameters>();
+		const bool bPreventBackwardMovement = Context.DoesArchetypeHaveTag<FMassAvoidancePreventBackwardMovementTag>();
 
 		const FVector::FReal InvPredictiveAvoidanceTime = 1. / MovingAvoidanceParams.PredictiveAvoidanceTime;
+		const FVector::FReal InvEnvironmentPredictiveAvoidanceTime = 1. / MovingAvoidanceParams.EnvironmentPredictiveAvoidanceTime;
 
 		// Arrays used to store close obstacles
 		TArray<FMassNavigationObstacleItem, TFixedAllocator<UE::MassAvoidance::MaxObstacleResults>> CloseEntities;
@@ -575,14 +601,15 @@ void UMassMovingAvoidanceProcessor::Execute(FMassEntityManager& EntityManager, F
 						{
 							// Draw environment predictive avoidance distance
 							FVector ZOffset = FVector(0., 0., 7.);
-							const FVector AvoidanceOffset = MovingAvoidanceParams.PredictiveAvoidanceDistance * Edge.LeftDir;
+							const FVector AvoidanceOffset = MovingAvoidanceParams.EnvironmentPredictiveAvoidanceDistance * Edge.LeftDir;
 							DebugDrawLine(ObstacleDebugContext, ZOffset + Edge.Start + AvoidanceOffset, ZOffset + Edge.End + AvoidanceOffset,
 								UE::MassAvoidance::ObstacleAvoidForceColor, /*Thickness=*/1.f);
 						}
 #endif // WITH_MASSGAMEPLAY_DEBUG
 
-						const FVector::FReal CPA = UE::MassAvoidance::ComputeClosestPointOfApproach(FVector2D(AgentLocation), FVector2D(DesiredVelocity), AgentRadiusForEnvironment,
-							FVector2D(Edge.Start), FVector2D(Edge.End), MovingAvoidanceParams.PredictiveAvoidanceTime);
+						const FVector::FReal CPA = UE::MassAvoidance::ComputeClosestPointOfApproach(
+							FVector2D(AgentLocation), FVector2D(DesiredVelocity), AgentRadiusForEnvironment,
+							FVector2D(Edge.Start), FVector2D(Edge.End), MovingAvoidanceParams.EnvironmentPredictiveAvoidanceTime);
 						const FVector HitAgentPos = AgentLocation + DesiredVelocity * CPA;
 						const FVector::FReal EdgeT = UE::MassNavigation::ProjectPtSeg(FVector2D(HitAgentPos), FVector2D(Edge.Start), FVector2D(Edge.End));
 						const FVector HitObPos = FMath::Lerp(Edge.Start, Edge.End, EdgeT);
@@ -593,14 +620,15 @@ void UMassMovingAvoidanceProcessor::Execute(FMassEntityManager& EntityManager, F
 						const FVector::FReal AvoidDist = AvoidRelPos.Size();
 						const FVector AvoidNormal = AvoidDist > UE_KINDA_SMALL_NUMBER ? (AvoidRelPos / AvoidDist) : Edge.LeftDir;
 
-						const FVector::FReal AvoidPen = (EnvironmentPredictiveAvoidanceAgentRadius + MovingAvoidanceParams.PredictiveAvoidanceDistance) - AvoidDist;
+						const FVector::FReal AvoidPen = (EnvironmentPredictiveAvoidanceAgentRadius + MovingAvoidanceParams.EnvironmentPredictiveAvoidanceDistance) - AvoidDist;
 						FVector::FReal AvoidMag = 1.;
-						if (MovingAvoidanceParams.PredictiveAvoidanceDistance != 0.)
+
+						if (MovingAvoidanceParams.EnvironmentPredictiveAvoidanceDistance != 0.)
 						{
-							AvoidMag = FMath::Square(FMath::Clamp(AvoidPen / MovingAvoidanceParams.PredictiveAvoidanceDistance, 0., 1.));
+							AvoidMag = FMath::Square(FMath::Clamp(AvoidPen / MovingAvoidanceParams.EnvironmentPredictiveAvoidanceDistance, 0., 1.));
 						}
 
-						const FVector::FReal AvoidMagDist = 1. + FMath::Square(1. - CPA * InvPredictiveAvoidanceTime);
+						const FVector::FReal AvoidMagDist = 1. + FMath::Square(1. - CPA * InvEnvironmentPredictiveAvoidanceTime);
 
 						// Predictive avoidance against environment is tuned down towards the end of the path
 						const FVector AvoidForce = AvoidNormal * AvoidMag * AvoidMagDist * MovingAvoidanceParams.EnvironmentPredictiveAvoidanceStiffness * NearEndScaling; 
@@ -654,15 +682,25 @@ void UMassMovingAvoidanceProcessor::Execute(FMassEntityManager& EntityManager, F
 				
 				// Process contacts to add edge separation force
 				const FVector SteeringForceBeforeSeparation = SteeringForce;
+
+				const FVector::FReal ValidBehindDistance = MovingAvoidanceParams.EnvironmentSeparationBehindEdgeDistanceScale * AgentRadius;
+				
 				for (const FEnvironmentContact& Contact : Contacts) 
 				{
-					if (Contact.bAlongTheEdge && Contact.Distance > -AgentRadius)
+					if (Contact.bAlongTheEdge && Contact.Distance > -ValidBehindDistance)
 					{
 						// Separation force (stay away from obstacles if possible)
 						const FVector::FReal SeparationPenalty = (EnvironmentSeparationAgentRadius + MovingAvoidanceParams.EnvironmentSeparationDistance) - Contact.Distance;
 
 						FVector::FReal SeparationMag;
-						if (MovingAvoidanceParams.EnvironmentSeparationDistance != 0.)
+
+						const FVector::FReal BehindDistanceThreshold = 0.25*AgentRadius;
+						if (Contact.Distance < -BehindDistanceThreshold)
+						{
+							// Far behind the edge, use a big magnitude to trump other forces.
+							SeparationMag = 10.;
+						}
+						else if (MovingAvoidanceParams.EnvironmentSeparationDistance != 0.)
 						{
 							SeparationMag = UE::MassNavigation::Smooth(FMath::Clamp(SeparationPenalty / MovingAvoidanceParams.EnvironmentSeparationDistance, 0., 1.));
 						}
@@ -966,31 +1004,29 @@ void UMassMovingAvoidanceProcessor::Execute(FMassEntityManager& EntityManager, F
 					UE::MassAvoidance::DebugDrawForce(AgentDebugContext,
 						Collider.Location + UE::MassAvoidance::DebugAgentSeparationHeightOffset,
 						Collider.Location + UE::MassAvoidance::DebugAgentSeparationHeightOffset + SeparationForce,
-						UE::MassAvoidance::AgentSeparationForceColor, UE::MassAvoidance::SeparationThickness); 
+						Collider.bIsMoving ? UE::MassAvoidance::AgentSeparationForceColor : UE::MassAvoidance::AgentsNotMovingSeparationColor,
+						UE::MassAvoidance::SeparationThickness);
+				}
+
+				if (UE::MassAvoidance::Tweakables::bEnableDetailedDebug)
+				{
+					UE::MassAvoidance::DebugDrawCircle(AgentDebugContext, Collider.Location, Collider.Radius + MovingAvoidanceParams.PredictiveAvoidanceDistance, UE::MassAvoidance::AgentsColor);
+					UE::MassAvoidance::DebugDrawCircle(AgentDebugContext, Collider.Location, Collider.Radius + MovingAvoidanceParams.ObstacleSeparationDistance, UE::MassAvoidance::AgentSeparationForceColor);
 				}
 				
 				if (AvoidForce.Size() > 0.)
 				{
 					// Draw agent vs agent hit positions
-					const FVector HitPosition = AgentLocation + (DesVel * CPA);
-					const FVector LeftOffset = PredictiveAvoidanceAgentRadius * UE::MassNavigation::GetLeftDirection(DesVel.GetSafeNormal(), FVector::UpVector);
-					UE::MassAvoidance::DebugDrawLine(AgentDebugContext, AgentLocation + UE::MassAvoidance::DebugAgentHeightOffset + LeftOffset,
-						HitPosition + UE::MassAvoidance::DebugAgentHeightOffset + LeftOffset, UE::MassAvoidance::CurrentAgentColor, 1.5f);
-					UE::MassAvoidance::DebugDrawLine(AgentDebugContext, AgentLocation + UE::MassAvoidance::DebugAgentHeightOffset - LeftOffset,
-						HitPosition + UE::MassAvoidance::DebugAgentHeightOffset - LeftOffset, UE::MassAvoidance::CurrentAgentColor, 1.5f);
-					UE::MassAvoidance::DebugDrawCylinder(AgentDebugContext, HitPosition,
-						HitPosition + UE::MassAvoidance::DebugAgentHeightOffset, PredictiveAvoidanceAgentRadius, UE::MassAvoidance::CurrentAgentColor);
+					UE::MassAvoidance::DebugDrawHitPosition(AgentDebugContext, AgentLocation, DesVel, CPA, PredictiveAvoidanceAgentRadius, UE::MassAvoidance::CurrentAgentColor);
+
+					// Draw collider hit position
+					const FColor ColliderColor = Collider.bIsMoving ? UE::MassAvoidance::AgentsColor : UE::MassAvoidance::AgentsNotMovingColor;
+					UE::MassAvoidance::DebugDrawHitPosition(AgentDebugContext, Collider.Location, Collider.Velocity, CPA, Collider.Radius, ColliderColor);
 
 					const FVector OtherHitPosition = Collider.Location + (Collider.Velocity * CPA);
-					const FVector OtherLeftOffset = Collider.Radius * UE::MassNavigation::GetLeftDirection(Collider.Velocity.GetSafeNormal(), FVector::UpVector);
-					const FVector Left = UE::MassAvoidance::DebugAgentHeightOffset + OtherLeftOffset;
-					const FVector Right = UE::MassAvoidance::DebugAgentHeightOffset - OtherLeftOffset;
-					UE::MassAvoidance::DebugDrawLine(AgentDebugContext, Collider.Location + Left, OtherHitPosition + Left, UE::MassAvoidance::AgentsColor, 1.5f);
-					UE::MassAvoidance::DebugDrawLine(AgentDebugContext, Collider.Location + Right, OtherHitPosition + Right, UE::MassAvoidance::AgentsColor, 1.5f);
+
 					UE::MassAvoidance::DebugDrawCylinder(AgentDebugContext, Collider.Location, Collider.Location + UE::MassAvoidance::DebugAgentHeightOffset,
-						AgentRadius, UE::MassAvoidance::AgentsColor);
-					UE::MassAvoidance::DebugDrawCylinder(AgentDebugContext, OtherHitPosition, OtherHitPosition + UE::MassAvoidance::DebugAgentHeightOffset,
-						AgentRadius, UE::MassAvoidance::AgentsColor);
+						AgentRadius, ColliderColor);
 
 					// Draw agent avoid force
 					UE::MassAvoidance::DebugDrawForce(AgentDebugContext,
@@ -1001,7 +1037,38 @@ void UMassMovingAvoidanceProcessor::Execute(FMassEntityManager& EntityManager, F
 #endif // WITH_MASSGAMEPLAY_DEBUG
 			} // close entities loop
 
-			SteeringForce *= NearStartScaling * NearEndScaling;
+			// Only scale the added avoidance force (ignoring the original steering force) and apply it back into the steering force.
+			FVector AccumulatedAvoidanceForce = SteeringForce - Force.Value;
+			AccumulatedAvoidanceForce *= NearStartScaling * NearEndScaling;
+			SteeringForce = Force.Value + AccumulatedAvoidanceForce;
+
+			// Make sure that the avoidance force will not bring the desired velocity backward.
+			if (bPreventBackwardMovement && bHasDesiredMovementFragments)
+			{
+				FMassDesiredMovementFragment& DesiredMovement = DesiredMovementList[EntityIt];
+				
+				const FVector TempDesiredVelocity  = DesiredMovement.DesiredVelocity + (UE::MassNavigation::ClampVector(SteeringForce, MaxSteerAccel) * DeltaTime);
+				const FVector PredictedDesiredVelocity = TempDesiredVelocity.GetClampedToMaxSize(MovementParams.MaxSpeed);
+
+				constexpr FVector::FReal Cos45 = 0.707;
+				const FVector::FReal Dot = PredictedDesiredVelocity.GetSafeNormal().Dot(MoveTarget.Forward);
+				if (Dot < Cos45)
+				{
+					// That force would bring the DesiredVelocity in an unwanted range. 
+					// Ignore it, but keep the component in the direction of movement.
+					const FVector::FReal ForwardComponent = FMath::Abs(SteeringForce.Dot(MoveTarget.Forward));
+					SteeringForce = ForwardComponent * MoveTarget.Forward;
+				}
+
+				// Snap small speeds to 0.
+				constexpr FVector::FReal MinSpeed = 20;
+				const bool bSlowing = PredictedDesiredVelocity.SizeSquared() < DesiredMovement.DesiredVelocity.SizeSquared();
+				if (bSlowing && PredictedDesiredVelocity.Size() < MinSpeed)
+				{
+					DesiredMovement.DesiredVelocity = FVector::ZeroVector;
+					SteeringForce = FVector::ZeroVector;
+				}
+			}
 
 			Force.Value = UE::MassNavigation::ClampVector(SteeringForce, MaxSteerAccel); // Assume unit mass
 
diff --git a/Engine/Plugins/AI/MassAI/Source/MassNavigation/Public/Avoidance/MassAvoidanceFragments.h b/Engine/Plugins/AI/MassAI/Source/MassNavigation/Public/Avoidance/MassAvoidanceFragments.h
--- a/Engine/Plugins/AI/MassAI/Source/MassNavigation/Public/Avoidance/MassAvoidanceFragments.h
+++ b/Engine/Plugins/AI/MassAI/Source/MassNavigation/Public/Avoidance/MassAvoidanceFragments.h
@@ -18,16 +18,19 @@ struct FMassMovingAvoidanceParameters : public FMassConstSharedFragment
 	{
 		FMassMovingAvoidanceParameters Copy = *this;
 		Copy.PredictiveAvoidanceTime = FMath::Max(Copy.PredictiveAvoidanceTime, KINDA_SMALL_NUMBER);
+		Copy.EnvironmentPredictiveAvoidanceTime = FMath::Max(Copy.EnvironmentPredictiveAvoidanceTime, KINDA_SMALL_NUMBER);
 		Copy.ObstacleSeparationDistance = FMath::Max(Copy.ObstacleSeparationDistance, KINDA_SMALL_NUMBER);
 		Copy.PredictiveAvoidanceDistance = FMath::Max(Copy.PredictiveAvoidanceDistance, KINDA_SMALL_NUMBER);
+		Copy.EnvironmentPredictiveAvoidanceDistance = FMath::Max(Copy.EnvironmentPredictiveAvoidanceDistance, KINDA_SMALL_NUMBER);
 		Copy.EnvironmentSeparationDistance = FMath::Max(Copy.EnvironmentSeparationDistance, KINDA_SMALL_NUMBER);
+		Copy.EnvironmentSeparationBehindEdgeDistanceScale = FMath::Max(Copy.EnvironmentSeparationBehindEdgeDistanceScale, KINDA_SMALL_NUMBER);
 		Copy.StartOfPathDuration = FMath::Max(Copy.StartOfPathDuration, KINDA_SMALL_NUMBER);
 		Copy.EndOfPathDuration = FMath::Max(Copy.EndOfPathDuration, KINDA_SMALL_NUMBER);
 
 		return Copy;
 	}
 
-	/** The distance at which neighbour agents are detected. Range: 200...600 */
+	/** The distance at which obstacle are detected. Range: 200...600 */
 	UPROPERTY(EditAnywhere, Category = "General", meta = (ClampMin = "0.0", ForceUnits="cm"))
 	float ObstacleDetectionDistance = 400.f;
 
@@ -76,11 +79,19 @@ struct FMassMovingAvoidanceParameters : public FMassConstSharedFragment
 	UPROPERTY(EditAnywhere, Category = "Separation", meta = (ClampMin = "0", ForceUnits="cm"))
 	float EnvironmentSeparationDistance = 50.f;
 
+	/** Environment separation distance behind edge where separation force is applied. The scale is applied to AgentRadius. Range: 1..2 */
+	UPROPERTY(EditAnywhere, Category = "Separation", meta = (ClampMin = "0", ForceUnits="x"))
+	float EnvironmentSeparationBehindEdgeDistanceScale = 1.f;
+	
 	
 	/** How far in the future the agent reacts to collisions. Range: 1..3, Indoor humans 1.4, outdoor humans 2.4 (seconds). */
 	UPROPERTY(EditAnywhere, Category = "Predictive Avoidance", meta = (ClampMin = "0.1", ForceUnits="s"))
 	float PredictiveAvoidanceTime = 2.5f;
 
+	/** How far in the future the agent reacts to environment collisions. */
+	UPROPERTY(EditAnywhere, Category = "Predictive Avoidance", meta = (ClampMin = "0.1", ForceUnits="s"))
+	float EnvironmentPredictiveAvoidanceTime = 1.f;
+
 	/** Agent radius scale for anticipatory avoidance. Making the scale smaller makes the agent more eager to squeeze through other agents. Range: 0.5..1 */
 	UPROPERTY(EditAnywhere, Category = "Predictive Avoidance", meta = (ClampMin = "0", ForceUnits="x"))
 	float PredictiveAvoidanceRadiusScale = 0.65f;
@@ -89,6 +100,10 @@ struct FMassMovingAvoidanceParameters : public FMassConstSharedFragment
 	UPROPERTY(EditAnywhere, Category = "Predictive Avoidance", meta = (ClampMin = "0", ForceUnits="cm"))
 	float PredictiveAvoidanceDistance = 75.f;
 
+	/** Predictive environment avoidance force effect distance. The avoidance force is applied at the point in future where the agents are closest. The actual observed separation distance will be smaller. Range: 0..100 */
+	UPROPERTY(EditAnywhere, Category = "Predictive Avoidance", meta = (ClampMin = "0", ForceUnits="cm"))
+	float EnvironmentPredictiveAvoidanceDistance = 30.f;
+	
 	/** Predictive avoidance force stiffness between agents and obstacles. Range: 400..1000 N/cm */
 	UPROPERTY(EditAnywhere, Category = "Predictive Avoidance", meta = (ClampMin = "0"))
 	float ObstaclePredictiveAvoidanceStiffness = 700.f;
diff --git a/Engine/Plugins/AI/MassAI/Source/MassNavigation/Public/MassNavigationFragments.h b/Engine/Plugins/AI/MassAI/Source/MassNavigation/Public/MassNavigationFragments.h
--- a/Engine/Plugins/AI/MassAI/Source/MassNavigation/Public/MassNavigationFragments.h
+++ b/Engine/Plugins/AI/MassAI/Source/MassNavigation/Public/MassNavigationFragments.h
@@ -197,6 +197,13 @@ struct FMassInNavigationObstacleGridTag : public FMassTag
 	GENERATED_BODY()
 };
 
+/** Tag used to prevent avoidance from generating backward movements */
+USTRUCT()
+struct FMassAvoidancePreventBackwardMovementTag : public FMassTag
+{
+	GENERATED_BODY()
+};
+
 /**
  * Shared Fragment holding the properties defining how a given entity should affect navigation data
  */
diff --git a/Engine/Plugins/Runtime/MassGameplay/Source/MassMovement/Private/MassMovementProcessors.cpp b/Engine/Plugins/Runtime/MassGameplay/Source/MassMovement/Private/MassMovementProcessors.cpp
--- a/Engine/Plugins/Runtime/MassGameplay/Source/MassMovement/Private/MassMovementProcessors.cpp
+++ b/Engine/Plugins/Runtime/MassGameplay/Source/MassMovement/Private/MassMovementProcessors.cpp
@@ -27,6 +27,8 @@ void UMassApplyForceProcessor::ConfigureQueries(const TSharedRef<FMassEntityMana
 	EntityQuery.AddRequirement<FMassDesiredMovementFragment>(EMassFragmentAccess::ReadWrite);
 	EntityQuery.AddRequirement<FMassForceFragment>(EMassFragmentAccess::ReadWrite);
 	EntityQuery.AddTagRequirement<FMassOffLODTag>(EMassFragmentPresence::None);
+
+	EntityQuery.AddConstSharedRequirement<FMassMovementParameters>(EMassFragmentPresence::All);
 }
 
 void UMassApplyForceProcessor::Execute(FMassEntityManager& EntityManager, FMassExecutionContext& Context)
@@ -39,13 +41,16 @@ void UMassApplyForceProcessor::Execute(FMassEntityManager& EntityManager, FMassE
 		const TArrayView<FMassForceFragment> ForceList = Context.GetMutableFragmentView<FMassForceFragment>();
 		const TArrayView<FMassDesiredMovementFragment> MovementList = Context.GetMutableFragmentView<FMassDesiredMovementFragment>();
 
+		const FMassMovementParameters& MovementParams = Context.GetConstSharedFragment<FMassMovementParameters>();
+
 		for (FMassExecutionContext::FEntityIterator EntityIt = Context.CreateEntityIterator(); EntityIt; ++EntityIt)
 		{
 			FMassForceFragment& Force = ForceList[EntityIt];
 			FMassDesiredMovementFragment& DesiredMovement = MovementList[EntityIt];
 
 			// Update desired velocity from steering forces.
-			DesiredMovement.DesiredVelocity += Force.Value * DeltaTime;
+			const FVector DesiredVelocity  = DesiredMovement.DesiredVelocity + (Force.Value * DeltaTime);
+			DesiredMovement.DesiredVelocity = DesiredVelocity.GetClampedToMaxSize(MovementParams.MaxSpeed);
 
 			// Reset to zero after force is applied. Processors accumulate forces into the force fragment
 			Force.Value = FVector::ZeroVector;
