Integrating steering changes: Option to avoid slowdown from the steering
(from other depot: 8710728, 8800961, 8932609)

diff --git a/Engine/Plugins/AI/MassAI/Source/MassNavigation/Private/MassSteeringProcessors.cpp b/Engine/Plugins/AI/MassAI/Source/MassNavigation/Private/MassSteeringProcessors.cpp
--- a/Engine/Plugins/AI/MassAI/Source/MassNavigation/Private/MassSteeringProcessors.cpp
+++ b/Engine/Plugins/AI/MassAI/Source/MassNavigation/Private/MassSteeringProcessors.cpp
@@ -74,6 +74,7 @@ void UMassSteerToMoveTargetProcessor::ConfigureQueries(const TSharedRef<FMassEnt
 	EntityQuery.AddConstSharedRequirement<FMassMovementParameters>(EMassFragmentPresence::All);
 	EntityQuery.AddConstSharedRequirement<FMassMovingSteeringParameters>(EMassFragmentPresence::All);
 	EntityQuery.AddConstSharedRequirement<FMassStandingSteeringParameters>(EMassFragmentPresence::All);
+	EntityQuery.AddTagRequirement<FMassSteerToMoveTargetPreventSlowdownTag>(EMassFragmentPresence::Optional);
 	
 #if WITH_MASSGAMEPLAY_DEBUG
 	EntityQuery.DebugEnableEntityOwnerLogging();
@@ -102,7 +103,8 @@ void UMassSteerToMoveTargetProcessor::Execute(FMassEntityManager& EntityManager,
 		const FMassMovementParameters& MovementParams = Context.GetConstSharedFragment<FMassMovementParameters>();
 		const FMassMovingSteeringParameters& MovingSteeringParams = Context.GetConstSharedFragment<FMassMovingSteeringParameters>();
 		const FMassStandingSteeringParameters& StandingSteeringParams = Context.GetConstSharedFragment<FMassStandingSteeringParameters>();
-		bool bIsCodeDriven = Context.DoesArchetypeHaveTag<FMassCodeDrivenMovementTag>();
+		const bool bIsCodeDriven = Context.DoesArchetypeHaveTag<FMassCodeDrivenMovementTag>();
+		const bool bPreventSteeringSlowdown = Context.DoesArchetypeHaveTag<FMassSteerToMoveTargetPreventSlowdownTag>();
 
 		const FVector::FReal SteerK = 1. / MovingSteeringParams.ReactionTime;
 		const float DeltaTime = Context.GetDeltaTimeSeconds();
@@ -121,6 +123,8 @@ void UMassSteerToMoveTargetProcessor::Execute(FMassEntityManager& EntityManager,
 
 			const FTransform& Transform = TransformFragment.GetTransform();;
 
+			FVector::FReal SlowdownScale = 1;
+
 			// Calculate velocity for steering.
 			const FVector CurrentLocation = Transform.GetLocation();
 			const FVector CurrentForward = Transform.GetRotation().GetForwardVector();
@@ -230,6 +234,23 @@ void UMassSteerToMoveTargetProcessor::Execute(FMassEntityManager& EntityManager,
 				// the difference between current desired velocity and target desired velocity.
 				// We dont want to read the actual agent velocity directly since this can create a feedback loop for animated characters
 				Force.Value = SteerK * (Steering.DesiredVelocity - DesiredMovement.DesiredVelocity); // Goal force
+
+				// Apply only if DesiredVelocity is above or equal to DefaultDesiredSpeed, else we want the steering force to slow us down to the desired speed.
+				if (bPreventSteeringSlowdown && DesiredMovement.DesiredVelocity.Size() <= MovementParams.DefaultDesiredSpeed)
+				{
+					const FVector::FReal DistanceToMoveTarget = FVector::Dist(CurrentLocation, MoveTarget.Center);
+					const FVector::FReal ExtraDistance = FMath::Max(0, DistanceToMoveTarget - (1.1*LookAheadDistance)); // Multiplied by 1.1 to add some buffer
+					SlowdownScale = 1 - FMath::Clamp(ExtraDistance/MovingSteeringParams.SteeringPreventSlowdownAttenuationDistance, 0, 1);
+
+					// Make sure the steering force does not reduce DesiredMovement.DesiredVelocity.
+					const FVector DesiredVelocityNormal = DesiredMovement.DesiredVelocity.GetSafeNormal();
+					const FVector::FReal ProjectedForceSize = Force.Value.Dot(DesiredVelocityNormal);
+					if (ProjectedForceSize < 0)
+					{
+						// The force is slowing the entity, avoid that.
+						Force.Value += SlowdownScale * FMath::Abs(ProjectedForceSize) * DesiredVelocityNormal;
+					}
+				}
 			}
 			else if (MoveTarget.GetCurrentAction() == EMassMovementAction::Stand)
 			{
@@ -399,6 +420,14 @@ void UMassSteerToMoveTargetProcessor::Execute(FMassEntityManager& EntityManager,
 				// Display Force
 				UE_VLOG_SEGMENT_THICK(LogOwner, LogMassNavigation, Log, CurrentLocation + ZOffset, CurrentLocation + Force.Value + ZOffset, FColor::Emerald, /*Thickness*/4,
 					TEXT("Steering Force %.1f"), Force.Value.Length());
+
+				// Display slowdown scale as a bigger circle when near 0.
+				constexpr FVector::FReal SlowdownScaleDisplayThreshold = 0.95; 
+				if (SlowdownScale < SlowdownScaleDisplayThreshold)
+				{
+					UE_VLOG_CIRCLE_THICK(this, LogMassNavigation, Log, CurrentLocation, FVector::UpVector, (1-SlowdownScale)*100, FColor::Red, /*Thickness*/2,
+						TEXT("Slowdown scale: %.0f"), SlowdownScale);
+				}
 			}
 #endif // WITH_MASSGAMEPLAY_DEBUG
 			
diff --git a/Engine/Plugins/AI/MassAI/Source/MassNavigation/Public/Steering/MassSteeringFragments.h b/Engine/Plugins/AI/MassAI/Source/MassNavigation/Public/Steering/MassSteeringFragments.h
--- a/Engine/Plugins/AI/MassAI/Source/MassNavigation/Public/Steering/MassSteeringFragments.h
+++ b/Engine/Plugins/AI/MassAI/Source/MassNavigation/Public/Steering/MassSteeringFragments.h
@@ -48,6 +48,12 @@ struct FMassStandingSteeringFragment : public FMassFragment
 	bool bEnteredFromMoveAction = false;
 };
 
+/** Tag used to prevent steering from generating a slowdown force. */
+USTRUCT()
+struct FMassSteerToMoveTargetPreventSlowdownTag : public FMassTag
+{
+	GENERATED_BODY()
+};
 
 /** Steering related movement parameters. */
 USTRUCT()
@@ -63,6 +69,10 @@ struct FMassMovingSteeringParameters : public FMassConstSharedFragment
 	UPROPERTY(EditAnywhere, Category = "Moving", meta = (ClampMin = "0", ForceUnits="s"))
 	float LookAheadTime = 1.0f;
 
+	/** When using FMassSteerToMoveTargetPreventSlowdownTag, maximum distance at which slowdown reduction is applied. */
+	UPROPERTY(EditAnywhere, Category = "Moving", meta = (ClampMin = "0", ForceUnits="s"))
+	float SteeringPreventSlowdownAttenuationDistance = 50.f;
+	
 	/** Allow directional and catchup speed variance. */
 	UPROPERTY(EditAnywhere, Category = "Moving")
 	bool bAllowSpeedVariance = true;
