UMassNavMeshPathFollowProcessor: improved logging information and clamp distance on closest segment

diff --git a/Engine/Plugins/AI/MassAI/Source/MassNavMeshNavigation/Private/MassNavMeshNavigationProcessors.cpp b/Engine/Plugins/AI/MassAI/Source/MassNavMeshNavigation/Private/MassNavMeshNavigationProcessors.cpp
--- a/Engine/Plugins/AI/MassAI/Source/MassNavMeshNavigation/Private/MassNavMeshNavigationProcessors.cpp
+++ b/Engine/Plugins/AI/MassAI/Source/MassNavMeshNavigation/Private/MassNavMeshNavigationProcessors.cpp
@@ -116,7 +116,7 @@ void UMassNavMeshPathFollowProcessor::Execute(FMassEntityManager& EntityManager,
 
 					const FMassNavMeshPathPoint& PointA = ShortPath.Points[ClosestSegmentIndex];
 					const FMassNavMeshPathPoint& PointB = ShortPath.Points[ClosestSegmentIndex+1];
-					const float DistanceOnClosestSegment =  (EntityLocation - PointA.Position).Dot((PointB.Position - PointA.Position).GetSafeNormal());
+					const float DistanceOnClosestSegment = FMath::Max(0.f, (EntityLocation - PointA.Position).Dot((PointB.Position - PointA.Position).GetSafeNormal()));
 					const float EntityProgress = ShortPath.Points[ClosestSegmentIndex].Distance.Get() + DistanceOnClosestSegment;
 					MoveTarget.EntityDistanceToGoal = FMath::Max(0.f, ShortPath.Points[UpdatePointIndex].Distance.Get() - EntityProgress);
 					
@@ -216,6 +216,7 @@ void UMassNavMeshPathFollowProcessor::Execute(FMassEntityManager& EntityManager,
 				// Signal path done.
 				if (!bWasDone && bIsDone)
 				{
+					UE_CVLOG_UELOG(bDisplayDebug, this, LogMassNavMeshNavigation, Log, TEXT("Entity [%s]   Signaling 'path done'."), *Entity.DebugGetDescription());
 					EntitiesToSignalPathDone.Add(Entity);
 				}
 
@@ -262,7 +263,8 @@ void UMassNavMeshPathFollowProcessor::Execute(FMassEntityManager& EntityManager,
 
 					// Draw MoveTarget
 					constexpr float Radius = 20.f;
-					UE_VLOG_CIRCLE(this, LogMassNavMeshNavigation, Display, MoveTarget.Center + ZOffset, FVector(0,0,1), Radius, LightEntityColor, TEXT(""));
+					UE_VLOG_CIRCLE(this, LogMassNavMeshNavigation, Display, MoveTarget.Center + ZOffset, FVector(0,0,1), Radius, LightEntityColor,
+						TEXT("IntentAtGoal: %s"), *UEnum::GetDisplayValueAsText(MoveTarget.IntentAtGoal).ToString());
 
 					// Draw tolerance distance at end (shown as a circle instead of an infinite perpendicular line)
 					if (!ShortPath.bPartialResult)
