FGameplayDebuggerCategory_Mass: Draw purple circle around entities that are LOD off.

diff --git a/Engine/Plugins/AI/MassAI/Source/MassAIDebug/Private/GameplayDebuggerCategory_Mass.cpp b/Engine/Plugins/AI/MassAI/Source/MassAIDebug/Private/GameplayDebuggerCategory_Mass.cpp
--- a/Engine/Plugins/AI/MassAI/Source/MassAIDebug/Private/GameplayDebuggerCategory_Mass.cpp
+++ b/Engine/Plugins/AI/MassAI/Source/MassAIDebug/Private/GameplayDebuggerCategory_Mass.cpp
@@ -525,6 +525,7 @@ void FGameplayDebuggerCategory_Mass::CollectData(APlayerController* OwnerPC, AAc
 		EntityQuery.AddRequirement<FMassNavMeshShortPathFragment>(EMassFragmentAccess::ReadOnly, EMassFragmentPresence::Optional);
 		EntityQuery.AddRequirement<FMassNavMeshCachedPathFragment>(EMassFragmentAccess::ReadOnly, EMassFragmentPresence::Optional);
 		EntityQuery.AddRequirement<FMassSmartObjectUserFragment>(EMassFragmentAccess::ReadOnly, EMassFragmentPresence::Optional);
+		EntityQuery.AddTagRequirement<FMassOffLODTag>(EMassFragmentPresence::Optional);
 
 		const double CurrentTime = World->GetTimeSeconds();
 		
@@ -699,6 +700,14 @@ void FGameplayDebuggerCategory_Mass::CollectData(APlayerController* OwnerPC, AAc
 						AddShape(FGameplayDebuggerShape::MakeArrow(BasePos, BasePos + DesiredVelocity, 10.0f, 2.0f, FColor::Blue));
 					}
 					
+					const bool bIsLODOff = Context.DoesArchetypeHaveTag<FMassOffLODTag>();
+					if (bIsLODOff)
+					{
+						// Draw a slightly bigger circle than the entity radius.
+						AddShape(FGameplayDebuggerShape::MakeCircle(BasePos+FVector(0,0,1), FVector::UpVector,
+							bHasRadius ? RadiusList[EntityIndex].Radius + 5 : DebugShapeRadius + 5, FColor::Purple)); 
+					}					
+					
 					// Path
 					if (bShowNearEntityPath && bHasMoveTarget && bHasRadius)
 					{
