[MassGameplay] Fixed priority logging line in LootAt's CreateRequests function.

Priority.Get produces and uint8, there's no LexToString for that.

#rnx

diff --git a/Engine/Plugins/AI/MassAI/Source/MassAIBehavior/Private/MassLookAtSubsystem.cpp b/Engine/Plugins/AI/MassAI/Source/MassAIBehavior/Private/MassLookAtSubsystem.cpp
--- a/Engine/Plugins/AI/MassAI/Source/MassAIBehavior/Private/MassLookAtSubsystem.cpp
+++ b/Engine/Plugins/AI/MassAI/Source/MassAIBehavior/Private/MassLookAtSubsystem.cpp
@@ -57,10 +57,10 @@ bool CreateRequests(FMassEntityManager& InEntityManager
 				, CustomInterpolationSpeed
 			})});
 
-		UE_VLOG_UELOG(InEntityManager.GetOwner(), LogMassBehavior, Log, TEXT("Created LookAt Request '%s', Target '%s', Priority=%s")
+		UE_VLOG_UELOG(InEntityManager.GetOwner(), LogMassBehavior, Log, TEXT("Created LookAt Request '%s', Target '%s', Priority=%u")
 			, *LexToString(Requests[Index])
 			, *LexToString(TargetEntity)
-			, *LexToString(Priority.Get()))
+			, Priority.Get());
 	}
 
 	return true;
