[Mass Navmesh Navigation] Increase default corridor width used by FMassNavMeshPathFollowTask

diff --git a/Engine/Plugins/AI/MassAI/Source/MassAIBehavior/Public/Tasks/MassNavMeshPathfollowTask.h b/Engine/Plugins/AI/MassAI/Source/MassAIBehavior/Public/Tasks/MassNavMeshPathfollowTask.h
--- a/Engine/Plugins/AI/MassAI/Source/MassAIBehavior/Public/Tasks/MassNavMeshPathfollowTask.h
+++ b/Engine/Plugins/AI/MassAI/Source/MassAIBehavior/Public/Tasks/MassNavMeshPathfollowTask.h
@@ -34,7 +34,7 @@ struct FMassNavMeshPathFollowTaskInstanceData
 	
 	/** Maximum width of the corridor to use. */
 	UPROPERTY(EditAnywhere, Category = Parameter)
-    float CorridorWidth = 300.f;
+    float CorridorWidth = 600.f;
 
 	/** Amount to offset corridor sides from navmesh borders. Used to keep movement away for navmesh borders. */
 	UPROPERTY(EditAnywhere, Category = Parameter)
