[MassAI] replaced FStateTreeStructRef by FStateTreePropertyRef

diff --git a/Engine/Plugins/AI/MassAI/Source/MassAIBehavior/MassAIBehavior.Build.cs b/Engine/Plugins/AI/MassAI/Source/MassAIBehavior/MassAIBehavior.Build.cs
--- a/Engine/Plugins/AI/MassAI/Source/MassAIBehavior/MassAIBehavior.Build.cs
+++ b/Engine/Plugins/AI/MassAI/Source/MassAIBehavior/MassAIBehavior.Build.cs
@@ -44,6 +44,12 @@ namespace UnrealBuildTool.Rules
 				}
 			);
 
+			PrivateDependencyModuleNames.AddRange(
+				new string[] {
+					"PropertyBindingUtils"
+				}
+			);
+
 			if (Target.bBuildEditor == true)
 			{
 				PrivateDependencyModuleNames.Add("EditorFramework");
diff --git a/Engine/Plugins/AI/MassAI/Source/MassAIBehavior/Private/Tasks/MassClaimSmartObjectTask.cpp b/Engine/Plugins/AI/MassAI/Source/MassAIBehavior/Private/Tasks/MassClaimSmartObjectTask.cpp
--- a/Engine/Plugins/AI/MassAI/Source/MassAIBehavior/Private/Tasks/MassClaimSmartObjectTask.cpp
+++ b/Engine/Plugins/AI/MassAI/Source/MassAIBehavior/Private/Tasks/MassClaimSmartObjectTask.cpp
@@ -56,7 +56,7 @@ EStateTreeRunStatus FMassClaimSmartObjectTask::EnterState(FStateTreeExecutionCon
 
 	FInstanceDataType& InstanceData = Context.GetInstanceData(*this);
 
-	const FMassSmartObjectCandidateSlots* CandidateSlots = InstanceData.CandidateSlots.GetPtr<FMassSmartObjectCandidateSlots>();
+	const FMassSmartObjectCandidateSlots* CandidateSlots = InstanceData.CandidateSlots.GetMutablePtr<FMassSmartObjectCandidateSlots>(Context);
 	if (CandidateSlots == nullptr)
 	{
 		MASSBEHAVIOR_LOG(Log, TEXT("Candidate slots not set"));
@@ -64,11 +64,11 @@ EStateTreeRunStatus FMassClaimSmartObjectTask::EnterState(FStateTreeExecutionCon
 	}
 
 	InstanceData.ClaimedSlot.Invalidate();
-	
+
 	// Setup MassSmartObject handler and claim
 	const FMassStateTreeExecutionContext& MassStateTreeContext = static_cast<FMassStateTreeExecutionContext&>(Context);
 	const FMassSmartObjectHandler MassSmartObjectHandler(MassStateTreeContext.GetMassEntityExecutionContext(), SmartObjectSubsystem, SignalSubsystem);
-	
+
 	InstanceData.ClaimedSlot = MassSmartObjectHandler.ClaimCandidate(MassStateTreeContext.GetEntity(), SOUser, *CandidateSlots);
 
 	// Treat claiming a slot as consuming all the candidate slots.
@@ -131,7 +131,7 @@ EStateTreeRunStatus FMassClaimSmartObjectTask::Tick(FStateTreeExecutionContext&
 {
 	FMassSmartObjectUserFragment& SOUser = Context.GetExternalData(SmartObjectUserHandle);
 	FInstanceDataType& InstanceData = Context.GetInstanceData(*this);
-	
+
 	// Prevent FindSmartObject() to query new objects while claimed.
 	// This is done here because of the limited ways we can communicate between FindSmartObject() and ClaimSmartObject().
 	// InteractionCooldownEndTime is used by the FindSmartObject() to invalidate the candidates.
diff --git a/Engine/Plugins/AI/MassAI/Source/MassAIBehavior/Private/Tasks/MassZoneGraphPathFollowTask.cpp b/Engine/Plugins/AI/MassAI/Source/MassAIBehavior/Private/Tasks/MassZoneGraphPathFollowTask.cpp
--- a/Engine/Plugins/AI/MassAI/Source/MassAIBehavior/Private/Tasks/MassZoneGraphPathFollowTask.cpp
+++ b/Engine/Plugins/AI/MassAI/Source/MassAIBehavior/Private/Tasks/MassZoneGraphPathFollowTask.cpp
@@ -54,7 +54,7 @@ bool FMassZoneGraphPathFollowTask::RequestPath(FMassStateTreeExecutionContext& C
 	const FMassMovementParameters& MovementParams = Context.GetExternalData(MovementParamsHandle);
 
 	const FInstanceDataType& InstanceData = Context.GetInstanceData(*this);
-	
+
 	bool bDisplayDebug = false;
 #if WITH_MASSGAMEPLAY_DEBUG
 	bDisplayDebug = UE::Mass::Debug::IsDebuggingEntity(Context.GetEntity());
@@ -114,14 +114,14 @@ EStateTreeRunStatus FMassZoneGraphPathFollowTask::EnterState(FStateTreeExecution
 	}
 
 	const FMassZoneGraphLaneLocationFragment& LaneLocation = Context.GetExternalData(LocationHandle);
-	const FMassZoneGraphTargetLocation* TargetLocation = InstanceData.TargetLocation.GetPtr<FMassZoneGraphTargetLocation>();
+	const FMassZoneGraphTargetLocation* TargetLocation = InstanceData.TargetLocation.GetMutablePtr<FMassZoneGraphTargetLocation>(Context);
 
 	if (TargetLocation == nullptr)
 	{
 		MASSBEHAVIOR_LOG(Error, TEXT("Target is not defined."));
 		return EStateTreeRunStatus::Failed;
 	}
-	
+
 	if (TargetLocation->LaneHandle != LaneLocation.LaneHandle)
 	{
 		// @todo: Verbosity of this message was downgraded from 'Error' to 'Log' until we have time to investigate.
@@ -155,7 +155,7 @@ EStateTreeRunStatus FMassZoneGraphPathFollowTask::Tick(FStateTreeExecutionContex
 
 	const FMassZoneGraphShortPathFragment& ShortPath = Context.GetExternalData(ShortPathHandle);
 	const FMassZoneGraphLaneLocationFragment& LaneLocation = Context.GetExternalData(LocationHandle);
-	const FMassZoneGraphTargetLocation* TargetLocation = InstanceData.TargetLocation.GetPtr<FMassZoneGraphTargetLocation>();
+	const FMassZoneGraphTargetLocation* TargetLocation = InstanceData.TargetLocation.GetMutablePtr<FMassZoneGraphTargetLocation>(Context);
 
 	// Current path follow is done, but it was partial (i.e. many points on a curve), try again until we get there.
 	if (ShortPath.IsDone() && ShortPath.bPartialResult)
@@ -170,13 +170,13 @@ EStateTreeRunStatus FMassZoneGraphPathFollowTask::Tick(FStateTreeExecutionContex
 			MASSBEHAVIOR_LOG(Error, TEXT("Target is not on current lane, target lane is %s expected %s."), *TargetLocation->LaneHandle.ToString(), *LaneLocation.LaneHandle.ToString());
 			return EStateTreeRunStatus::Failed;
 		}
-		
+
 		if (!RequestPath(MassContext, *TargetLocation))
 		{
 			MASSBEHAVIOR_LOG(Error, TEXT("Failed to request path."));
 			return EStateTreeRunStatus::Failed;
 		}
 	}
-	
+
 	return ShortPath.IsDone() ? EStateTreeRunStatus::Succeeded : EStateTreeRunStatus::Running;
 }
diff --git a/Engine/Plugins/AI/MassAI/Source/MassAIBehavior/Public/Tasks/MassClaimSmartObjectTask.h b/Engine/Plugins/AI/MassAI/Source/MassAIBehavior/Public/Tasks/MassClaimSmartObjectTask.h
--- a/Engine/Plugins/AI/MassAI/Source/MassAIBehavior/Public/Tasks/MassClaimSmartObjectTask.h
+++ b/Engine/Plugins/AI/MassAI/Source/MassAIBehavior/Public/Tasks/MassClaimSmartObjectTask.h
@@ -7,6 +7,7 @@
 #endif // UE_ENABLE_INCLUDE_ORDER_DEPRECATED_IN_5_6
 #include "MassStateTreeTypes.h"
 #include "SmartObjectRuntime.h"
+#include "StateTreePropertyRef.h"
 #include "MassClaimSmartObjectTask.generated.h"
 
 struct FStateTreeExecutionContext;
@@ -33,8 +34,8 @@ struct FMassClaimSmartObjectTaskInstanceData
 	GENERATED_BODY()
 
 	/** Result of the candidates search request (Input) */
-	UPROPERTY(VisibleAnywhere, Category = Input, meta = (BaseStruct = "/Script/MassSmartObjects.MassSmartObjectCandidateSlots"))
-	FStateTreeStructRef CandidateSlots;
+	UPROPERTY(VisibleAnywhere, Category = Input, meta = (RefType = "/Script/MassSmartObjects.MassSmartObjectCandidateSlots"))
+	FStateTreePropertyRef CandidateSlots;
 
 	UPROPERTY(VisibleAnywhere, Category = Output)
 	FSmartObjectClaimHandle ClaimedSlot;
diff --git a/Engine/Plugins/AI/MassAI/Source/MassAIBehavior/Public/Tasks/MassZoneGraphPathFollowTask.h b/Engine/Plugins/AI/MassAI/Source/MassAIBehavior/Public/Tasks/MassZoneGraphPathFollowTask.h
--- a/Engine/Plugins/AI/MassAI/Source/MassAIBehavior/Public/Tasks/MassZoneGraphPathFollowTask.h
+++ b/Engine/Plugins/AI/MassAI/Source/MassAIBehavior/Public/Tasks/MassZoneGraphPathFollowTask.h
@@ -6,6 +6,7 @@
 #include "MassMovementTypes.h"
 #include "MassCommonTypes.h"
 #include "MassStateTreeTypes.h"
+#include "StateTreePropertyRef.h"
 #include "ZoneGraphTypes.h"
 #include "MassZoneGraphPathFollowTask.generated.h"
 
@@ -78,8 +79,8 @@ struct FMassZoneGraphPathFollowTaskInstanceData
 {
 	GENERATED_BODY()
 
-	UPROPERTY(EditAnywhere, Category = Input, meta=(BaseStruct = "/Script/MassAIBehavior.MassZoneGraphTargetLocation"))
-	FStateTreeStructRef TargetLocation; 
+	UPROPERTY(EditAnywhere, Category = Input, meta=(RefType = "/Script/MassAIBehavior.MassZoneGraphTargetLocation"))
+	FStateTreePropertyRef TargetLocation;
 
 	UPROPERTY(EditAnywhere, Category = Parameter)
 	FMassMovementStyleRef MovementStyle;
