[MassAI] add find random wander target location task

diff --git a/Engine/Plugins/AI/MassAI/Source/MassAIBehavior/Private/Tasks/MassNavMeshFindReachablePointTask.cpp b/Engine/Plugins/AI/MassAI/Source/MassAIBehavior/Private/Tasks/MassNavMeshFindReachablePointTask.cpp
new file mode 100644
--- /dev/null
+++ b/Engine/Plugins/AI/MassAI/Source/MassAIBehavior/Private/Tasks/MassNavMeshFindReachablePointTask.cpp
@@ -0,0 +1,63 @@
+// Copyright Epic Games, Inc. All Rights Reserved.
+
+#include "Tasks/MassNavMeshFindReachablePointTask.h"
+#include "MassAIBehaviorTypes.h"
+#include "MassCommonFragments.h"
+#include "MassNavigationFragments.h"
+#include "MassDebugger.h"
+#include "MassStateTreeDependency.h"
+#include "MassStateTreeExecutionContext.h"
+#include "NavigationSystem.h"
+#include "StateTreeLinker.h"
+
+#include UE_INLINE_GENERATED_CPP_BY_NAME(MassNavMeshFindReachablePointTask)
+
+bool FMassNavMeshFindReachablePointTask::Link(FStateTreeLinker& Linker)
+{
+	Linker.LinkExternalData(TransformHandle);
+	Linker.LinkExternalData(MoveTargetHandle);
+	Linker.LinkExternalData(AgentRadiusHandle);
+
+	return true;
+}
+
+EStateTreeRunStatus FMassNavMeshFindReachablePointTask::EnterState(FStateTreeExecutionContext& Context, const FStateTreeTransitionResult& Transition) const
+{
+	const FMassStateTreeExecutionContext& MassContext = static_cast<FMassStateTreeExecutionContext&>(Context);
+	FInstanceDataType& InstanceData = Context.GetInstanceData(*this);
+
+	bool bDisplayDebug = false;
+#if WITH_MASSGAMEPLAY_DEBUG
+	bDisplayDebug = UE::Mass::Debug::IsDebuggingEntity(MassContext.GetEntity());
+#endif // WITH_MASSGAMEPLAY_DEBUG
+
+	const FAgentRadiusFragment& AgentRadius = Context.GetExternalData(AgentRadiusHandle);
+	const FVector AgentNavLocation = Context.GetExternalData(TransformHandle).GetTransform().GetLocation();
+	const FNavAgentProperties& NavAgentProperties = FNavAgentProperties(AgentRadius.Radius);
+
+	UNavigationSystemV1* NavMeshSubsystem = Cast<UNavigationSystemV1>(Context.GetWorld()->GetNavigationSystem());
+	const ANavigationData* NavData = NavMeshSubsystem->GetNavDataForProps(NavAgentProperties, AgentNavLocation);
+
+	if (NavData == nullptr)
+	{
+		MASSBEHAVIOR_LOG(Warning, TEXT("Invalid NavData"));
+		return EStateTreeRunStatus::Failed;
+	}
+
+	FNavLocation NavLocation;
+	if (!NavData->GetRandomReachablePointInRadius(AgentNavLocation, InstanceData.SearchRadius, NavLocation))
+	{
+		MASSBEHAVIOR_LOG(Warning, TEXT("No reachable points found"));
+		return EStateTreeRunStatus::Failed;
+	}
+
+	InstanceData.TargetLocation.EndOfPathPosition = NavLocation.Location;
+	InstanceData.TargetLocation.EndOfPathIntent = EMassMovementAction::Move;
+
+	if (bDisplayDebug)
+	{
+		MASSBEHAVIOR_LOG(Log, TEXT("Found target: '%s' at %s"), *UEnum::GetValueAsString(InstanceData.TargetLocation.EndOfPathIntent),  *NavLocation.Location.ToString());
+	}
+
+	return EStateTreeRunStatus::Running;
+}
diff --git a/Engine/Plugins/AI/MassAI/Source/MassAIBehavior/Public/Tasks/MassNavMeshAnimateTask.h b/Engine/Plugins/AI/MassAI/Source/MassAIBehavior/Public/Tasks/MassNavMeshAnimateTask.h
--- a/Engine/Plugins/AI/MassAI/Source/MassAIBehavior/Public/Tasks/MassNavMeshAnimateTask.h
+++ b/Engine/Plugins/AI/MassAI/Source/MassAIBehavior/Public/Tasks/MassNavMeshAnimateTask.h
@@ -28,7 +28,7 @@ struct FMassNavMeshAnimateTaskInstanceData
 /**
  * Stop and animate on current navmesh location
  */
-USTRUCT(meta = (DisplayName = "Navmesh Animate"))
+USTRUCT(meta = (DisplayName = "NavMesh Animate"))
 struct FMassNavMeshAnimateTask : public FMassStateTreeTaskBase
 {
 	GENERATED_BODY()
diff --git a/Engine/Plugins/AI/MassAI/Source/MassAIBehavior/Public/Tasks/MassNavMeshFindReachablePointTask.h b/Engine/Plugins/AI/MassAI/Source/MassAIBehavior/Public/Tasks/MassNavMeshFindReachablePointTask.h
new file mode 100644
--- /dev/null
+++ b/Engine/Plugins/AI/MassAI/Source/MassAIBehavior/Public/Tasks/MassNavMeshFindReachablePointTask.h
@@ -0,0 +1,52 @@
+// Copyright Epic Games, Inc. All Rights Reserved.
+
+#pragma once
+
+#include "MassNavigationTypes.h"
+#include "MassStateTreeTypes.h"
+#include "MassNavMeshFindReachablePointTask.generated.h"
+
+struct FAgentRadiusFragment;
+struct FMassMoveTargetFragment;
+struct FStateTreeExecutionContext;
+struct FTransformFragment;
+namespace UE::MassBehavior
+{
+struct FStateTreeDependencyBuilder;
+};
+
+USTRUCT()
+struct FMassNavMeshFindReachablePointTaskInstanceData
+{
+	GENERATED_BODY()
+
+	/** Radius around the current entity's location to finds a random reachable location */
+	UPROPERTY(EditAnywhere, Category = Parameters)
+	float SearchRadius = 500.f;
+
+	UPROPERTY(EditAnywhere, Category = Output)
+	FMassTargetLocation TargetLocation;
+};
+
+/**
+ * Updates TargetLocation to a random reachable location based on the agents current location on NavMesh.
+ */
+USTRUCT(meta = (DisplayName = "NavMesh Find Random Reachable Target"))
+struct FMassNavMeshFindReachablePointTask : public FMassStateTreeTaskBase
+{
+	GENERATED_BODY()
+
+	using FInstanceDataType = FMassNavMeshFindReachablePointTaskInstanceData;
+
+protected:
+	virtual bool Link(FStateTreeLinker& Linker) override;
+	virtual const UStruct* GetInstanceDataType() const override
+	{
+		return FInstanceDataType::StaticStruct();
+	}
+	virtual EStateTreeRunStatus EnterState(FStateTreeExecutionContext& Context, const FStateTreeTransitionResult& Transition) const override;
+
+	TStateTreeExternalDataHandle<FTransformFragment> TransformHandle;
+	TStateTreeExternalDataHandle<FMassMoveTargetFragment> MoveTargetHandle;
+	TStateTreeExternalDataHandle<FAgentRadiusFragment> AgentRadiusHandle;
+};
