[MassAI] removed unused external data handle from FMassNavMeshFindReachablePointTask

diff --git a/Engine/Plugins/AI/MassAI/Source/MassAIBehavior/Private/Tasks/MassNavMeshFindReachablePointTask.cpp b/Engine/Plugins/AI/MassAI/Source/MassAIBehavior/Private/Tasks/MassNavMeshFindReachablePointTask.cpp
--- a/Engine/Plugins/AI/MassAI/Source/MassAIBehavior/Private/Tasks/MassNavMeshFindReachablePointTask.cpp
+++ b/Engine/Plugins/AI/MassAI/Source/MassAIBehavior/Private/Tasks/MassNavMeshFindReachablePointTask.cpp
@@ -15,7 +15,6 @@
 bool FMassNavMeshFindReachablePointTask::Link(FStateTreeLinker& Linker)
 {
 	Linker.LinkExternalData(TransformHandle);
-	Linker.LinkExternalData(MoveTargetHandle);
 	Linker.LinkExternalData(AgentRadiusHandle);
 
 	return true;
diff --git a/Engine/Plugins/AI/MassAI/Source/MassAIBehavior/Public/Tasks/MassNavMeshFindReachablePointTask.h b/Engine/Plugins/AI/MassAI/Source/MassAIBehavior/Public/Tasks/MassNavMeshFindReachablePointTask.h
--- a/Engine/Plugins/AI/MassAI/Source/MassAIBehavior/Public/Tasks/MassNavMeshFindReachablePointTask.h
+++ b/Engine/Plugins/AI/MassAI/Source/MassAIBehavior/Public/Tasks/MassNavMeshFindReachablePointTask.h
@@ -7,7 +7,6 @@
 #include "MassNavMeshFindReachablePointTask.generated.h"
 
 struct FAgentRadiusFragment;
-struct FMassMoveTargetFragment;
 struct FStateTreeExecutionContext;
 struct FTransformFragment;
 namespace UE::MassBehavior
@@ -47,6 +46,5 @@ protected:
 	virtual EStateTreeRunStatus EnterState(FStateTreeExecutionContext& Context, const FStateTreeTransitionResult& Transition) const override;
 
 	TStateTreeExternalDataHandle<FTransformFragment> TransformHandle;
-	TStateTreeExternalDataHandle<FMassMoveTargetFragment> MoveTargetHandle;
 	TStateTreeExternalDataHandle<FAgentRadiusFragment> AgentRadiusHandle;
 };
