[MassSmartObject] updated MassFindSmartObjectTargetTask to be able to use entrance location request.

diff --git a/Engine/Plugins/AI/MassAI/Source/MassAIBehavior/Private/Tasks/MassFindSmartObjectTargetTask.cpp b/Engine/Plugins/AI/MassAI/Source/MassAIBehavior/Private/Tasks/MassFindSmartObjectTargetTask.cpp
--- a/Engine/Plugins/AI/MassAI/Source/MassAIBehavior/Private/Tasks/MassFindSmartObjectTargetTask.cpp
+++ b/Engine/Plugins/AI/MassAI/Source/MassAIBehavior/Private/Tasks/MassFindSmartObjectTargetTask.cpp
@@ -3,17 +3,21 @@
 #include "Tasks/MassFindSmartObjectTargetTask.h"
 
 #include "MassAIBehaviorTypes.h"
-#include "SmartObjectSubsystem.h"
+#include "MassCommonFragments.h"
 #include "MassSmartObjectFragments.h"
 #include "MassStateTreeExecutionContext.h"
+#include "NavigationSystem.h"
+#include "SmartObjectSubsystem.h"
 #include "StateTreeLinker.h"
 
 #include UE_INLINE_GENERATED_CPP_BY_NAME(MassFindSmartObjectTargetTask)
 
 bool FMassFindSmartObjectTargetTask::Link(FStateTreeLinker& Linker)
 {
+	Linker.LinkExternalData(TransformHandle);
+	Linker.LinkExternalData(AgentRadiusHandle);
 	Linker.LinkExternalData(SmartObjectSubsystemHandle);
-	
+
 	return true;
 }
 
@@ -31,10 +35,56 @@ EStateTreeRunStatus FMassFindSmartObjectTargetTask::EnterState(FStateTreeExecuti
 		return EStateTreeRunStatus::Failed;
 	}
 
+	const FVector EntityLocation = Context.GetExternalData(TransformHandle).GetTransform().GetLocation();
+	const FAgentRadiusFragment* AgentRadius = Context.GetExternalDataPtr(AgentRadiusHandle);
+
+	// Try to find an entrance first
+	if (bUseEntranceLocationRequest)
+	{
+		FSmartObjectSlotEntranceLocationRequest Request = InstanceData.EntranceRequest;
+
+		if (Request.bProjectNavigationLocation)
+		{
+			if (AgentRadius != nullptr)
+			{
+				const FNavAgentProperties& NavAgentProperties = FNavAgentProperties(AgentRadius->Radius);
+				const UNavigationSystemV1* NavMeshSubsystem = Cast<UNavigationSystemV1>(Context.GetWorld()->GetNavigationSystem());
+				if (const ANavigationData* NavData = NavMeshSubsystem->GetNavDataForProps(NavAgentProperties, EntityLocation))
+				{
+					Request.NavigationData = NavData;
+				}
+				else
+				{
+					MASSBEHAVIOR_LOG(Error, TEXT("Entrance request can not be executed since no navigation data can be found and is required for 'ProjectNavigationLocation'"));
+					return EStateTreeRunStatus::Failed;
+				}
+			}
+			else
+			{
+				MASSBEHAVIOR_LOG(Error, TEXT("Entrance request can not be executed since the required AgentRadius fragment for 'ProjectNavigationLocation' is missing"));
+				return EStateTreeRunStatus::Failed;
+			}
+		}
+
+		if (InstanceData.bUseEntityLocationAsSearchLocation)
+		{
+			Request.SearchLocation = EntityLocation;
+		}
+
+		FSmartObjectSlotEntranceLocationResult EntryLocation;
+		if (SmartObjectSubsystem.FindEntranceLocationForSlot(InstanceData.ClaimedSlot.SlotHandle, Request, EntryLocation))
+		{
+			InstanceData.SmartObjectLocation.EndOfPathIntent = EMassMovementAction::Stand;
+			InstanceData.SmartObjectLocation.EndOfPathPosition = EntryLocation.Location;
+			return EStateTreeRunStatus::Running;
+		}
+	}
+
+	// Use slot location (can also be the fallback when no entrance location request is used and no entrances were found)
 	const FTransform Transform = SmartObjectSubsystem.GetSlotTransform(InstanceData.ClaimedSlot).Get(FTransform::Identity);
 
 	InstanceData.SmartObjectLocation.EndOfPathIntent = EMassMovementAction::Stand;
 	InstanceData.SmartObjectLocation.EndOfPathPosition = Transform.GetLocation();
-	
+
 	return EStateTreeRunStatus::Running;
 }
diff --git a/Engine/Plugins/AI/MassAI/Source/MassAIBehavior/Public/Tasks/MassFindSmartObjectTargetTask.h b/Engine/Plugins/AI/MassAI/Source/MassAIBehavior/Public/Tasks/MassFindSmartObjectTargetTask.h
--- a/Engine/Plugins/AI/MassAI/Source/MassAIBehavior/Public/Tasks/MassFindSmartObjectTargetTask.h
+++ b/Engine/Plugins/AI/MassAI/Source/MassAIBehavior/Public/Tasks/MassFindSmartObjectTargetTask.h
@@ -5,10 +5,13 @@
 #include "MassNavigationTypes.h"
 #include "MassStateTreeTypes.h"
 #include "SmartObjectRuntime.h"
+#include "SmartObjectSubsystem.h"
 #include "MassFindSmartObjectTargetTask.generated.h"
 
 #define UE_API MASSAIBEHAVIOR_API
 
+struct FTransformFragment;
+struct FAgentRadiusFragment;
 class USmartObjectSubsystem;
 
 USTRUCT()
@@ -16,6 +19,18 @@ struct FMassFindSmartObjectTargetInstanceData
 {
 	GENERATED_BODY()
 
+	/**
+	 * When using the entrance location request with selection method 'NearestToSearchLocation',
+	 * this property indicates whether the request will set the search location using
+	 * the entity transform or use the request value (can be bound).
+	 */
+	UPROPERTY(EditAnywhere, Category = Parameter)
+	bool bUseEntityLocationAsSearchLocation = true;
+
+	/** Request parameters when using the entrance location request */
+	UPROPERTY(EditAnywhere, Category = Parameter)
+	FSmartObjectSlotEntranceLocationRequest EntranceRequest;
+
 	UPROPERTY(VisibleAnywhere, Category = Input)
 	FSmartObjectClaimHandle ClaimedSlot;
 
@@ -23,7 +38,11 @@ struct FMassFindSmartObjectTargetInstanceData
 	FMassTargetLocation SmartObjectLocation;
 };
 
-/** Computes move target to a smart object based on current location. */
+/**
+ * Computes move target to a smart object based on current location.
+ * The move target will use the entrance location found by the EntranceRequest if possible,
+ * otherwise the slot location will be used.
+ */
 USTRUCT(meta = (DisplayName = "Find Smart Object Target"))
 struct FMassFindSmartObjectTargetTask : public FMassStateTreeTaskBase
 {
@@ -33,10 +52,22 @@ struct FMassFindSmartObjectTargetTask : public FMassStateTreeTaskBase
 
 protected:
 	UE_API virtual bool Link(FStateTreeLinker& Linker) override;
-	virtual const UStruct* GetInstanceDataType() const override { return FInstanceDataType::StaticStruct(); }
+	virtual const UStruct* GetInstanceDataType() const override
+	{
+		return FInstanceDataType::StaticStruct();
+	}
 	UE_API virtual EStateTreeRunStatus EnterState(FStateTreeExecutionContext& Context, const FStateTreeTransitionResult& Transition) const override;
 
+	TStateTreeExternalDataHandle<FAgentRadiusFragment, EStateTreeExternalDataRequirement::Optional> AgentRadiusHandle;
+	TStateTreeExternalDataHandle<FTransformFragment> TransformHandle;
 	TStateTreeExternalDataHandle<USmartObjectSubsystem> SmartObjectSubsystemHandle;
+
+	/**
+	 * Whether the entrance location request should be used first to find an entrance location
+	 * to use instead of the slot location. Slot location will be used if not entrances can be found.
+	 */
+	UPROPERTY(EditAnywhere, Category = Parameter)
+	bool bUseEntranceLocationRequest = true;
 };
 
 #undef UE_API
