[MassAI] minor update to SmartObject task to expose NumSlots for bindings and fail claim task when no slots are available

diff --git a/Engine/Plugins/AI/MassAI/Source/MassAIBehavior/Private/Tasks/MassClaimSmartObjectTask.cpp b/Engine/Plugins/AI/MassAI/Source/MassAIBehavior/Private/Tasks/MassClaimSmartObjectTask.cpp
--- a/Engine/Plugins/AI/MassAI/Source/MassAIBehavior/Private/Tasks/MassClaimSmartObjectTask.cpp
+++ b/Engine/Plugins/AI/MassAI/Source/MassAIBehavior/Private/Tasks/MassClaimSmartObjectTask.cpp
@@ -63,6 +63,12 @@ EStateTreeRunStatus FMassClaimSmartObjectTask::EnterState(FStateTreeExecutionCon
 		return EStateTreeRunStatus::Failed;
 	}
 
+	if (CandidateSlots->NumSlots == 0)
+	{
+		MASSBEHAVIOR_LOG(Log, TEXT("No candidate slots"));
+		return EStateTreeRunStatus::Failed;
+	}
+
 	InstanceData.ClaimedSlot.Invalidate();
 
 	// Setup MassSmartObject handler and claim
diff --git a/Engine/Plugins/Runtime/MassGameplay/Source/MassSmartObjects/Public/MassSmartObjectRequest.h b/Engine/Plugins/Runtime/MassGameplay/Source/MassSmartObjects/Public/MassSmartObjectRequest.h
--- a/Engine/Plugins/Runtime/MassGameplay/Source/MassSmartObjects/Public/MassSmartObjectRequest.h
+++ b/Engine/Plugins/Runtime/MassGameplay/Source/MassSmartObjects/Public/MassSmartObjectRequest.h
@@ -72,7 +72,7 @@ struct FMassSmartObjectCandidateSlots
 	static constexpr uint32 MaxNumCandidates = 4;
 	TStaticArray<FSmartObjectCandidateSlot, MaxNumCandidates> Slots;
 
-	UPROPERTY(Transient)
+	UPROPERTY(Transient, VisibleAnywhere, Category = SmartObject)
 	uint8 NumSlots = 0;
 };
 
