[MassStateTree] updated Trait and processors to only run for Standalone and Server execution mode (now consistent with UMassStateTreeFragmentDestructor).
This fixes issue with some failed requirement (e.g., SmartObjectSubsystem) not available on Client.
More work will be required to enable MassStateTree on clients.

diff --git a/Engine/Plugins/AI/MassAI/Source/MassAIBehavior/Private/MassStateTreeProcessors.cpp b/Engine/Plugins/AI/MassAI/Source/MassAIBehavior/Private/MassStateTreeProcessors.cpp
--- a/Engine/Plugins/AI/MassAI/Source/MassAIBehavior/Private/MassStateTreeProcessors.cpp
+++ b/Engine/Plugins/AI/MassAI/Source/MassAIBehavior/Private/MassStateTreeProcessors.cpp
@@ -103,7 +103,7 @@ struct FLinearQueueAllocator
 UMassStateTreeFragmentDestructor::UMassStateTreeFragmentDestructor()
 	: EntityQuery(*this)
 {
-	ExecutionFlags = (int32)(EProcessorExecutionFlags::Standalone | EProcessorExecutionFlags::Server);
+	ExecutionFlags = static_cast<uint8>(UE::MassStateTree::ExecutionFlags);
 	ObservedType = FMassStateTreeInstanceFragment::StaticStruct();
 	ObservedOperations = EMassObservedOperationFlags::Remove;
 	bRequiresGameThreadExecution = true;
@@ -161,6 +161,7 @@ void UMassStateTreeFragmentDestructor::Execute(FMassEntityManager& EntityManager
 UMassStateTreeActivationProcessor::UMassStateTreeActivationProcessor()
 	: EntityQuery(*this)
 {
+	ExecutionFlags = static_cast<uint8>(UE::MassStateTree::ExecutionFlags);
 	ExecutionOrder.ExecuteAfter.Add(UE::Mass::ProcessorGroupNames::LOD);
 	ExecutionOrder.ExecuteBefore.Add(UE::Mass::ProcessorGroupNames::Behavior);
 	bRequiresGameThreadExecution = true; // due to UMassStateTreeSubsystem RW access
@@ -251,6 +252,7 @@ UMassStateTreeProcessor::UMassStateTreeProcessor(const FObjectInitializer& Objec
 {
 	bRequiresGameThreadExecution = bProcessEntitiesInParallel;
 
+	ExecutionFlags = static_cast<uint8>(UE::MassStateTree::ExecutionFlags);
 	ExecutionOrder.ExecuteInGroup = UE::Mass::ProcessorGroupNames::Behavior;
 
 	// `Behavior` doesn't run on clients but `Tasks` do.
diff --git a/Engine/Plugins/AI/MassAI/Source/MassAIBehavior/Private/MassStateTreeTrait.cpp b/Engine/Plugins/AI/MassAI/Source/MassAIBehavior/Private/MassStateTreeTrait.cpp
--- a/Engine/Plugins/AI/MassAI/Source/MassAIBehavior/Private/MassStateTreeTrait.cpp
+++ b/Engine/Plugins/AI/MassAI/Source/MassAIBehavior/Private/MassStateTreeTrait.cpp
@@ -15,23 +15,30 @@
 
 void UMassStateTreeTrait::BuildTemplate(FMassEntityTemplateBuildContext& BuildContext, const UWorld& World) const
 {
+	const EProcessorExecutionFlags WorldExecutionFlags = UE::Mass::Utils::DetermineProcessorExecutionFlags(&World);
+	if (!EnumHasAnyFlags(WorldExecutionFlags, UE::MassStateTree::ExecutionFlags))
+	{
+		UE_VLOG_UELOG(&World, LogMassBehavior, Verbose, TEXT("StateTree doesn't run on the current execution mode '%s'"), *UEnum::GetValueAsString(WorldExecutionFlags));
+		return;
+	}
+
 	FMassEntityManager& EntityManager = UE::Mass::Utils::GetEntityManagerChecked(World);
 
 	UMassStateTreeSubsystem* MassStateTreeSubsystem = World.GetSubsystem<UMassStateTreeSubsystem>();
 	if (!MassStateTreeSubsystem && !BuildContext.IsInspectingData())
 	{
-		UE_VLOG(&World, LogMassBehavior, Error, TEXT("Failed to get Mass StateTree Subsystem."));
+		UE_VLOG_UELOG(&World, LogMassBehavior, Error, TEXT("Failed to get Mass StateTree Subsystem."));
 		return;
 	}
 
 	if (!StateTree && !BuildContext.IsInspectingData())
 	{
-		UE_VLOG(MassStateTreeSubsystem, LogMassBehavior, Error, TEXT("StateTree asset is not set or unavailable."));
+		UE_VLOG_UELOG(MassStateTreeSubsystem, LogMassBehavior, Error, TEXT("StateTree asset is not set or unavailable."));
 		return;
 	}
 	if (StateTree != nullptr && !BuildContext.IsInspectingData() && !StateTree->IsReadyToRun())
 	{
-		UE_VLOG(MassStateTreeSubsystem, LogMassBehavior, Error, TEXT("StateTree asset is ready to run."));
+		UE_VLOG_UELOG(MassStateTreeSubsystem, LogMassBehavior, Error, TEXT("StateTree asset is ready to run."));
 		return;
 	}
 
@@ -46,16 +53,24 @@ void UMassStateTreeTrait::BuildTemplate(FMassEntityTemplateBuildContext& BuildCo
 
 bool UMassStateTreeTrait::ValidateTemplate(const FMassEntityTemplateBuildContext& BuildContext, const UWorld& World, FAdditionalTraitRequirements& OutTraitRequirements) const
 {
+	constexpr EProcessorExecutionFlags ExecutionFlags(EProcessorExecutionFlags::Standalone | EProcessorExecutionFlags::Server);
+	const EProcessorExecutionFlags WorldExecutionFlags = UE::Mass::Utils::DetermineProcessorExecutionFlags(&World);
+	if (!EnumHasAnyFlags(WorldExecutionFlags, ExecutionFlags))
+	{
+		UE_VLOG_UELOG(&World, LogMassBehavior, Verbose, TEXT("StateTree doesn't run on the current execution mode '%s'"), *UEnum::GetValueAsString(WorldExecutionFlags));
+		return true;
+	}
+
 	UMassStateTreeSubsystem* MassStateTreeSubsystem = World.GetSubsystem<UMassStateTreeSubsystem>();
 	if (!MassStateTreeSubsystem)
 	{
-		UE_VLOG(&World, LogMassBehavior, Error, TEXT("Failed to get Mass StateTree Subsystem."));
+		UE_VLOG_UELOG(&World, LogMassBehavior, Error, TEXT("Failed to get Mass StateTree Subsystem."));
 		return false;
 	}
 
 	if (!StateTree)
 	{
-		UE_VLOG(MassStateTreeSubsystem, LogMassBehavior, Error, TEXT("StateTree asset is not set or unavailable."));
+		UE_VLOG_UELOG(MassStateTreeSubsystem, LogMassBehavior, Error, TEXT("StateTree asset is not set or unavailable."));
 		return false;
 	}
 
diff --git a/Engine/Plugins/AI/MassAI/Source/MassAIBehavior/Public/MassStateTreeTypes.h b/Engine/Plugins/AI/MassAI/Source/MassAIBehavior/Public/MassStateTreeTypes.h
--- a/Engine/Plugins/AI/MassAI/Source/MassAIBehavior/Public/MassStateTreeTypes.h
+++ b/Engine/Plugins/AI/MassAI/Source/MassAIBehavior/Public/MassStateTreeTypes.h
@@ -2,8 +2,9 @@
 
 #pragma once
 
+#include "MassProcessingTypes.h"
 #include "StateTreeConditionBase.h"
-#include "StateTreeEvaluatorBase.h" 
+#include "StateTreeEvaluatorBase.h"
 #include "StateTreePropertyFunctionBase.h"
 #include "StateTreeTaskBase.h"
 #if UE_ENABLE_INCLUDE_ORDER_DEPRECATED_IN_5_6
@@ -11,11 +12,18 @@
 #endif // UE_ENABLE_INCLUDE_ORDER_DEPRECATED_IN_5_6
 #include "MassStateTreeTypes.generated.h"
 
+enum class EProcessorExecutionFlags : uint8;
 namespace UE::MassBehavior
 {
 struct FStateTreeDependencyBuilder;
 }
 
+namespace UE::MassStateTree
+{
+/** Flags to indicates to the MassStateTree processors and traits in which contexts they can execute (e.g., Standalone, Server, etc.) */
+static constexpr EProcessorExecutionFlags ExecutionFlags(EProcessorExecutionFlags::Standalone | EProcessorExecutionFlags::Server);;
+}
+
 /**
  * Signals used by the StateTree framework in Mass
  */
