[MassAI] added new output to expose the transform of the claimed slot

#rnx

diff --git a/Engine/Plugins/AI/MassAI/Source/MassAIBehavior/Private/Tasks/MassClaimSmartObjectTask.cpp b/Engine/Plugins/AI/MassAI/Source/MassAIBehavior/Private/Tasks/MassClaimSmartObjectTask.cpp
--- a/Engine/Plugins/AI/MassAI/Source/MassAIBehavior/Private/Tasks/MassClaimSmartObjectTask.cpp
+++ b/Engine/Plugins/AI/MassAI/Source/MassAIBehavior/Private/Tasks/MassClaimSmartObjectTask.cpp
@@ -70,6 +70,7 @@ EStateTreeRunStatus FMassClaimSmartObjectTask::EnterState(FStateTreeExecutionCon
 	}
 
 	InstanceData.ClaimedSlot.Invalidate();
+	InstanceData.ClaimedSlotTransform = FTransform::Identity;
 
 	// Setup MassSmartObject handler and claim
 	const FMassStateTreeExecutionContext& MassStateTreeContext = static_cast<FMassStateTreeExecutionContext&>(Context);
@@ -88,6 +89,7 @@ EStateTreeRunStatus FMassClaimSmartObjectTask::EnterState(FStateTreeExecutionCon
 		return EStateTreeRunStatus::Failed;
 	}
 
+	InstanceData.ClaimedSlotTransform = SmartObjectSubsystem.GetSlotTransform(InstanceData.ClaimedSlot).Get(FTransform::Identity);
 	return EStateTreeRunStatus::Running;
 }
 
@@ -152,6 +154,7 @@ EStateTreeRunStatus FMassClaimSmartObjectTask::Tick(FStateTreeExecutionContext&
 		if (!SmartObjectSubsystem.IsClaimedSmartObjectValid(InstanceData.ClaimedSlot))
 		{
 			InstanceData.ClaimedSlot.Invalidate();
+			InstanceData.ClaimedSlotTransform = FTransform::Identity;
 		}
 	}
 
diff --git a/Engine/Plugins/AI/MassAI/Source/MassAIBehavior/Public/Tasks/MassClaimSmartObjectTask.h b/Engine/Plugins/AI/MassAI/Source/MassAIBehavior/Public/Tasks/MassClaimSmartObjectTask.h
--- a/Engine/Plugins/AI/MassAI/Source/MassAIBehavior/Public/Tasks/MassClaimSmartObjectTask.h
+++ b/Engine/Plugins/AI/MassAI/Source/MassAIBehavior/Public/Tasks/MassClaimSmartObjectTask.h
@@ -39,6 +39,9 @@ struct FMassClaimSmartObjectTaskInstanceData
 
 	UPROPERTY(VisibleAnywhere, Category = Output)
 	FSmartObjectClaimHandle ClaimedSlot;
+
+	UPROPERTY(VisibleAnywhere, Category = Output)
+	FTransform ClaimedSlotTransform;
 };
 
 USTRUCT(meta = (DisplayName = "Claim SmartObject"))
