[MassAI] added optional rotation to FMassTargetLocation and updated FMassFindSmartObjectTargetTask to provide it

diff --git a/Engine/Plugins/AI/MassAI/Source/MassAIBehavior/Private/Tasks/MassFindSmartObjectTargetTask.cpp b/Engine/Plugins/AI/MassAI/Source/MassAIBehavior/Private/Tasks/MassFindSmartObjectTargetTask.cpp
--- a/Engine/Plugins/AI/MassAI/Source/MassAIBehavior/Private/Tasks/MassFindSmartObjectTargetTask.cpp
+++ b/Engine/Plugins/AI/MassAI/Source/MassAIBehavior/Private/Tasks/MassFindSmartObjectTargetTask.cpp
@@ -76,6 +76,7 @@ EStateTreeRunStatus FMassFindSmartObjectTargetTask::EnterState(FStateTreeExecuti
 		{
 			InstanceData.SmartObjectLocation.EndOfPathIntent = EMassMovementAction::Stand;
 			InstanceData.SmartObjectLocation.EndOfPathPosition = EntryLocation.Location;
+			InstanceData.SmartObjectLocation.EndOfPathRotation = EntryLocation.Rotation.Quaternion();
 			return EStateTreeRunStatus::Running;
 		}
 	}
@@ -85,6 +86,7 @@ EStateTreeRunStatus FMassFindSmartObjectTargetTask::EnterState(FStateTreeExecuti
 
 	InstanceData.SmartObjectLocation.EndOfPathIntent = EMassMovementAction::Stand;
 	InstanceData.SmartObjectLocation.EndOfPathPosition = Transform.GetLocation();
+	InstanceData.SmartObjectLocation.EndOfPathRotation = Transform.GetRotation();
 
 	return EStateTreeRunStatus::Running;
 }
diff --git a/Engine/Plugins/AI/MassAI/Source/MassNavigation/Public/MassNavigationTypes.h b/Engine/Plugins/AI/MassAI/Source/MassNavigation/Public/MassNavigationTypes.h
--- a/Engine/Plugins/AI/MassAI/Source/MassNavigation/Public/MassNavigationTypes.h
+++ b/Engine/Plugins/AI/MassAI/Source/MassNavigation/Public/MassNavigationTypes.h
@@ -39,13 +39,18 @@ struct FMassTargetLocation
 	void Reset()
 	{
 		EndOfPathPosition.Reset();
+		EndOfPathRotation.Reset();
 		EndOfPathIntent = EMassMovementAction::Move;
 	}
-	
+
 	/** Optional end of path location. */
 	UPROPERTY(VisibleAnywhere, Category = Navigation, transient)
 	TOptional<FVector> EndOfPathPosition;
 
+	/** Optional end of path rotation. */
+	UPROPERTY(VisibleAnywhere, Category = Navigation, transient)
+	TOptional<FQuat> EndOfPathRotation;
+
 	/** Movement intent at the end of the path */
 	UPROPERTY(VisibleAnywhere, Category = Navigation, transient)
 	EMassMovementAction EndOfPathIntent = EMassMovementAction::Move;
