Creating new folder in engine for all new Engine Mass API (fragments/processors)
Moved all the AI navigation Engine API there as well
Reroute Scenegraph to this new API

#rnx

diff --git a/Engine/Config/BaseEngine.ini b/Engine/Config/BaseEngine.ini
--- a/Engine/Config/BaseEngine.ini
+++ b/Engine/Config/BaseEngine.ini
@@ -1369,6 +1369,11 @@ ClassName=/Script/OnlineSubsystemUtils.OnlineEngineInterfaceImpl
 +PropertyRedirects=(OldName="MeshNaniteSettings.bPreserveArea",NewName="ShapePreservation")
 +PropertyRedirects=(OldName="LandscapeProxy.LandscapeSectionOffset",NewName="SectionBase")
 
+; 5.8
++StructRedirects=(OldName="/Script/MassCommon.TransformFragment",NewName="/Script/MassEntity.TransformFragment")
++StructRedirects=(OldName="/Script/MassAI.NavigationRelevantParameters",NewName="/Script/MassEngine.MassNavigationRelevantParameters")
++StructRedirects=(OldName="/Script/MassAI.NavigationRelevantFragment",NewName="/Script/MassEngine.MassNavigationRelevantFragment")
+
 ; Movie Scene Text Track move to Movie Scene Tracks
 +StructRedirects=(OldName="/Script/MovieSceneTextTrack.MovieSceneTextChannel",NewName="/Script/MovieSceneTracks.MovieSceneTextChannel")
 +ClassRedirects=(OldName="/Script/MovieSceneTextTrack.MovieSceneTextPropertySystem",NewName="/Script/MovieSceneTracks.MovieSceneTextPropertySystem")
diff --git a/Engine/Plugins/AI/MassAI/Source/MassNavigation/Public/MassNavigationFragments.h b/Engine/Plugins/AI/MassAI/Source/MassNavigation/Public/MassNavigationFragments.h
--- a/Engine/Plugins/AI/MassAI/Source/MassNavigation/Public/MassNavigationFragments.h
+++ b/Engine/Plugins/AI/MassAI/Source/MassNavigation/Public/MassNavigationFragments.h
@@ -2,14 +2,9 @@
 
 #pragma once
 
-#if UE_ENABLE_INCLUDE_ORDER_DEPRECATED_IN_5_5
-#include "CoreMinimal.h"
-#include "MassCommonFragments.h"
-#endif
+#if UE_ENABLE_INCLUDE_ORDER_DEPRECATED_IN_5_8
 #include "AI/Navigation/NavigationElement.h"
-#if UE_ENABLE_INCLUDE_ORDER_DEPRECATED_IN_5_6
-#include "MassCommonFragments.h"
-#endif // UE_ENABLE_INCLUDE_ORDER_DEPRECATED_IN_5_6
+#endif // UE_ENABLE_INCLUDE_ORDER_DEPRECATED_IN_5_8
 #include "MassCommonTypes.h"
 #include "MassEntityTypes.h"
 #include "MassNavigationSubsystem.h"
@@ -204,34 +199,5 @@ struct FMassAvoidancePreventBackwardMovementTag : public FMassTag
 	GENERATED_BODY()
 };
 
-/**
- * Shared Fragment holding the properties defining how a given entity should affect navigation data
- */
-USTRUCT()
-struct FNavigationRelevantParameters : public FMassConstSharedFragment
-{
-	GENERATED_BODY()
-
-	/** If set, navmesh will not be generated under the surface of the geometry */
-	UPROPERTY()
-	bool bFillCollisionUnderneathForNavData = false;
-};
-
-/**
- * Fragment holding the registration handle to the navigation element created from a Mass entity.
- * The fragment is added to indicate that a Mass entity is relevant to the AI navigation system.
- */
-USTRUCT()
-struct FNavigationRelevantFragment : public FMassFragment
-{
-	GENERATED_BODY()
-
-	FNavigationRelevantFragment() = default;
-	explicit FNavigationRelevantFragment(const FNavigationElementHandle Handle)
-		: Handle(Handle)
-	{
-	}
-
-	/** Handle to the Navigation element created and registered for the entity */
-	FNavigationElementHandle Handle;
-};
+UE_DEPRECATED(5.8, "Renamed to FMassNavigationRelevantParameters.") typedef struct FNavigationRelevantParameters FMassNavigationRelevantParameters;
+UE_DEPRECATED(5.8, "Renamed to FMassNavigationRelevantFragment") typedef struct FNavigationRelevantFragment FMassNavigationRelevantFragment;
diff --git a/Engine/Plugins/Runtime/MassGameplay/Source/MassCommon/MassCommon.Build.cs b/Engine/Plugins/Runtime/MassGameplay/Source/MassCommon/MassCommon.Build.cs
--- a/Engine/Plugins/Runtime/MassGameplay/Source/MassCommon/MassCommon.Build.cs
+++ b/Engine/Plugins/Runtime/MassGameplay/Source/MassCommon/MassCommon.Build.cs
@@ -17,6 +17,7 @@ namespace UnrealBuildTool.Rules
 					"Engine",
 					"AIModule",
 					"MassEntity",
+					"MassEngine",
 				}
 			);
 
diff --git a/Engine/Plugins/Runtime/MassGameplay/Source/MassCommon/Public/MassCommonFragments.h b/Engine/Plugins/Runtime/MassGameplay/Source/MassCommon/Public/MassCommonFragments.h
--- a/Engine/Plugins/Runtime/MassGameplay/Source/MassCommon/Public/MassCommonFragments.h
+++ b/Engine/Plugins/Runtime/MassGameplay/Source/MassCommon/Public/MassCommonFragments.h
@@ -3,27 +3,9 @@
 #pragma once
 
 #include "MassEntityTypes.h"
-#include "MassCommonFragments.generated.h"
-
-
-USTRUCT()
-struct FTransformFragment : public FMassFragment
-{
-	GENERATED_BODY()
-
-	FTransformFragment() = default;
-	FTransformFragment(const FTransform& InTransform) 
-	: Transform(InTransform)
-	{}
-
-	const FTransform& GetTransform() const { return Transform; }
-	void SetTransform(const FTransform& InTransform) { Transform = InTransform; }
-	FTransform& GetMutableTransform() { return Transform; }
+#include "MassEntityFragments.h"
 
-protected:
-	UPROPERTY(Transient)
-	FTransform Transform;
-};
+#include "MassCommonFragments.generated.h"
 
 USTRUCT()
 struct FAgentRadiusFragment : public FMassFragment
diff --git a/Engine/Source/Runtime/Mass/MassEngine/MassEngine.Build.cs b/Engine/Source/Runtime/Mass/MassEngine/MassEngine.Build.cs
new file mode 100644
--- /dev/null
+++ b/Engine/Source/Runtime/Mass/MassEngine/MassEngine.Build.cs
@@ -0,0 +1,26 @@
+// Copyright Epic Games, Inc. All Rights Reserved.
+
+namespace UnrealBuildTool.Rules
+{
+	public class MassEngine : ModuleRules
+	{
+		public MassEngine(ReadOnlyTargetRules Target) : base(Target)
+		{
+			PCHUsage = ModuleRules.PCHUsageMode.UseExplicitOrSharedPCHs;
+
+			bAllowUETypesInNamespaces = true;
+
+			CppCompileWarningSettings.UnsafeTypeCastWarningLevel = WarningLevel.Warning;
+
+			PublicDependencyModuleNames.AddRange(
+				new string[] {
+					"Core",
+					"CoreUObject",
+					"Engine",
+					"MassEntity",
+					"MassSignals",
+				}
+			);
+		}
+	}
+}
diff --git a/Engine/Source/Runtime/Mass/MassEngine/Private/MassEngineModule.cpp b/Engine/Source/Runtime/Mass/MassEngine/Private/MassEngineModule.cpp
new file mode 100644
--- /dev/null
+++ b/Engine/Source/Runtime/Mass/MassEngine/Private/MassEngineModule.cpp
@@ -0,0 +1,23 @@
+// Copyright Epic Games, Inc. All Rights Reserved.
+
+#include "IMassEngineModule.h"
+
+class FMassEngineModule : public IMassEngineModule
+{
+	/** IModuleInterface implementation */
+	virtual void StartupModule() override;
+	virtual void ShutdownModule() override;
+};
+
+IMPLEMENT_MODULE(FMassEngineModule, MassEngine)
+
+void FMassEngineModule::StartupModule()
+{
+	// This code will execute after your module is loaded into memory (but after global variables are initialized, of course.)
+}
+
+void FMassEngineModule::ShutdownModule()
+{
+	// This function may be called during shutdown to clean up your module.  For modules that support dynamic reloading,
+	// we call this function before unloading the module.
+}
\ No newline at end of file
diff --git a/Engine/Source/Runtime/Mass/MassEngine/Public/AI/MassEngineNavigationFragments.h b/Engine/Source/Runtime/Mass/MassEngine/Public/AI/MassEngineNavigationFragments.h
new file mode 100644
--- /dev/null
+++ b/Engine/Source/Runtime/Mass/MassEngine/Public/AI/MassEngineNavigationFragments.h
@@ -0,0 +1,41 @@
+// Copyright Epic Games, Inc. All Rights Reserved.
+
+#pragma once
+
+#include "MassEntityElementTypes.h"
+#include "AI/Navigation/NavigationElement.h"
+
+#include "MassEngineNavigationFragments.generated.h"
+
+
+/**
+ * Shared Fragment holding the properties defining how a given entity should affect navigation data
+ */
+USTRUCT()
+struct FMassNavigationRelevantParameters : public FMassConstSharedFragment
+{
+	GENERATED_BODY()
+
+	/** If set, navmesh will not be generated under the surface of the geometry */
+	UPROPERTY()
+	bool bFillCollisionUnderneathForNavData = false;
+};
+
+/**
+ * Fragment holding the registration handle to the navigation element created from a Mass entity.
+ * The fragment is added to indicate that a Mass entity is relevant to the AI navigation system.
+ */
+USTRUCT()
+struct FMassNavigationRelevantFragment : public FMassFragment
+{
+	GENERATED_BODY()
+
+	FMassNavigationRelevantFragment() = default;
+	explicit FMassNavigationRelevantFragment(const FNavigationElementHandle Handle)
+		: Handle(Handle)
+	{
+	}
+
+	/** Handle to the Navigation element created and registered for the entity */
+	FNavigationElementHandle Handle;
+};
diff --git a/Engine/Source/Runtime/Mass/MassEngine/Public/IMassEngineModule.h b/Engine/Source/Runtime/Mass/MassEngine/Public/IMassEngineModule.h
new file mode 100644
--- /dev/null
+++ b/Engine/Source/Runtime/Mass/MassEngine/Public/IMassEngineModule.h
@@ -0,0 +1,38 @@
+// Copyright Epic Games, Inc. All Rights Reserved.
+
+#pragma once
+
+#include "Modules/ModuleInterface.h"
+#include "Modules/ModuleManager.h"
+
+/**
+ * The public interface to this module.  In most cases, this interface is only public to sibling modules 
+ * within this plugin.
+ */
+class IMassEngineModule : public IModuleInterface
+{
+
+public:
+
+	/**
+	 * Singleton-like access to this module's interface.  This is just for convenience!
+	 * Beware of calling this during the shutdown phase, though.  Your module might have been unloaded already.
+	 *
+	 * @return Returns singleton instance, loading the module on demand if needed
+	 */
+	static inline IMassEngineModule& Get()
+	{
+		return FModuleManager::LoadModuleChecked<IMassEngineModule>( "MassEngine" );
+	}
+
+	/**
+	 * Checks to see if this module is loaded and ready.  It is only valid to call Get() if IsAvailable() returns true.
+	 *
+	 * @return True if the module is loaded and ready to use
+	 */
+	static inline bool IsAvailable()
+	{
+		return FModuleManager::Get().IsModuleLoaded( "MassEngine" );
+	}
+};
+
diff --git a/Engine/Source/Runtime/Mass/MassEngine/Public/MassEngineTypes.h b/Engine/Source/Runtime/Mass/MassEngine/Public/MassEngineTypes.h
new file mode 100644
--- /dev/null
+++ b/Engine/Source/Runtime/Mass/MassEngine/Public/MassEngineTypes.h
@@ -0,0 +1,11 @@
+// Copyright Epic Games, Inc. All Rights Reserved.
+
+#pragma once
+
+#include "UObject/NameTypes.h"
+
+namespace UE::Mass::Signals
+{
+	const FName TransformChanged = FName(TEXT("TransformChanged"));
+	const FName MeshChanged = FName(TEXT("MeshChanged"));
+}
diff --git a/Engine/Source/Runtime/Mass/MassEngine/Public/Physics/MassEnginePhysicsFragments.h b/Engine/Source/Runtime/Mass/MassEngine/Public/Physics/MassEnginePhysicsFragments.h
new file mode 100644
--- /dev/null
+++ b/Engine/Source/Runtime/Mass/MassEngine/Public/Physics/MassEnginePhysicsFragments.h
@@ -0,0 +1,32 @@
+// Copyright Epic Games, Inc. All Rights Reserved.
+
+#pragma once
+
+#include "Engine/EngineTypes.h"
+#include "MassEntityElementTypes.h"
+
+#include "MassEnginePhysicsFragments.generated.h"
+
+class UStaticMesh;
+
+#define UE_API MASSENGINE_API
+
+/**
+ * Fragment holding data related to physics collision
+ */
+USTRUCT()
+struct FMassColliderFragment : public FMassConstSharedFragment
+{
+	GENERATED_BODY()
+
+	FMassColliderFragment() = default;
+	UE_API explicit FMassColliderFragment(TNotNull<const UStaticMesh*> Mesh);
+
+	UPROPERTY(Transient, VisibleAnywhere, Category=Debug)
+	TEnumAsByte<ECollisionEnabled::Type> CollisionType = ECollisionEnabled::NoCollision;
+
+	UPROPERTY(Transient, VisibleAnywhere, Category=Debug)
+	FCollisionResponseContainer CollisionResponse;
+};
+
+#undef UE_API
\ No newline at end of file
diff --git a/Engine/Source/Runtime/MassEntity/Public/MassEntityFragments.h b/Engine/Source/Runtime/MassEntity/Public/MassEntityFragments.h
new file mode 100644
--- /dev/null
+++ b/Engine/Source/Runtime/MassEntity/Public/MassEntityFragments.h
@@ -0,0 +1,38 @@
+// Copyright Epic Games, Inc. All Rights Reserved.
+
+#pragma once
+
+#include "MassEntityElementTypes.h"
+
+#include "MassEntityFragments.generated.h"
+
+USTRUCT()
+struct FTransformFragment : public FMassFragment
+{
+	GENERATED_BODY()
+
+	FTransformFragment() = default;
+	FTransformFragment(const FTransform& InTransform) 
+	: Transform(InTransform)
+	{}
+
+	const FTransform& GetTransform() const
+	{
+		return Transform;
+	}
+
+	void SetTransform(const FTransform& InTransform)
+	{
+		Transform = InTransform;
+	}
+
+	FTransform& GetMutableTransform()
+	{
+		return Transform;
+	}
+
+protected:
+	UPROPERTY(Transient)
+	FTransform Transform;
+};
+
