[MassGameplay] debug drawing improvements to support drawing server-side logic in the PIE client's view

Core of the change: with any (supported) debugging enabled you can use `mass.debug.ServerPIEDebugDrawMode` to influence what's being drawn:
0 - default, unaltered drawing
1 - drawing server-side drawing in PIE client view
2 - drawing both client and server in PIE client

Also cleaned up some processors to not store pointers to subsystems and fetch them directly from the execution context

#Synced-CL 48934027

diff --git a/Engine/Plugins/AI/MassAI/Source/MassNavigation/Private/MassAvoidanceProcessors.cpp b/Engine/Plugins/AI/MassAI/Source/MassNavigation/Private/MassAvoidanceProcessors.cpp
--- a/Engine/Plugins/AI/MassAI/Source/MassNavigation/Private/MassAvoidanceProcessors.cpp
+++ b/Engine/Plugins/AI/MassAI/Source/MassNavigation/Private/MassAvoidanceProcessors.cpp
@@ -290,6 +290,8 @@ UMassMovingAvoidanceProcessor::UMassMovingAvoidanceProcessor()
 
 void UMassMovingAvoidanceProcessor::ConfigureQueries(const TSharedRef<FMassEntityManager>& EntityManager)
 {
+	ProcessorRequirements.AddSubsystemRequirement<UMassNavigationSubsystem>(EMassFragmentAccess::ReadWrite);
+
 	EntityQuery.AddRequirement<FMassForceFragment>(EMassFragmentAccess::ReadWrite);
 	EntityQuery.AddRequirement<FMassDesiredMovementFragment>(EMassFragmentAccess::ReadWrite, EMassFragmentPresence::Optional);
 	EntityQuery.AddRequirement<FMassNavigationEdgesFragment>(EMassFragmentAccess::ReadOnly);
@@ -310,28 +312,24 @@ void UMassMovingAvoidanceProcessor::ConfigureQueries(const TSharedRef<FMassEntit
 #endif
 }
 
-void UMassMovingAvoidanceProcessor::InitializeInternal(UObject& Owner, const TSharedRef<FMassEntityManager>& EntityManager)
-{
-	Super::InitializeInternal(Owner, EntityManager);
-
-	World = Owner.GetWorld();
-	NavigationSubsystem = UWorld::GetSubsystem<UMassNavigationSubsystem>(Owner.GetWorld());
-}
-
 void UMassMovingAvoidanceProcessor::Execute(FMassEntityManager& EntityManager, FMassExecutionContext& Context)
 {
 	QUICK_SCOPE_CYCLE_COUNTER(UMassMovingAvoidanceProcessor);
 
-	if (!World || !NavigationSubsystem)
+	UMassNavigationSubsystem* NavigationSubsystem = Context.GetMutableSubsystem<UMassNavigationSubsystem>();
+	if (!NavigationSubsystem)
 	{
 		return;
 	}
 
-	EntityQuery.ForEachEntityChunk(Context, [this, &EntityManager](FMassExecutionContext& Context)
+	UWorld* DebugDrawWorld = UE::Mass::Debug::GetDebugDrawWorld(Context);
+	UWorld* World = Context.GetWorld();
+	check(World);
+	const double CurrentTime = World->GetTimeSeconds();
+	const float DeltaTime = Context.GetDeltaTimeSeconds();
+
+	EntityQuery.ForEachEntityChunk(Context, [NavigationSubsystem, FallbackLogOwner = this, &EntityManager, DebugDrawWorld, CurrentTime, DeltaTime](FMassExecutionContext& Context)
 	{
-		const float DeltaTime = Context.GetDeltaTimeSeconds();
-		const double CurrentTime = World->GetTimeSeconds();
-		
 		const TArrayView<FMassForceFragment> ForceList = Context.GetMutableFragmentView<FMassForceFragment>();
 		const TArrayView<FMassDesiredMovementFragment> DesiredMovementList = Context.GetMutableFragmentView<FMassDesiredMovementFragment>();
 		const bool bHasDesiredMovementFragments = !DesiredMovementList.IsEmpty();
@@ -440,10 +438,10 @@ void UMassMovingAvoidanceProcessor::Execute(FMassEntityManager& EntityManager, F
 			const FVector::FReal NearEndScaling = FMath::Lerp<FVector::FReal>(MovingAvoidanceParams.EndOfPathAvoidanceScale, 1., NearEndFade);
 			
 #if WITH_MASSGAMEPLAY_DEBUG
-			const UE::MassAvoidance::FDebugContext BaseDebugContext(Context, this, LogAvoidance, World, Entity, EntityIt);
-			const UE::MassAvoidance::FDebugContext VelocitiesDebugContext(Context, this, LogAvoidanceVelocities, World, Entity, EntityIt);
-			const UE::MassAvoidance::FDebugContext ObstacleDebugContext(Context, this, LogAvoidanceObstacles, World, Entity, EntityIt);
-			const UE::MassAvoidance::FDebugContext AgentDebugContext(Context, this, LogAvoidanceAgents, World, Entity, EntityIt);
+			const UE::MassAvoidance::FDebugContext BaseDebugContext(Context, FallbackLogOwner, LogAvoidance, DebugDrawWorld, Entity, EntityIt);
+			const UE::MassAvoidance::FDebugContext VelocitiesDebugContext(Context, FallbackLogOwner, LogAvoidanceVelocities, DebugDrawWorld, Entity, EntityIt);
+			const UE::MassAvoidance::FDebugContext ObstacleDebugContext(Context, FallbackLogOwner, LogAvoidanceObstacles, DebugDrawWorld, Entity, EntityIt);
+			const UE::MassAvoidance::FDebugContext AgentDebugContext(Context, FallbackLogOwner, LogAvoidanceAgents, DebugDrawWorld, Entity, EntityIt);
 
 			FColor EntityColor = FColor::White;
 			if (BaseDebugContext.ShouldLogEntity(&EntityColor))
@@ -458,7 +456,7 @@ void UMassMovingAvoidanceProcessor::Execute(FMassEntityManager& EntityManager, F
 
 				// Draw circle for agent in LogMassNavigation.
 				const FVector ZOffset(0,0,25);
-				UE_VLOG_WIRECIRCLE_THICK(this, LogMassNavigation, Log, AgentLocation + ZOffset, FVector::UpVector, AgentRadius, EntityColor, /*Thickness*/4,
+				UE_VLOG_WIRECIRCLE_THICK(FallbackLogOwner, LogMassNavigation, Log, AgentLocation + ZOffset, FVector::UpVector, AgentRadius, EntityColor, /*Thickness*/4,
 					TEXT("%s"), *Entity.DebugGetDescription(), TEXT("%s"), *Entity.DebugGetDescription());
 
 				// Draw current velocity (black)
@@ -770,7 +768,7 @@ void UMassMovingAvoidanceProcessor::Execute(FMassEntityManager& EntityManager, F
 			const FVector DesVel = UE::MassNavigation::ClampVector(AgentVelocity + DesAcc * DeltaTime, MaximumSpeed);
 
 			// Find close obstacles
-			const FNavigationObstacleHashGrid2D& AvoidanceObstacleGrid = NavigationSubsystem->GetObstacleGridMutable();
+			const FNavigationObstacleHashGrid2D& AvoidanceObstacleGrid = NavigationSubsystem->GetObstacleGrid();
 
 			UE::MassAvoidance::FindCloseObstacles(AgentLocation, MovingAvoidanceParams.ObstacleDetectionDistance,
 				AvoidanceObstacleGrid, CloseEntities, UE::MassAvoidance::MaxObstacleResults);
@@ -1113,6 +1111,8 @@ UMassStandingAvoidanceProcessor::UMassStandingAvoidanceProcessor()
 
 void UMassStandingAvoidanceProcessor::ConfigureQueries(const TSharedRef<FMassEntityManager>& EntityManager)
 {
+	ProcessorRequirements.AddSubsystemRequirement<UMassNavigationSubsystem>(EMassFragmentAccess::ReadWrite);
+
 	EntityQuery.AddRequirement<FMassGhostLocationFragment>(EMassFragmentAccess::ReadWrite);
 	EntityQuery.AddRequirement<FMassNavigationEdgesFragment>(EMassFragmentAccess::ReadOnly);
 	EntityQuery.AddRequirement<FMassMoveTargetFragment>(EMassFragmentAccess::ReadOnly);
@@ -1129,25 +1129,21 @@ void UMassStandingAvoidanceProcessor::ConfigureQueries(const TSharedRef<FMassEnt
 #endif
 }
 
-void UMassStandingAvoidanceProcessor::InitializeInternal(UObject& Owner, const TSharedRef<FMassEntityManager>& EntityManager)
-{
-	Super::InitializeInternal(Owner, EntityManager);
-
-	World = Owner.GetWorld();
-	NavigationSubsystem = UWorld::GetSubsystem<UMassNavigationSubsystem>(Owner.GetWorld());
-}
-
 void UMassStandingAvoidanceProcessor::Execute(FMassEntityManager& EntityManager, FMassExecutionContext& Context)
 {
 	QUICK_SCOPE_CYCLE_COUNTER(UMassStandingAvoidanceProcessor);
 
-	if (!World || !NavigationSubsystem)
+	const UMassNavigationSubsystem* NavigationSubsystem = Context.GetSubsystem<UMassNavigationSubsystem>();
+
+	if (!NavigationSubsystem)
 	{
 		return;
 	}
 
+	UWorld* DebugDrawWorld = UE::Mass::Debug::GetDebugDrawWorld(Context);
+
 	// Avoidance while standing
-	EntityQuery.ForEachEntityChunk(Context, [this, &EntityManager](FMassExecutionContext& Context)
+	EntityQuery.ForEachEntityChunk(Context, [NavigationSubsystem, FallbackLogOwner = this, &EntityManager, DebugDrawWorld](FMassExecutionContext& Context)
 	{
 		const float DeltaTime = Context.GetDeltaTimeSeconds();
 
@@ -1232,7 +1228,7 @@ void UMassStandingAvoidanceProcessor::Execute(FMassEntityManager& EntityManager,
 			
 			// Find close obstacles
 			// @todo: optimize FindCloseObstacles() and cache results. We're intentionally using agent location here, to allow to share the results with moving avoidance.
-			const FNavigationObstacleHashGrid2D& ObstacleGrid = NavigationSubsystem->GetObstacleGridMutable();
+			const FNavigationObstacleHashGrid2D& ObstacleGrid = NavigationSubsystem->GetObstacleGrid();
 			UE::MassAvoidance::FindCloseObstacles(AgentLocation, StandingParams.GhostObstacleDetectionDistance, ObstacleGrid, CloseEntities, UE::MassAvoidance::MaxObstacleResults);
 
 			// Remove unwanted and find the closest in the CloseEntities
@@ -1274,8 +1270,8 @@ void UMassStandingAvoidanceProcessor::Execute(FMassEntityManager& EntityManager,
 
 			const FVector::FReal GhostRadius = AgentRadius * StandingParams.GhostSeparationRadiusScale;
 
-#if WITH_MASSGAMEPLAY_DEBUG			
-			const UE::MassAvoidance::FDebugContext AgentDebugContext(Context, this, LogAvoidanceAgents, World, Entity, EntityIt);
+#if WITH_MASSGAMEPLAY_DEBUG
+			const UE::MassAvoidance::FDebugContext AgentDebugContext(Context, FallbackLogOwner, LogAvoidanceAgents, DebugDrawWorld, Entity, EntityIt);
 
 			// Draw entity
 			UE::MassAvoidance::DebugDrawCylinder(AgentDebugContext, AgentLocation, AgentLocation + UE::MassAvoidance::DebugAgentHeightOffset,
diff --git a/Engine/Plugins/AI/MassAI/Source/MassNavigation/Private/MassSteeringProcessors.cpp b/Engine/Plugins/AI/MassAI/Source/MassNavigation/Private/MassSteeringProcessors.cpp
--- a/Engine/Plugins/AI/MassAI/Source/MassNavigation/Private/MassSteeringProcessors.cpp
+++ b/Engine/Plugins/AI/MassAI/Source/MassNavigation/Private/MassSteeringProcessors.cpp
@@ -89,10 +89,9 @@ void UMassSteerToMoveTargetProcessor::ConfigureQueries(const TSharedRef<FMassEnt
 
 void UMassSteerToMoveTargetProcessor::Execute(FMassEntityManager& EntityManager, FMassExecutionContext& Context)
 {
-	const UWorld* World = EntityManager.GetWorld();
-	ensure(World);
+	UWorld* DebugDrawWorld = UE::Mass::Debug::GetDebugDrawWorld(Context);
 	
-	EntityQuery.ForEachEntityChunk(Context, [this, World](FMassExecutionContext& Context)
+	EntityQuery.ForEachEntityChunk(Context, [FallbackLogOwner = this, DebugDrawWorld](FMassExecutionContext& Context)
 	{
 		const TConstArrayView<FTransformFragment> TransformList = Context.GetFragmentView<FTransformFragment>();
 		const TConstArrayView<FMassVelocityFragment> VelocityList = Context.GetFragmentView<FMassVelocityFragment>();
@@ -134,7 +133,7 @@ void UMassSteerToMoveTargetProcessor::Execute(FMassEntityManager& EntityManager,
 			const FVector::FReal LookAheadDistance = FMath::Max(1.0f, MovingSteeringParams.LookAheadTime * MoveTarget.DesiredSpeed.Get());
 
 #if WITH_MASSGAMEPLAY_DEBUG
-			UE::MassNavigation::Debug::FDebugContext NavigationDebugContext(Context, this, LogMassNavigation, World, Entity, EntityIt); 
+			UE::MassNavigation::Debug::FDebugContext NavigationDebugContext(Context, FallbackLogOwner, LogMassNavigation, DebugDrawWorld, Entity, EntityIt); 
 			const bool bDisplayDebug = NavigationDebugContext.ShouldLogEntity();
 			const UObject* LogOwner = NavigationDebugContext.GetLogOwner();
 #endif
@@ -427,7 +426,7 @@ void UMassSteerToMoveTargetProcessor::Execute(FMassEntityManager& EntityManager,
 				constexpr FVector::FReal SlowdownScaleDisplayThreshold = 0.95; 
 				if (SlowdownScale < SlowdownScaleDisplayThreshold)
 				{
-					UE_VLOG_CIRCLE_THICK(this, LogMassNavigation, Log, CurrentLocation, FVector::UpVector, (1-SlowdownScale)*100, FColor::Red, /*Thickness*/2,
+					UE_VLOG_CIRCLE_THICK(FallbackLogOwner, LogMassNavigation, Log, CurrentLocation, FVector::UpVector, (1-SlowdownScale)*100, FColor::Red, /*Thickness*/2,
 						TEXT("Slowdown scale: %.0f"), SlowdownScale);
 				}
 			}
diff --git a/Engine/Plugins/AI/MassAI/Source/MassNavigation/Public/Avoidance/MassAvoidanceProcessors.h b/Engine/Plugins/AI/MassAI/Source/MassNavigation/Public/Avoidance/MassAvoidanceProcessors.h
--- a/Engine/Plugins/AI/MassAI/Source/MassNavigation/Public/Avoidance/MassAvoidanceProcessors.h
+++ b/Engine/Plugins/AI/MassAI/Source/MassNavigation/Public/Avoidance/MassAvoidanceProcessors.h
@@ -26,12 +26,9 @@ public:
 
 protected:
 	UE_API virtual void ConfigureQueries(const TSharedRef<FMassEntityManager>& EntityManager) override;
-	UE_API virtual void InitializeInternal(UObject& Owner, const TSharedRef<FMassEntityManager>& EntityManager) override;
 	UE_API virtual void Execute(FMassEntityManager& EntityManager, FMassExecutionContext& Context) override;
 
 private:
-	TObjectPtr<UWorld> World;
-	TObjectPtr<UMassNavigationSubsystem> NavigationSubsystem;
 	FMassEntityQuery EntityQuery;
 };
 
@@ -46,12 +43,9 @@ public:
 
 protected:
 	UE_API virtual void ConfigureQueries(const TSharedRef<FMassEntityManager>& EntityManager) override;
-	UE_API virtual void InitializeInternal(UObject& Owner, const TSharedRef<FMassEntityManager>& EntityManager) override;
 	UE_API virtual void Execute(FMassEntityManager& EntityManager, FMassExecutionContext& Context) override;
 
 private:
-	TObjectPtr<UWorld> World;
-	TObjectPtr<UMassNavigationSubsystem> NavigationSubsystem;
 	FMassEntityQuery EntityQuery;
 };
 
diff --git a/Engine/Plugins/AI/MassAI/Source/MassNavigation/Public/MassNavigationDebug.h b/Engine/Plugins/AI/MassAI/Source/MassNavigation/Public/MassNavigationDebug.h
--- a/Engine/Plugins/AI/MassAI/Source/MassNavigation/Public/MassNavigationDebug.h
+++ b/Engine/Plugins/AI/MassAI/Source/MassNavigation/Public/MassNavigationDebug.h
@@ -26,12 +26,21 @@ namespace UE::MassNavigation::Debug
 			, EntityIndex(InEntityIndex)
 		{}
 
-		inline const UObject* GetLogOwner() const
+		FDebugContext(const FMassExecutionContext& Context, const UObject* InFallbackLogOwner, const FLogCategoryBase& InCategory, const FMassEntityHandle InEntity, int32 InEntityIndex)
+			: LogContext(Context, ShouldLogEverythingWhenRecording())
+			, FallbackLogOwner(InFallbackLogOwner)
+			, Category(InCategory)
+			, World(Mass::Debug::GetDebugDrawWorld(Context))
+			, Entity(InEntity)
+			, EntityIndex(InEntityIndex)
+		{}
+
+		const UObject* GetLogOwner() const
 		{
 			return LogContext.GetLogOwner(EntityIndex, FallbackLogOwner);
 		}
 		
-		inline bool ShouldLogEntity(FColor* OutEntityColor = nullptr) const
+		bool ShouldLogEntity(FColor* OutEntityColor = nullptr) const
 		{
 			return LogContext.ShouldLogEntity(EntityIndex, OutEntityColor);
 		}
@@ -67,7 +76,7 @@ namespace UE::MassNavigation::Debug
 
 	MASSNAVIGATION_API void DebugDrawCircle(const FDebugContext& Context, const FVector& Bottom, const FVector::FReal InRadius, const FColor& Color,
 		const FString& Text = FString());
-	
+
 	inline FColor MixColors(const FColor ColorA, const FColor ColorB)
 	{
 		const int32 R = ((int32)ColorA.R + (int32)ColorB.R) / 2;
diff --git a/Engine/Plugins/AI/MassAI/Source/MassZoneGraphNavigation/Private/MassZoneGraphNavigationProcessors.cpp b/Engine/Plugins/AI/MassAI/Source/MassZoneGraphNavigation/Private/MassZoneGraphNavigationProcessors.cpp
--- a/Engine/Plugins/AI/MassAI/Source/MassZoneGraphNavigation/Private/MassZoneGraphNavigationProcessors.cpp
+++ b/Engine/Plugins/AI/MassAI/Source/MassZoneGraphNavigation/Private/MassZoneGraphNavigationProcessors.cpp
@@ -107,14 +107,10 @@ UMassZoneGraphPathFollowProcessor::UMassZoneGraphPathFollowProcessor()
 	ExecutionOrder.ExecuteBefore.Add(UE::Mass::ProcessorGroupNames::Avoidance);
 }
 
-void UMassZoneGraphPathFollowProcessor::InitializeInternal(UObject& Owner, const TSharedRef<FMassEntityManager>& EntityManager)
-{
-	Super::InitializeInternal(Owner, EntityManager);
-	SignalSubsystem = UWorld::GetSubsystem<UMassSignalSubsystem>(Owner.GetWorld());
-}
-
 void UMassZoneGraphPathFollowProcessor::ConfigureQueries(const TSharedRef<FMassEntityManager>& EntityManager)
 {
+	ProcessorRequirements.AddSubsystemRequirement<UMassSignalSubsystem>(EMassFragmentAccess::ReadWrite);
+
 	EntityQuery_Conditional.AddRequirement<FMassZoneGraphShortPathFragment>(EMassFragmentAccess::ReadWrite);
 	EntityQuery_Conditional.AddRequirement<FMassZoneGraphLaneLocationFragment>(EMassFragmentAccess::ReadWrite);
 	EntityQuery_Conditional.AddRequirement<FMassMoveTargetFragment>(EMassFragmentAccess::ReadWrite);
@@ -133,18 +129,18 @@ void UMassZoneGraphPathFollowProcessor::ConfigureQueries(const TSharedRef<FMassE
 
 void UMassZoneGraphPathFollowProcessor::Execute(FMassEntityManager& EntityManager, FMassExecutionContext& Context)
 {
+	UMassSignalSubsystem* SignalSubsystem = Context.GetMutableSubsystem<UMassSignalSubsystem>();
 	if (!SignalSubsystem)
 	{
 		return;
 	}
 
-	const UWorld* World = EntityManager.GetWorld();
-	ensure(World);
+	UWorld* DebugDrawWorld = UE::Mass::Debug::GetDebugDrawWorld(Context);
 	
 	TArray<FMassEntityHandle> EntitiesToSignalPathDone;
 	TArray<FMassEntityHandle> EntitiesToSignalLaneChanged;
 
-	EntityQuery_Conditional.ForEachEntityChunk(Context, [this, &EntitiesToSignalPathDone, &EntitiesToSignalLaneChanged, World](FMassExecutionContext& Context)
+	EntityQuery_Conditional.ForEachEntityChunk(Context, [FallbackLogOwner = this, &EntitiesToSignalPathDone, &EntitiesToSignalLaneChanged, DebugDrawWorld](FMassExecutionContext& Context)
 	{
 		const UZoneGraphSubsystem& ZoneGraphSubsystem = Context.GetSubsystemChecked<UZoneGraphSubsystem>();
 
@@ -168,7 +164,7 @@ void UMassZoneGraphPathFollowProcessor::Execute(FMassEntityManager& EntityManage
 			bool bDisplayDebug = false;
 			const UObject* LogOwner = nullptr;
 #if WITH_MASSGAMEPLAY_DEBUG // this will result in bDisplayDebug == false and disabling of all the vlogs below
-			const UE::MassNavigation::Debug::FDebugContext NavigationDebugContext(Context, this, LogMassNavigation, World, Entity, EntityIt);
+			const UE::MassNavigation::Debug::FDebugContext NavigationDebugContext(Context, FallbackLogOwner, LogMassNavigation, DebugDrawWorld, Entity, EntityIt);
 			bDisplayDebug = NavigationDebugContext.ShouldLogEntity();
 			LogOwner = NavigationDebugContext.GetLogOwner();
 			if (bDisplayDebug)
@@ -432,24 +428,13 @@ void UMassZoneGraphLaneCacheBoundaryProcessor::ConfigureQueries(const TSharedRef
 #endif
 }
 
-void UMassZoneGraphLaneCacheBoundaryProcessor::InitializeInternal(UObject& Owner, const TSharedRef<FMassEntityManager>& EntityManager)
-{
-	Super::InitializeInternal(Owner, EntityManager);
-
-	WeakWorld = Owner.GetWorld();
-}
-
 void UMassZoneGraphLaneCacheBoundaryProcessor::Execute(FMassEntityManager& EntityManager, FMassExecutionContext& Context)
 {
 	QUICK_SCOPE_CYCLE_COUNTER(MassLaneCacheBoundaryProcessor);
 
-	const UWorld* World = WeakWorld.Get();
-	if (!World)
-	{
-		return;
-	}
+	UWorld* DebugDrawWorld = UE::Mass::Debug::GetDebugDrawWorld(Context);
 
-	EntityQuery.ForEachEntityChunk(Context, [this, World](FMassExecutionContext& Context)
+	EntityQuery.ForEachEntityChunk(Context, [FallbackLogOwner = this, DebugDrawWorld](FMassExecutionContext& Context)
 	{
 		TConstArrayView<FMassZoneGraphCachedLaneFragment> CachedLaneList = Context.GetFragmentView<FMassZoneGraphCachedLaneFragment>();
 		TConstArrayView<FMassZoneGraphLaneLocationFragment> LaneLocationList = Context.GetFragmentView<FMassZoneGraphLaneLocationFragment>();
@@ -474,7 +459,7 @@ void UMassZoneGraphLaneCacheBoundaryProcessor::Execute(FMassEntityManager& Entit
 			const FVector::FReal UpdateDistanceThresholdSquared = FMath::Square(50.);
 
 #if WITH_MASSGAMEPLAY_DEBUG && WITH_ZONEGRAPH_DEBUG_DETAILED
-			const UE::MassNavigation::Debug::FDebugContext ObstacleDebugContext(Context, this, LogAvoidanceObstacles, World, Entity, EntityIt);
+			const UE::MassNavigation::Debug::FDebugContext ObstacleDebugContext(Context, FallbackLogOwner, LogAvoidanceObstacles, DebugDrawWorld, Entity, EntityIt);
 			if (ObstacleDebugContext.ShouldLogEntity())
 			{
 				DebugDrawSphere(ObstacleDebugContext, LaneCacheBoundary.LastUpdatePosition, /*Radius=*/10.f, FColor(128,128,128));
@@ -522,7 +507,7 @@ void UMassZoneGraphLaneCacheBoundaryProcessor::Execute(FMassEntityManager& Entit
 			// Get points
 			for (int32 Index = 0; Index < NumPoints; Index++)
 			{
-				Points[Index] = CachedLane.LanePoints[Index];
+				Points[Index] = CachedLane.LanePoints[FirstSegment + Index];
 			}
 			
 			// Calculate segment direction and normal. Normal points to left, away from the segment.  
diff --git a/Engine/Plugins/AI/MassAI/Source/MassZoneGraphNavigation/Public/MassZoneGraphNavigationProcessors.h b/Engine/Plugins/AI/MassAI/Source/MassZoneGraphNavigation/Public/MassZoneGraphNavigationProcessors.h
--- a/Engine/Plugins/AI/MassAI/Source/MassZoneGraphNavigation/Public/MassZoneGraphNavigationProcessors.h
+++ b/Engine/Plugins/AI/MassAI/Source/MassZoneGraphNavigation/Public/MassZoneGraphNavigationProcessors.h
@@ -40,15 +40,11 @@ class UMassZoneGraphPathFollowProcessor : public UMassProcessor
 
 protected:
 	UE_API UMassZoneGraphPathFollowProcessor();
-	
-	UE_API virtual void InitializeInternal(UObject& Owner, const TSharedRef<FMassEntityManager>& EntityManager) override;
+
 	UE_API virtual void ConfigureQueries(const TSharedRef<FMassEntityManager>& EntityManager) override;
 	UE_API virtual void Execute(FMassEntityManager& EntityManager, FMassExecutionContext& Context) override;
 
 	FMassEntityQuery EntityQuery_Conditional;
-
-	UPROPERTY(Transient)
-	TObjectPtr<UMassSignalSubsystem> SignalSubsystem = nullptr;
 };
 
 /** ZoneGraph lane cache boundary processor */
@@ -63,11 +59,9 @@ public:
 
 protected:
 	UE_API virtual void ConfigureQueries(const TSharedRef<FMassEntityManager>& EntityManager) override;
-	UE_API virtual void InitializeInternal(UObject& Owner, const TSharedRef<FMassEntityManager>& EntityManager) override;
 	UE_API virtual void Execute(FMassEntityManager& EntityManager, FMassExecutionContext& Context) override;
 
 private:
-	TWeakObjectPtr<UWorld> WeakWorld;
 	FMassEntityQuery EntityQuery;
 };
 
diff --git a/Engine/Plugins/AI/MassCrowd/Source/MassCrowd/Private/MassCrowdServerRepresentationLODProcessor.cpp b/Engine/Plugins/AI/MassCrowd/Source/MassCrowd/Private/MassCrowdServerRepresentationLODProcessor.cpp
--- a/Engine/Plugins/AI/MassCrowd/Source/MassCrowd/Private/MassCrowdServerRepresentationLODProcessor.cpp
+++ b/Engine/Plugins/AI/MassCrowd/Source/MassCrowd/Private/MassCrowdServerRepresentationLODProcessor.cpp
@@ -4,6 +4,7 @@
 #include "MassCommonTypes.h"
 #include "MassCommonFragments.h"
 #include "MassCrowdFragments.h"
+#include "MassDebugger.h"
 #include "MassExecutionContext.h"
 #include "MassRepresentationFragments.h"
 #include "MassEntityManager.h"
@@ -90,13 +91,15 @@ void UMassCrowdServerRepresentationLODProcessor::Execute(FMassEntityManager& Ent
 	if (UE::MassCrowd::bDebugCrowdServerRepresentationLOD)
 	{
 		TRACE_CPUPROFILER_EVENT_SCOPE(DebugDisplayLOD)
-		UWorld* World = EntityManager.GetWorld();
-		EntityQuery.ForEachEntityChunk(Context, [this, World](FMassExecutionContext& Context)
+		if (UWorld* DebugDrawingWorld = UE::Mass::Debug::GetDebugDrawWorld(Context))
 		{
-			const TConstArrayView<FTransformFragment> LocationList = Context.GetFragmentView<FTransformFragment>();
-			const TConstArrayView<FMassRepresentationLODFragment> RepresentationLODFragments = Context.GetFragmentView<FMassRepresentationLODFragment>();
-			LODCalculator.DebugDisplayLOD(Context, RepresentationLODFragments, LocationList, World);
-		});
+			EntityQuery.ForEachEntityChunk(Context, [this, DebugDrawingWorld](FMassExecutionContext& Context)	
+			{
+				const TConstArrayView<FTransformFragment> LocationList = Context.GetFragmentView<FTransformFragment>();
+				const TConstArrayView<FMassRepresentationLODFragment> RepresentationLODFragments = Context.GetFragmentView<FMassRepresentationLODFragment>();
+				LODCalculator.DebugDisplayLOD(Context, RepresentationLODFragments, LocationList, DebugDrawingWorld);
+			});
+		}
 	}
 #endif // WITH_MASSGAMEPLAY_DEBUG
 }
diff --git a/Engine/Plugins/AI/MassCrowd/Source/MassCrowd/Private/MassCrowdVisualizationLODProcessor.cpp b/Engine/Plugins/AI/MassCrowd/Source/MassCrowd/Private/MassCrowdVisualizationLODProcessor.cpp
--- a/Engine/Plugins/AI/MassCrowd/Source/MassCrowd/Private/MassCrowdVisualizationLODProcessor.cpp
+++ b/Engine/Plugins/AI/MassCrowd/Source/MassCrowd/Private/MassCrowdVisualizationLODProcessor.cpp
@@ -5,6 +5,7 @@
 #include "MassCommonTypes.h"
 #include "MassExecutionContext.h"
 #include "MassCrowdFragments.h"
+#include "MassDebugger.h"
 #include "MassEntityManager.h"
 #include "MassRepresentationFragments.h"
 
@@ -54,8 +55,6 @@ void UMassCrowdVisualizationLODProcessor::ConfigureQueries(const TSharedRef<FMas
 
 void UMassCrowdVisualizationLODProcessor::Execute(FMassEntityManager& EntityManager, FMassExecutionContext& Context)
 {
-	UWorld* World = EntityManager.GetWorld();
-
 	ForceOffLOD((bool)UE::MassCrowd::GCrowdTurnOffVisualization);
 
 	TRACE_CPUPROFILER_EVENT_SCOPE(CrowdVisualizationLOD)
@@ -67,37 +66,42 @@ void UMassCrowdVisualizationLODProcessor::Execute(FMassEntityManager& EntityMana
 	{
 		TRACE_CPUPROFILER_EVENT_SCOPE(DebugDisplayLOD)
 
-		DebugEntityQuery.ForEachEntityChunk(Context, [World](FMassExecutionContext& Context)
+		if (UWorld* DebugDrawingWorld = UE::Mass::Debug::GetDebugDrawWorld(Context))
 		{
-			FMassVisualizationLODSharedFragment& LODSharedFragment = Context.GetMutableSharedFragment<FMassVisualizationLODSharedFragment>();
-			const TConstArrayView<FTransformFragment> LocationList = Context.GetFragmentView<FTransformFragment>();
-			const TConstArrayView<FMassRepresentationLODFragment> VisualizationLODList = Context.GetFragmentView<FMassRepresentationLODFragment>();
-			LODSharedFragment.LODCalculator.DebugDisplayLOD(Context, VisualizationLODList, LocationList, World);
-		});
+			DebugEntityQuery.ForEachEntityChunk(Context, [DebugDrawingWorld](FMassExecutionContext& Context)
+			{
+				FMassVisualizationLODSharedFragment& LODSharedFragment = Context.GetMutableSharedFragment<FMassVisualizationLODSharedFragment>();
+				const TConstArrayView<FTransformFragment> LocationList = Context.GetFragmentView<FTransformFragment>();
+				const TConstArrayView<FMassRepresentationLODFragment> VisualizationLODList = Context.GetFragmentView<FMassRepresentationLODFragment>();
+				LODSharedFragment.LODCalculator.DebugDisplayLOD(Context, VisualizationLODList, LocationList, DebugDrawingWorld);
+			});
+		}
 	}
 #endif // WITH_MASSGAMEPLAY_DEBUG
 
 	if (UE::MassCrowd::bDebugShowISMUnderSpecifiedRange > 0)
 	{
 		TRACE_CPUPROFILER_EVENT_SCOPE(ShowISMUnderSpecifiedRange)
-
-		DebugEntityQuery.ForEachEntityChunk(Context, [World](FMassExecutionContext& Context)
+		if (UWorld* DebugDrawingWorld = UE::Mass::Debug::GetDebugDrawWorld(Context))
 		{
-			const TConstArrayView<FTransformFragment> LocationList = Context.GetFragmentView<FTransformFragment>();
-			const TConstArrayView<FMassRepresentationFragment> RepresentationFragmentList = Context.GetFragmentView<FMassRepresentationFragment>();
-			const TConstArrayView<FMassViewerInfoFragment> LODInfoFragmentList = Context.GetFragmentView<FMassViewerInfoFragment>();
-			const float SpecifiedRangeSquaredCentimeters = static_cast<float>(FMath::Square(UE::MassCrowd::bDebugShowISMUnderSpecifiedRange * 100));
-			for (FMassExecutionContext::FEntityIterator EntityIt = Context.CreateEntityIterator(); EntityIt; ++EntityIt)
+			DebugEntityQuery.ForEachEntityChunk(Context, [DebugDrawingWorld](FMassExecutionContext& Context)
 			{
-				const FMassRepresentationFragment& RepresentationFragment = RepresentationFragmentList[EntityIt];
-				const FMassViewerInfoFragment& LODInfoFragment = LODInfoFragmentList[EntityIt];
-				if (RepresentationFragment.CurrentRepresentation == EMassRepresentationType::StaticMeshInstance && SpecifiedRangeSquaredCentimeters > LODInfoFragment.ClosestViewerDistanceSq)
+				const TConstArrayView<FTransformFragment> LocationList = Context.GetFragmentView<FTransformFragment>();
+				const TConstArrayView<FMassRepresentationFragment> RepresentationFragmentList = Context.GetFragmentView<FMassRepresentationFragment>();
+				const TConstArrayView<FMassViewerInfoFragment> LODInfoFragmentList = Context.GetFragmentView<FMassViewerInfoFragment>();
+				const float SpecifiedRangeSquaredCentimeters = static_cast<float>(FMath::Square(UE::MassCrowd::bDebugShowISMUnderSpecifiedRange * 100));
+				for (FMassExecutionContext::FEntityIterator EntityIt = Context.CreateEntityIterator(); EntityIt; ++EntityIt)
 				{
-					const FTransformFragment& EntityLocation = LocationList[EntityIt];
-					DrawDebugSolidBox(World, EntityLocation.GetTransform().GetLocation() + FVector(0.0f, 0.0f, 150.0f), FVector(50.0f), FColor::Red);
-				}
-			}
-		});
+					const FMassRepresentationFragment& RepresentationFragment = RepresentationFragmentList[EntityIt];
+					const FMassViewerInfoFragment& LODInfoFragment = LODInfoFragmentList[EntityIt];
+					if (RepresentationFragment.CurrentRepresentation == EMassRepresentationType::StaticMeshInstance && SpecifiedRangeSquaredCentimeters > LODInfoFragment.ClosestViewerDistanceSq)
+					{
+						const FTransformFragment& EntityLocation = LocationList[EntityIt];
+						DrawDebugSolidBox(DebugDrawingWorld, EntityLocation.GetTransform().GetLocation() + FVector(0.0f, 0.0f, 150.0f), FVector(50.0f), FColor::Red);
+					}
+			   }
+		   });
+		}
 	}
 }
 
diff --git a/Engine/Plugins/Runtime/MassGameplay/Source/MassGameplayDebug/Private/DebugVisLocationProcessor.cpp b/Engine/Plugins/Runtime/MassGameplay/Source/MassGameplayDebug/Private/DebugVisLocationProcessor.cpp
--- a/Engine/Plugins/Runtime/MassGameplay/Source/MassGameplayDebug/Private/DebugVisLocationProcessor.cpp
+++ b/Engine/Plugins/Runtime/MassGameplay/Source/MassGameplayDebug/Private/DebugVisLocationProcessor.cpp
@@ -1,6 +1,8 @@
 // Copyright Epic Games, Inc. All Rights Reserved.
 
 #include "DebugVisLocationProcessor.h"
+#include "MassDebugger.h"
+
 #include UE_INLINE_GENERATED_CPP_BY_NAME(DebugVisLocationProcessor)
 
 #if WITH_MASSGAMEPLAY_DEBUG
@@ -41,6 +43,7 @@ UDebugVisLocationProcessor::UDebugVisLocationProcessor()
 	ExecutionOrder.ExecuteAfter.Add(UE::Mass::ProcessorGroupNames::SyncWorldToMass);
 	bRequiresGameThreadExecution = true; // due to UMassDebuggerSubsystem access
 	QueryBasedPruning = EMassQueryBasedPruning::Never;
+	ExecutionFlags = static_cast<int32>(EProcessorExecutionFlags::AllWorldModes);
 #else
 	bAutoRegisterWithProcessingPhases = false;
 #endif // WITH_MASSGAMEPLAY_DEBUG
@@ -103,7 +106,9 @@ void UDebugVisLocationProcessor::Execute(FMassEntityManager& EntityManager, FMas
 #endif // WITH_EDITORONLY_DATA
 	if (UE::Mass::Debug::bDebugDrawAllEntities)
 	{
-		if (UWorld* World = GetWorld())
+		UWorld* World = UE::Mass::Debug::GetDebugDrawWorld(Context);
+
+		if (World && World->GetLineBatcher(UWorld::ELineBatcherType::World))
 		{
 			UE::Mass::Debug::FLineBatcher LineBatcher = UE::Mass::Debug::FLineBatcher::MakeLineBatcher(World);
 
@@ -121,7 +126,6 @@ void UDebugVisLocationProcessor::Execute(FMassEntityManager& EntityManager, FMas
 						LineBatcher.DrawSolidBox(LocationFragment.GetTransform().GetLocation(), DebugDrawBoxExtent, ArchetypeColor);
 						if (LocationFragment.GetTransform().GetRotation().IsIdentity() == false)
 						{
-							
 							LineBatcher.DrawArrow(LocationFragment.GetTransform(), ArrowLength, ArchetypeColor);
 						}
 					}
diff --git a/Engine/Plugins/Runtime/MassGameplay/Source/MassLOD/Private/MassSimulationLOD.cpp b/Engine/Plugins/Runtime/MassGameplay/Source/MassLOD/Private/MassSimulationLOD.cpp
--- a/Engine/Plugins/Runtime/MassGameplay/Source/MassLOD/Private/MassSimulationLOD.cpp
+++ b/Engine/Plugins/Runtime/MassGameplay/Source/MassLOD/Private/MassSimulationLOD.cpp
@@ -8,6 +8,7 @@
 #include "Engine/World.h"
 #include "VisualLogger/VisualLogger.h"
 #include "MassLODUtils.h"
+#include "MassDebugger.h"
 
 //-----------------------------------------------------------------------------
 // FMassSimulationLODParameters
@@ -198,13 +199,16 @@ void UMassSimulationLODProcessor::Execute(FMassEntityManager& EntityManager, FMa
 	if (UE::MassLOD::bDebugSimulationLOD)
 	{
 		TRACE_CPUPROFILER_EVENT_SCOPE(DebugDisplayLOD);
-		EntityQuery.ForEachEntityChunk(Context, [World](FMassExecutionContext& Context)
+		if (UWorld* DebugDrawingWorld = UE::Mass::Debug::GetDebugDrawWorld(Context))
 		{
-			FMassSimulationLODSharedFragment& LODSharedFragment = Context.GetMutableSharedFragment<FMassSimulationLODSharedFragment>();
-			TConstArrayView<FTransformFragment> LocationList = Context.GetFragmentView<FTransformFragment>();
-			TConstArrayView<FMassSimulationLODFragment> SimulationLODList = Context.GetFragmentView<FMassSimulationLODFragment>();
-			LODSharedFragment.LODCalculator.DebugDisplayLOD(Context, SimulationLODList, LocationList, World);
-		});
+			EntityQuery.ForEachEntityChunk(Context, [DebugDrawingWorld](FMassExecutionContext& Context)
+			{
+				FMassSimulationLODSharedFragment& LODSharedFragment = Context.GetMutableSharedFragment<FMassSimulationLODSharedFragment>();
+				TConstArrayView<FTransformFragment> LocationList = Context.GetFragmentView<FTransformFragment>();
+				TConstArrayView<FMassSimulationLODFragment> SimulationLODList = Context.GetFragmentView<FMassSimulationLODFragment>();
+				LODSharedFragment.LODCalculator.DebugDisplayLOD(Context, SimulationLODList, LocationList, DebugDrawingWorld);
+		   });
+		}
 	}
 #endif // WITH_MASSGAMEPLAY_DEBUG
 }
diff --git a/Engine/Plugins/Runtime/MassGameplay/Source/MassRepresentation/Private/MassDistanceLODProcessor.cpp b/Engine/Plugins/Runtime/MassGameplay/Source/MassRepresentation/Private/MassDistanceLODProcessor.cpp
--- a/Engine/Plugins/Runtime/MassGameplay/Source/MassRepresentation/Private/MassDistanceLODProcessor.cpp
+++ b/Engine/Plugins/Runtime/MassGameplay/Source/MassRepresentation/Private/MassDistanceLODProcessor.cpp
@@ -4,6 +4,7 @@
 #include "MassRepresentationDebug.h"
 #include "MassExecutionContext.h"
 #include "MassCommonFragments.h"
+#include "MassDebugger.h"
 #include "MassLODFragments.h"
 #include "MassLODTypes.h"
 #include "MassEntityQuery.h"
@@ -90,14 +91,16 @@ void UMassDistanceLODProcessor::Execute(FMassEntityManager& EntityManager, FMass
 	if (UE::Mass::Representation::Debug::DebugRepresentationLOD == 1 || UE::Mass::Representation::Debug::DebugRepresentationLOD >= 3)
 	{
 		TRACE_CPUPROFILER_EVENT_SCOPE(DebugDisplayLOD)
-		UWorld* World = EntityManager.GetWorld();
-		DebugEntityQuery.ForEachEntityChunk(Context, [World](FMassExecutionContext& Context)
+		if (UWorld* DebugDrawingWorld = UE::Mass::Debug::GetDebugDrawWorld(Context))
 		{
-			FMassDistanceLODSharedFragment& LODSharedFragment = Context.GetMutableSharedFragment<FMassDistanceLODSharedFragment>();
-			TConstArrayView<FMassRepresentationLODFragment> RepresentationLODList = Context.GetFragmentView<FMassRepresentationLODFragment>();
-			TConstArrayView<FTransformFragment> TransformList = Context.GetFragmentView<FTransformFragment>();
-			LODSharedFragment.LODCalculator.DebugDisplaySignificantLOD(Context, RepresentationLODList, TransformList, World, UE::Mass::Representation::Debug::DebugRepresentationLODMaxSignificance);
-		});
+			DebugEntityQuery.ForEachEntityChunk(Context, [DebugDrawingWorld](FMassExecutionContext& LocalContext)
+			{
+				FMassDistanceLODSharedFragment& LODSharedFragment = LocalContext.GetMutableSharedFragment<FMassDistanceLODSharedFragment>();
+				TConstArrayView<FMassRepresentationLODFragment> RepresentationLODList = LocalContext.GetFragmentView<FMassRepresentationLODFragment>();
+				TConstArrayView<FTransformFragment> TransformList = LocalContext.GetFragmentView<FTransformFragment>();
+				LODSharedFragment.LODCalculator.DebugDisplaySignificantLOD(LocalContext, RepresentationLODList, TransformList, DebugDrawingWorld, UE::Mass::Representation::Debug::DebugRepresentationLODMaxSignificance);
+		   });
+		}
 	}
 	// Optional vislog
 	if (UE::Mass::Representation::Debug::DebugRepresentationLOD >= 2)
diff --git a/Engine/Plugins/Runtime/MassGameplay/Source/MassRepresentation/Private/MassRepresentationProcessor.cpp b/Engine/Plugins/Runtime/MassGameplay/Source/MassRepresentation/Private/MassRepresentationProcessor.cpp
--- a/Engine/Plugins/Runtime/MassGameplay/Source/MassRepresentation/Private/MassRepresentationProcessor.cpp
+++ b/Engine/Plugins/Runtime/MassGameplay/Source/MassRepresentation/Private/MassRepresentationProcessor.cpp
@@ -12,6 +12,7 @@
 #include "Engine/World.h"
 #include "MassRepresentationActorManagement.h"
 #include "MassCommonUtils.h"
+#include "MassDebugger.h"
 #include "MassEntityUtils.h"
 #include "MassExecutionContext.h"
 
@@ -281,8 +282,10 @@ void UMassRepresentationProcessor::UpdateRepresentation(FMassExecutionContext& C
 	if (UE::Mass::Representation::Debug::DebugRepresentation == 1 || UE::Mass::Representation::Debug::DebugRepresentation >= 3)
 	{
 		TRACE_CPUPROFILER_EVENT_SCOPE(DebugDisplayRepresentation)
-		UWorld* World = CachedEntityManager.GetWorld();
-		UE::Mass::Representation::Debug::DebugDisplayRepresentation(Context, RepresentationLODList, RepresentationList, TransformList, World);
+		if (UWorld* DebugDrawingWorld = UE::Mass::Debug::GetDebugDrawWorld(Context))
+		{
+			UE::Mass::Representation::Debug::DebugDisplayRepresentation(Context, RepresentationLODList, RepresentationList, TransformList, DebugDrawingWorld);
+		}
 	}
 	// Optional vislog
 	if (UE::Mass::Representation::Debug::DebugRepresentation >= 2)
diff --git a/Engine/Plugins/Runtime/MassGameplay/Source/MassRepresentation/Private/MassVisualizationLODProcessor.cpp b/Engine/Plugins/Runtime/MassGameplay/Source/MassRepresentation/Private/MassVisualizationLODProcessor.cpp
--- a/Engine/Plugins/Runtime/MassGameplay/Source/MassRepresentation/Private/MassVisualizationLODProcessor.cpp
+++ b/Engine/Plugins/Runtime/MassGameplay/Source/MassRepresentation/Private/MassVisualizationLODProcessor.cpp
@@ -2,6 +2,7 @@
 
 #include "MassVisualizationLODProcessor.h"
 #include "MassCommonFragments.h"
+#include "MassDebugger.h"
 #include "MassRepresentationDebug.h"
 #include "MassExecutionContext.h"
 
@@ -115,14 +116,17 @@ void UMassVisualizationLODProcessor::Execute(FMassEntityManager& EntityManager,
 	if (UE::Mass::Representation::Debug::DebugRepresentationLOD == 1 || UE::Mass::Representation::Debug::DebugRepresentationLOD >= 3)
 	{
 		TRACE_CPUPROFILER_EVENT_SCOPE(DebugDisplayLOD)
-		UWorld* World = EntityManager.GetWorld();
-		DebugEntityQuery.ForEachEntityChunk(Context, [World](FMassExecutionContext& Context)
+		
+		if (UWorld* DebugDrawingWorld = UE::Mass::Debug::GetDebugDrawWorld(Context))
 		{
-			FMassVisualizationLODSharedFragment& LODSharedFragment = Context.GetMutableSharedFragment<FMassVisualizationLODSharedFragment>();
-			TConstArrayView<FMassRepresentationLODFragment> RepresentationLODList = Context.GetFragmentView<FMassRepresentationLODFragment>();
-			TConstArrayView<FTransformFragment> TransformList = Context.GetFragmentView<FTransformFragment>();
-			LODSharedFragment.LODCalculator.DebugDisplaySignificantLOD(Context, RepresentationLODList, TransformList, World, UE::Mass::Representation::Debug::DebugRepresentationLODMaxSignificance);
-		});
+			DebugEntityQuery.ForEachEntityChunk(Context, [DebugDrawingWorld](FMassExecutionContext& Context)
+			{
+				FMassVisualizationLODSharedFragment& LODSharedFragment = Context.GetMutableSharedFragment<FMassVisualizationLODSharedFragment>();
+				TConstArrayView<FMassRepresentationLODFragment> RepresentationLODList = Context.GetFragmentView<FMassRepresentationLODFragment>();
+				TConstArrayView<FTransformFragment> TransformList = Context.GetFragmentView<FTransformFragment>();
+				LODSharedFragment.LODCalculator.DebugDisplaySignificantLOD(Context, RepresentationLODList, TransformList, DebugDrawingWorld, UE::Mass::Representation::Debug::DebugRepresentationLODMaxSignificance);
+			});
+		}
 	}
 	// Optional vislog
 	if (UE::Mass::Representation::Debug::DebugRepresentationLOD >= 2)
diff --git a/Engine/Source/Runtime/MassEntity/Private/MassDebugger.cpp b/Engine/Source/Runtime/MassEntity/Private/MassDebugger.cpp
--- a/Engine/Source/Runtime/MassEntity/Private/MassDebugger.cpp
+++ b/Engine/Source/Runtime/MassEntity/Private/MassDebugger.cpp
@@ -1,10 +1,10 @@
 // Copyright Epic Games, Inc. All Rights Reserved.
 
 #include "MassDebugger.h"
+#include "MassExecutionContext.h"
 
 #include UE_INLINE_GENERATED_CPP_BY_NAME(MassDebugger)
 #if WITH_MASSENTITY_DEBUG
-#include "Algo/ForEach.h"
 #include "MassProcessor.h"
 #include "MassEntityManager.h"
 #include "MassEntityManagerStorage.h"
@@ -22,14 +22,23 @@
 #include "Misc/MessageDialog.h"
 #include "HAL/PlatformMisc.h"
 #include "ProfilingDebugging/TraceAuxiliary.h"
+#include "CoreGlobals.h"
 
 #define LOCTEXT_NAMESPACE "MassDebugger"
 
 namespace UE::Mass::Debug
 {
+	enum class EServerPIEDebugDrawMode : uint8
+	{
+		Default = 0,			// Attempt to draw server in server world, overal debug drawing unaffected
+		ReplaceClient = 1,		// Draw server in client world, replace client debug drawing
+		AppendToClient = 2,		// Draw server in client world, add to client debug drawing
+	};
+
 	bool bAllowProceduralDebuggedEntitySelection = false;
 	bool bAllowBreakOnDebuggedEntity = false;
 	bool bTestSelectedEntityAgainstProcessorQueries = true;
+	int32 ServerPIEDebugDrawMode = static_cast<int32>(EServerPIEDebugDrawMode::Default);
 
 	FAutoConsoleVariableRef CVars[] =
 	{
@@ -41,6 +50,13 @@ namespace UE::Mass::Debug
 			, TEXT("Enabling will result in testing all processors' queries against SelectedEntity (as indicated by")
 			TEXT("mass.debug.DebugEntity or the gameplay debugger) and storing potential failure results to be viewed in MassDebugger")
 			, ECVF_Cheat }
+		, {	TEXT("mass.debug.ServerPIEDebugDrawMode"), ServerPIEDebugDrawMode
+			, TEXT("If enabled will (where supported) draw server's debug drawing in the PIE's client view\n"
+			" 0 - default, debug drawing unaffected\n"
+			" 1 - will draw only server information\n"
+			" 2 - will draw both client and server\n"
+			"Non-networked modes are not affected")
+			, ECVF_Cheat }
 	};
 
 	struct FNoDebuggerNotification
@@ -191,6 +207,81 @@ namespace UE::Mass::Debug
 		return EntityIdx != INDEX_NONE ? GColorList.GetFColorByIndex(EntityIdx % GColorList.GetColorsNum()) : FColor::Black;
 	}
 
+	UWorld* GetDebugDrawWorld(const FMassExecutionContext& Context)
+	{
+		UWorld* World = Context.GetWorld();
+#if WITH_EDITOR
+		if (GIsEditor && World)
+		{
+			//                    | Default  | ReplaceClient | AppendToClient |
+			// NM_Standalone      | Allow    | Allow            | Allow           |
+			// NM_DedicatedServer | Suppress | Redirect         | Redirect        |
+			// NM_ListenServer    | Allow    | Redirect         | Redirect        |
+			// NM_Client          | Allow    | Suppress         | Allow           |
+
+			enum class EDrawStatus
+			{
+				Allow,
+				Suppress,
+				Redirect,
+			};
+
+			const ENetMode NetMode = World->GetNetMode();
+			EServerPIEDebugDrawMode DrawMode = static_cast<EServerPIEDebugDrawMode>(FMath::Clamp(ServerPIEDebugDrawMode, 0, 2));
+
+			EDrawStatus DrawStatus;
+			switch (NetMode)
+			{
+				default: // falling though on purpose
+				case NM_Standalone:
+					DrawStatus = EDrawStatus::Allow;
+					break;
+
+				case NM_DedicatedServer:
+					DrawStatus = (DrawMode == EServerPIEDebugDrawMode::Default)
+						? EDrawStatus::Suppress
+						: EDrawStatus::Redirect;
+					break;
+
+				case NM_ListenServer:
+					DrawStatus = (DrawMode == EServerPIEDebugDrawMode::Default)
+						? EDrawStatus::Allow
+						: EDrawStatus::Redirect;
+					break;
+
+				case NM_Client:
+					DrawStatus = (DrawMode == EServerPIEDebugDrawMode::ReplaceClient)
+						? EDrawStatus::Suppress
+						: EDrawStatus::Allow;
+					break;
+			}
+
+			switch (DrawStatus)
+			{
+				case EDrawStatus::Allow:
+					break;
+
+				case EDrawStatus::Suppress:
+					World = nullptr;
+					break;
+
+				case EDrawStatus::Redirect:
+					for (const FWorldContext& WorldContext : GEngine->GetWorldContexts())
+					{
+						UWorld* LocalWorld = WorldContext.World();
+						if (LocalWorld && LocalWorld->GetNetMode() == NM_Client)
+						{
+							World = LocalWorld;
+							break;
+						}
+					}
+					break;
+			}
+		}
+#endif // WITH_EDITOR
+		return World;
+	}
+
 	FAutoConsoleCommandWithWorldArgsAndOutputDevice PrintEntityFragmentsCmd(
 		TEXT("mass.PrintEntityFragments"),
 		TEXT("Prints all fragment types and values (uproperties) for the specified Entity index"),
@@ -1637,4 +1728,15 @@ UE::Mass::Debug::FBreakpoint* FMassDebugger::FEnvironment::FindBreakpointByHandl
 }
 
 #undef LOCTEXT_NAMESPACE
+
+#else
+
+namespace UE::Mass::Debug
+{
+	UWorld* GetDebugDrawWorld(const FMassExecutionContext& Context)
+	{
+		return Context.GetWorld();
+	}
+}
+
 #endif // WITH_MASSENTITY_DEBUG
diff --git a/Engine/Source/Runtime/MassEntity/Private/MassExecutionContext.cpp b/Engine/Source/Runtime/MassEntity/Private/MassExecutionContext.cpp
--- a/Engine/Source/Runtime/MassEntity/Private/MassExecutionContext.cpp
+++ b/Engine/Source/Runtime/MassEntity/Private/MassExecutionContext.cpp
@@ -121,7 +121,7 @@ void FMassExecutionContext::SetFragmentRequirements(const FMassFragmentRequireme
 	}
 }
 
-UWorld* FMassExecutionContext::GetWorld() 
+UWorld* FMassExecutionContext::GetWorld() const
 { 
 	return EntityManager->GetWorld(); 
 }
diff --git a/Engine/Source/Runtime/MassEntity/Public/MassDebugger.h b/Engine/Source/Runtime/MassEntity/Public/MassDebugger.h
--- a/Engine/Source/Runtime/MassEntity/Public/MassDebugger.h
+++ b/Engine/Source/Runtime/MassEntity/Public/MassDebugger.h
@@ -19,12 +19,14 @@
 class FOutputDevice;
 class FStructOnScope;
 class UMassProcessor;
+class UWorld;
 struct FMassEntityQuery;
 struct FMassEntityManager;
 struct FMassArchetypeHandle;
 struct FMassArchetypeChunk;
 struct FMassFragmentRequirements;
 struct FMassFragmentRequirementDescription;
+struct FMassExecutionContext;
 enum class EMassFragmentAccess : uint8;
 enum class EMassFragmentPresence : uint8;
 #endif // WITH_MASSENTITY_DEBUG
@@ -124,6 +126,13 @@ namespace UE::Mass::Debug
 	MASSENTITY_API extern bool IsDebuggingEntity(FMassEntityHandle Entity, FColor* OutEntityColor = nullptr);
 	MASSENTITY_API extern FColor GetEntityDebugColor(FMassEntityHandle Entity);
 
+	/**
+	 * Fetches the right world to draw the debug drawing in, depending on the state of mass.debug.ServerPIEDebugDrawMode.
+	 * Makes sense only in PIE, nop for other cases. Will return nullptr if drawing shouldn't be attempted,
+	 * for example when mass.debug.ServerPIEDebugDrawMode = 1 and the call is on client.
+	 */
+	MASSENTITY_API extern UWorld* GetDebugDrawWorld(const FMassExecutionContext& Context);
+
 	inline EMessageSeverity::Type MassSeverityToMessageSeverity(EMessageSeverity::Type OriginalSeverity, EMassDebugMessageSeverity MassSeverity)
 	{
 		static constexpr EMessageSeverity::Type ConversionMap[int(EMassDebugMessageSeverity::MAX)] =
@@ -449,6 +458,11 @@ struct FMassDebugger
 	static FString GetArchetypeRequirementCompatibilityDescription(const FMassFragmentRequirements&, const FMassArchetypeCompositionDescriptor&) { return TEXT("[no debug information]"); }
 };
 
+namespace UE::Mass::Debug
+{
+	UWorld* GetDebugDrawWorld(const FMassExecutionContext& Context);
+}
+
 #define MASS_IF_ENTITY_DEBUGGED(a, b) false
 #define MASS_BREAK_IF_ENTITY_DEBUGGED(a, b)
 #define MASS_BREAK_IF_ENTITY_INDEX(a, b)
diff --git a/Engine/Source/Runtime/MassEntity/Public/MassExecutionContext.h b/Engine/Source/Runtime/MassEntity/Public/MassExecutionContext.h
--- a/Engine/Source/Runtime/MassEntity/Public/MassExecutionContext.h
+++ b/Engine/Source/Runtime/MassEntity/Public/MassExecutionContext.h
@@ -288,7 +288,7 @@ public:
 
 	float GetDeltaTimeSeconds() const;
 
-	MASSENTITY_API UWorld* GetWorld();
+	MASSENTITY_API UWorld* GetWorld() const;
 
 	TSharedPtr<FMassCommandBuffer> GetSharedDeferredCommandBuffer() const { return DeferredCommandBuffer; }
 	FMassCommandBuffer& Defer() const { checkSlow(DeferredCommandBuffer.IsValid()); return *DeferredCommandBuffer.Get(); }
