Add FAgentHeightFragment to allow query on navmesh (without a height FMassNavMeshFindReachablePointTask creates an invalid FNavAgentProperties)
#Synced-CL 48987192

diff --git a/Engine/Plugins/AI/MassAI/Source/MassAIBehavior/Private/Tasks/MassNavMeshFindReachablePointTask.cpp b/Engine/Plugins/AI/MassAI/Source/MassAIBehavior/Private/Tasks/MassNavMeshFindReachablePointTask.cpp
--- a/Engine/Plugins/AI/MassAI/Source/MassAIBehavior/Private/Tasks/MassNavMeshFindReachablePointTask.cpp
+++ b/Engine/Plugins/AI/MassAI/Source/MassAIBehavior/Private/Tasks/MassNavMeshFindReachablePointTask.cpp
@@ -16,6 +16,7 @@ bool FMassNavMeshFindReachablePointTask::Link(FStateTreeLinker& Linker)
 {
 	Linker.LinkExternalData(TransformHandle);
 	Linker.LinkExternalData(AgentRadiusHandle);
+	Linker.LinkExternalData(AgentHeightHandle);
 
 	return true;
 }
@@ -31,8 +32,9 @@ EStateTreeRunStatus FMassNavMeshFindReachablePointTask::EnterState(FStateTreeExe
 #endif // WITH_MASSGAMEPLAY_DEBUG
 
 	const FAgentRadiusFragment& AgentRadius = Context.GetExternalData(AgentRadiusHandle);
+	const FAgentHeightFragment& AgentHeight = Context.GetExternalData(AgentHeightHandle);
 	const FVector AgentNavLocation = Context.GetExternalData(TransformHandle).GetTransform().GetLocation();
-	const FNavAgentProperties& NavAgentProperties = FNavAgentProperties(AgentRadius.Radius);
+	const FNavAgentProperties& NavAgentProperties = FNavAgentProperties(AgentRadius.Radius, AgentHeight.Height);
 
 	UNavigationSystemV1* NavMeshSubsystem = Cast<UNavigationSystemV1>(Context.GetWorld()->GetNavigationSystem());
 	const ANavigationData* NavData = NavMeshSubsystem->GetNavDataForProps(NavAgentProperties, AgentNavLocation);
diff --git a/Engine/Plugins/AI/MassAI/Source/MassAIBehavior/Public/Tasks/MassNavMeshFindReachablePointTask.h b/Engine/Plugins/AI/MassAI/Source/MassAIBehavior/Public/Tasks/MassNavMeshFindReachablePointTask.h
--- a/Engine/Plugins/AI/MassAI/Source/MassAIBehavior/Public/Tasks/MassNavMeshFindReachablePointTask.h
+++ b/Engine/Plugins/AI/MassAI/Source/MassAIBehavior/Public/Tasks/MassNavMeshFindReachablePointTask.h
@@ -7,6 +7,7 @@
 #include "MassNavMeshFindReachablePointTask.generated.h"
 
 struct FAgentRadiusFragment;
+struct FAgentHeightFragment;
 struct FStateTreeExecutionContext;
 struct FTransformFragment;
 namespace UE::MassBehavior
@@ -47,4 +48,5 @@ protected:
 
 	TStateTreeExternalDataHandle<FTransformFragment> TransformHandle;
 	TStateTreeExternalDataHandle<FAgentRadiusFragment> AgentRadiusHandle;
+	TStateTreeExternalDataHandle<FAgentHeightFragment> AgentHeightHandle;
 };
diff --git a/Engine/Plugins/Runtime/MassGameplay/Source/MassCommon/Public/MassCommonFragments.h b/Engine/Plugins/Runtime/MassGameplay/Source/MassCommon/Public/MassCommonFragments.h
--- a/Engine/Plugins/Runtime/MassGameplay/Source/MassCommon/Public/MassCommonFragments.h
+++ b/Engine/Plugins/Runtime/MassGameplay/Source/MassCommon/Public/MassCommonFragments.h
@@ -15,6 +15,14 @@ struct FAgentRadiusFragment : public FMassFragment
 	float Radius = 40.f;
 };
 
+USTRUCT()
+struct FAgentHeightFragment : public FMassFragment
+{
+	GENERATED_BODY()
+	UPROPERTY(EditAnywhere, Category = "")
+	float Height = 180.f;
+};
+
 /** This is a common type for all the wrappers pointing at UObjects used to copy data from them or set data based on
  *	Mass simulation..
  */
