More addition of FAgentHeightFragment to resolve warning from state tree tasks using the navmesh

#Synced-CL 48681000

diff --git a/Engine/Plugins/AI/MassAI/Source/MassAIBehavior/Private/Tasks/MassFindSmartObjectTargetTask.cpp b/Engine/Plugins/AI/MassAI/Source/MassAIBehavior/Private/Tasks/MassFindSmartObjectTargetTask.cpp
--- a/Engine/Plugins/AI/MassAI/Source/MassAIBehavior/Private/Tasks/MassFindSmartObjectTargetTask.cpp
+++ b/Engine/Plugins/AI/MassAI/Source/MassAIBehavior/Private/Tasks/MassFindSmartObjectTargetTask.cpp
@@ -16,6 +16,7 @@ bool FMassFindSmartObjectTargetTask::Link(FStateTreeLinker& Linker)
 {
 	Linker.LinkExternalData(TransformHandle);
 	Linker.LinkExternalData(AgentRadiusHandle);
+	Linker.LinkExternalData(AgentHeightHandle);
 	Linker.LinkExternalData(SmartObjectSubsystemHandle);
 
 	return true;
@@ -37,6 +38,7 @@ EStateTreeRunStatus FMassFindSmartObjectTargetTask::EnterState(FStateTreeExecuti
 
 	const FVector EntityLocation = Context.GetExternalData(TransformHandle).GetTransform().GetLocation();
 	const FAgentRadiusFragment* AgentRadius = Context.GetExternalDataPtr(AgentRadiusHandle);
+	const FAgentHeightFragment* AgentHeight = Context.GetExternalDataPtr(AgentHeightHandle);
 
 	// Try to find an entrance first
 	if (bUseEntranceLocationRequest)
@@ -45,11 +47,17 @@ EStateTreeRunStatus FMassFindSmartObjectTargetTask::EnterState(FStateTreeExecuti
 
 		if (Request.bProjectNavigationLocation)
 		{
-			if (AgentRadius != nullptr)
+			if (AgentRadius != nullptr && AgentHeight != nullptr)
 			{
-				const FNavAgentProperties& NavAgentProperties = FNavAgentProperties(AgentRadius->Radius);
+				const FNavAgentProperties& NavAgentProperties = FNavAgentProperties(AgentRadius->Radius, AgentHeight->Height);
 				const UNavigationSystemV1* NavMeshSubsystem = Cast<UNavigationSystemV1>(Context.GetWorld()->GetNavigationSystem());
-				if (const ANavigationData* NavData = NavMeshSubsystem->GetNavDataForProps(NavAgentProperties, EntityLocation))
+				const ANavigationData* NavData = nullptr;
+				if (NavMeshSubsystem)
+				{
+					NavData = NavMeshSubsystem->GetNavDataForProps(NavAgentProperties, EntityLocation);
+				}
+
+				if (NavMeshSubsystem && NavData)
 				{
 					Request.NavigationData = NavData;
 				}
@@ -61,7 +69,7 @@ EStateTreeRunStatus FMassFindSmartObjectTargetTask::EnterState(FStateTreeExecuti
 			}
 			else
 			{
-				MASSBEHAVIOR_LOG(Error, TEXT("Entrance request can not be executed since the required AgentRadius fragment for 'ProjectNavigationLocation' is missing"));
+				MASSBEHAVIOR_LOG(Error, TEXT("Entrance request can not be executed since the required AgentRadius and AgentHeight fragment for 'ProjectNavigationLocation' is missing"));
 				return EStateTreeRunStatus::Failed;
 			}
 		}
diff --git a/Engine/Plugins/AI/MassAI/Source/MassAIBehavior/Private/Tasks/MassNavMeshPathFollowTask.cpp b/Engine/Plugins/AI/MassAI/Source/MassAIBehavior/Private/Tasks/MassNavMeshPathFollowTask.cpp
--- a/Engine/Plugins/AI/MassAI/Source/MassAIBehavior/Private/Tasks/MassNavMeshPathFollowTask.cpp
+++ b/Engine/Plugins/AI/MassAI/Source/MassAIBehavior/Private/Tasks/MassNavMeshPathFollowTask.cpp
@@ -17,6 +17,7 @@ bool FMassNavMeshPathFollowTask::Link(FStateTreeLinker& Linker)
 	Linker.LinkExternalData(TransformHandle);
 	Linker.LinkExternalData(MoveTargetHandle);
 	Linker.LinkExternalData(AgentRadiusHandle);
+	Linker.LinkExternalData(AgentHeightHandle);
 	Linker.LinkExternalData(DesiredMovementHandle);
 	Linker.LinkExternalData(MovementParamsHandle);
 	Linker.LinkExternalData(CachedPathHandle);
@@ -25,7 +26,7 @@ bool FMassNavMeshPathFollowTask::Link(FStateTreeLinker& Linker)
 	return true;
 }
 
-bool FMassNavMeshPathFollowTask::RequestPath(FStateTreeExecutionContext& Context, const FMassTargetLocation& TargetLocation) const
+bool FMassNavMeshPathFollowTask::RequestPath(FStateTreeExecutionContext& Context, const FMassTargetLocation& InTargetLocation) const
 {
 	FMassStateTreeExecutionContext& MassContext = static_cast<FMassStateTreeExecutionContext&>(Context);
 	FMassNavMeshPathFollowTaskInstanceData& InstanceData = Context.GetInstanceData<FMassNavMeshPathFollowTaskInstanceData>(*this);
@@ -36,19 +37,26 @@ bool FMassNavMeshPathFollowTask::RequestPath(FStateTreeExecutionContext& Context
 #endif // WITH_MASSGAMEPLAY_DEBUG
 	
 	const FAgentRadiusFragment& AgentRadius = Context.GetExternalData(AgentRadiusHandle);
+	const FAgentHeightFragment& AgentHeight = Context.GetExternalData(AgentHeightHandle);
 	const FVector AgentNavLocation = Context.GetExternalData(TransformHandle).GetTransform().GetLocation();
-	const FNavAgentProperties& NavAgentProperties = FNavAgentProperties(AgentRadius.Radius);
+	const FNavAgentProperties& NavAgentProperties = FNavAgentProperties(AgentRadius.Radius, AgentHeight.Height);
 
 	UNavigationSystemV1* NavMeshSubsystem = Cast<UNavigationSystemV1>(Context.GetWorld()->GetNavigationSystem());
+	if (!NavMeshSubsystem)
+	{
+		MASSBEHAVIOR_LOG(Warning, TEXT("Missing navigation system."));
+		return false;
+	}
+	
 	const ANavigationData* NavData = NavMeshSubsystem->GetNavDataForProps(NavAgentProperties, AgentNavLocation);
 	
-	if (!NavData || !TargetLocation.EndOfPathPosition.IsSet())
+	if (!NavData || !InTargetLocation.EndOfPathPosition.IsSet())
 	{
 		MASSBEHAVIOR_LOG(Warning, TEXT("%s"), !NavData ? TEXT("Invalid NavData") : TEXT("EndOfPathPosition not set") );
 		return false;
 	}
 		
-	FPathFindingQuery Query(NavMeshSubsystem, *NavData, AgentNavLocation, InstanceData.TargetLocation.EndOfPathPosition.GetValue());
+	FPathFindingQuery Query(NavMeshSubsystem, *NavData, AgentNavLocation, InTargetLocation.EndOfPathPosition.GetValue());
 
 	// Why fix it after if there is none??
 	if (!Query.NavData.IsValid())
diff --git a/Engine/Plugins/AI/MassAI/Source/MassAIBehavior/Public/Tasks/MassFindSmartObjectTargetTask.h b/Engine/Plugins/AI/MassAI/Source/MassAIBehavior/Public/Tasks/MassFindSmartObjectTargetTask.h
--- a/Engine/Plugins/AI/MassAI/Source/MassAIBehavior/Public/Tasks/MassFindSmartObjectTargetTask.h
+++ b/Engine/Plugins/AI/MassAI/Source/MassAIBehavior/Public/Tasks/MassFindSmartObjectTargetTask.h
@@ -12,6 +12,7 @@
 
 struct FTransformFragment;
 struct FAgentRadiusFragment;
+struct FAgentHeightFragment;
 class USmartObjectSubsystem;
 
 USTRUCT()
@@ -59,6 +60,7 @@ protected:
 	UE_API virtual EStateTreeRunStatus EnterState(FStateTreeExecutionContext& Context, const FStateTreeTransitionResult& Transition) const override;
 
 	TStateTreeExternalDataHandle<FAgentRadiusFragment, EStateTreeExternalDataRequirement::Optional> AgentRadiusHandle;
+	TStateTreeExternalDataHandle<FAgentHeightFragment, EStateTreeExternalDataRequirement::Optional> AgentHeightHandle;
 	TStateTreeExternalDataHandle<FTransformFragment> TransformHandle;
 	TStateTreeExternalDataHandle<USmartObjectSubsystem> SmartObjectSubsystemHandle;
 
diff --git a/Engine/Plugins/AI/MassAI/Source/MassAIBehavior/Public/Tasks/MassNavMeshPathfollowTask.h b/Engine/Plugins/AI/MassAI/Source/MassAIBehavior/Public/Tasks/MassNavMeshPathfollowTask.h
--- a/Engine/Plugins/AI/MassAI/Source/MassAIBehavior/Public/Tasks/MassNavMeshPathfollowTask.h
+++ b/Engine/Plugins/AI/MassAI/Source/MassAIBehavior/Public/Tasks/MassNavMeshPathfollowTask.h
@@ -14,6 +14,7 @@
 struct FTransformFragment;
 struct FMassMoveTargetFragment;
 struct FAgentRadiusFragment;
+struct FAgentHeightFragment;
 struct FMassMovementParameters;
 struct FMassDesiredMovementFragment;
 
@@ -55,7 +56,7 @@ protected:
 	UE_API virtual bool Link(FStateTreeLinker& Linker) override;
 	virtual const UStruct* GetInstanceDataType() const override { return FMassNavMeshPathFollowTaskInstanceData::StaticStruct(); };
 
-	UE_API bool RequestPath(FStateTreeExecutionContext& Context, const FMassTargetLocation& TargetLocation) const;
+	UE_API bool RequestPath(FStateTreeExecutionContext& Context, const FMassTargetLocation& InTargetLocation) const;
 	UE_API bool UpdateShortPath(FStateTreeExecutionContext& Context) const;
 
 	UE_API virtual EStateTreeRunStatus EnterState(FStateTreeExecutionContext& Context, const FStateTreeTransitionResult& Transition) const override;
@@ -64,6 +65,7 @@ protected:
 	TStateTreeExternalDataHandle<FTransformFragment> TransformHandle;
 	TStateTreeExternalDataHandle<FMassMoveTargetFragment> MoveTargetHandle;
 	TStateTreeExternalDataHandle<FAgentRadiusFragment> AgentRadiusHandle;
+	TStateTreeExternalDataHandle<FAgentHeightFragment> AgentHeightHandle;
 	TStateTreeExternalDataHandle<FMassDesiredMovementFragment> DesiredMovementHandle;
 	TStateTreeExternalDataHandle<FMassMovementParameters> MovementParamsHandle;
 
diff --git a/Engine/Plugins/AI/MassAI/Source/MassNavMeshNavigation/Private/MassNavMeshNavigationTrait.cpp b/Engine/Plugins/AI/MassAI/Source/MassNavMeshNavigation/Private/MassNavMeshNavigationTrait.cpp
--- a/Engine/Plugins/AI/MassAI/Source/MassNavMeshNavigation/Private/MassNavMeshNavigationTrait.cpp
+++ b/Engine/Plugins/AI/MassAI/Source/MassNavMeshNavigation/Private/MassNavMeshNavigationTrait.cpp
@@ -2,6 +2,7 @@
 
 #include "MassNavMeshNavigationTrait.h"
 
+#include "MassCommonFragments.h"
 #include "MassEntityTemplateRegistry.h"
 #include "MassNavigationFragments.h"
 #include "MassNavMeshNavigationFragments.h"
@@ -12,6 +13,9 @@ void UMassNavMeshNavigationTrait::BuildTemplate(FMassEntityTemplateBuildContext&
 {
 	BuildContext.RequireFragment<FMassMoveTargetFragment>();
 
+	BuildContext.AddFragment<FAgentRadiusFragment>();
+	BuildContext.AddFragment<FAgentHeightFragment>();
+
 	BuildContext.AddFragment<FMassNavMeshCachedPathFragment>();
 	BuildContext.AddFragment<FMassNavMeshShortPathFragment>();
 
