Allow mass avoidance debug to display entity radius when no collider is present

#Synced-CL 49035158

diff --git a/Engine/Plugins/AI/MassAI/Source/MassAIDebug/Private/GameplayDebuggerCategory_Mass.cpp b/Engine/Plugins/AI/MassAI/Source/MassAIDebug/Private/GameplayDebuggerCategory_Mass.cpp
--- a/Engine/Plugins/AI/MassAI/Source/MassAIDebug/Private/GameplayDebuggerCategory_Mass.cpp
+++ b/Engine/Plugins/AI/MassAI/Source/MassAIDebug/Private/GameplayDebuggerCategory_Mass.cpp
@@ -914,7 +914,7 @@ void FGameplayDebuggerCategory_Mass::CollectData(APlayerController* OwnerPC, AAc
 		if (bShowNearEntityAvoidance)
 		{
 			FMassEntityQuery EntityColliderQuery(EntityManager.AsShared());
-			EntityColliderQuery.AddRequirement<FMassAvoidanceColliderFragment>(EMassFragmentAccess::ReadOnly);
+			EntityColliderQuery.AddRequirement<FMassAvoidanceColliderFragment>(EMassFragmentAccess::ReadOnly, EMassFragmentPresence::Optional);
 			EntityColliderQuery.AddRequirement<FTransformFragment>(EMassFragmentAccess::ReadOnly);
 			EntityColliderQuery.AddRequirement<FAgentRadiusFragment>(EMassFragmentAccess::ReadOnly);
 			FMassExecutionContext Context(EntityManager, 0.f);
@@ -925,6 +925,8 @@ void FGameplayDebuggerCategory_Mass::CollectData(APlayerController* OwnerPC, AAc
 				const TConstArrayView<FMassAvoidanceColliderFragment> CollidersList = Context.GetFragmentView<FMassAvoidanceColliderFragment>();
 				const TConstArrayView<FAgentRadiusFragment> RadiiList = Context.GetFragmentView<FAgentRadiusFragment>();
 
+				const bool bHasCollider = (CollidersList.Num() > 0);
+					
 				for (int32 EntityIndex = 0; EntityIndex < NumEntities; ++EntityIndex)
 				{
 					const FTransformFragment& Transform = TransformList[EntityIndex];
@@ -940,16 +942,19 @@ void FGameplayDebuggerCategory_Mass::CollectData(APlayerController* OwnerPC, AAc
 					}
 					
 					// Display colliders
-					const FMassAvoidanceColliderFragment& Collider = CollidersList[EntityIndex];
-					if (Collider.Type == EMassColliderType::Circle)
-					{
-						AddShape(FGameplayDebuggerShape::MakeCircle(BasePos, FVector::UpVector, Collider.GetCircleCollider().Radius, FColor::Blue));
-					}
-					else if (Collider.Type == EMassColliderType::Pill)
+					if (bHasCollider)
 					{
-						const FMassPillCollider& Pill = Collider.GetPillCollider();
-						AddShape(FGameplayDebuggerShape::MakeCircle(BasePos + Pill.HalfLength * EntityForward, FVector::UpVector, Pill.Radius, FColor::Blue));
-						AddShape(FGameplayDebuggerShape::MakeCircle(BasePos - Pill.HalfLength * EntityForward, FVector::UpVector, Pill.Radius, FColor::Blue));
+						const FMassAvoidanceColliderFragment& Collider = CollidersList[EntityIndex];
+						if (Collider.Type == EMassColliderType::Circle)
+						{
+							AddShape(FGameplayDebuggerShape::MakeCircle(BasePos, FVector::UpVector, Collider.GetCircleCollider().Radius, FColor::Blue));
+						}
+						else if (Collider.Type == EMassColliderType::Pill)
+						{
+							const FMassPillCollider& Pill = Collider.GetPillCollider();
+							AddShape(FGameplayDebuggerShape::MakeCircle(BasePos + Pill.HalfLength * EntityForward, FVector::UpVector, Pill.Radius, FColor::Blue));
+							AddShape(FGameplayDebuggerShape::MakeCircle(BasePos - Pill.HalfLength * EntityForward, FVector::UpVector, Pill.Radius, FColor::Blue));
+						}
 					}
 					else
 					{
