[Mass] de-modulo in Mass gameplay debugger category

#rnx

#Synced-CL 49265344

diff --git a/Engine/Plugins/AI/MassAI/Source/MassAIDebug/Private/GameplayDebuggerCategory_Mass.cpp b/Engine/Plugins/AI/MassAI/Source/MassAIDebug/Private/GameplayDebuggerCategory_Mass.cpp
--- a/Engine/Plugins/AI/MassAI/Source/MassAIDebug/Private/GameplayDebuggerCategory_Mass.cpp
+++ b/Engine/Plugins/AI/MassAI/Source/MassAIDebug/Private/GameplayDebuggerCategory_Mass.cpp
@@ -460,11 +460,12 @@ void FGameplayDebuggerCategory_Mass::CollectData(APlayerController* OwnerPC, AAc
 					auto DescriptionBuilder = [](const TArray<FName>& ItemNames) -> FString {
 						constexpr int ColumnsCount = 2;
 						FString Description;
-						int i = 0;
+						int ColumnIndex = 0;
 						for (const FName Name : ItemNames)
 						{
-							if ((i++ % ColumnsCount) == 0)
+							if (++ColumnIndex >= ColumnsCount)
 							{
+								ColumnIndex = 0;
 								Description += TEXT("\n");
 							}
 							Description += FString::Printf(TEXT("%s,\t"), *Name.ToString());
