[MassStateTree]
Fixed buffer overrun in UMassStateTreeActivationProcessor::Execute when the LOD is set to MAX

#Synced-CL 50130926

diff --git a/Engine/Plugins/AI/MassAI/Source/MassAIBehavior/Private/MassStateTreeProcessors.cpp b/Engine/Plugins/AI/MassAI/Source/MassAIBehavior/Private/MassStateTreeProcessors.cpp
--- a/Engine/Plugins/AI/MassAI/Source/MassAIBehavior/Private/MassStateTreeProcessors.cpp
+++ b/Engine/Plugins/AI/MassAI/Source/MassAIBehavior/Private/MassStateTreeProcessors.cpp
@@ -197,8 +197,10 @@ void UMassStateTreeActivationProcessor::Execute(FMassEntityManager& EntityManage
 		{
 			UMassStateTreeSubsystem& MassStateTreeSubsystem = Context.GetMutableSubsystemChecked<UMassStateTreeSubsystem>();
 			const int32 NumEntities = Context.GetNumEntities();
-			// Check if we already reached the maximum for this frame
 			const EMassLOD::Type ChunkLOD = FMassSimulationVariableTickChunkFragment::GetChunkLOD(Context);
+			checkf(ChunkLOD != EMassLOD::Type::Max, TEXT("Attempting to activate a state tree with an uninitialized LOD"));
+
+			// Check if we already reached the maximum for this frame
 			if (ActivationCounts[ChunkLOD] > MaxActivationsPerLOD[ChunkLOD])
 			{
 				return;
