[Mass Navigation] Fix navmesh path follow end-of-path detection
- Added CheckEndOfPathReached() helper with dual-check logic:
- Distance check: completes when entity within EndReachedDistance of endpoint
- Projection check: completes when projected distance exceeds threshold
- Fixed division-by-zero
#Synced-CL 51900195

diff --git a/Engine/Plugins/AI/MassAI/Source/MassNavMeshNavigation/Private/MassNavMeshNavigationProcessors.cpp b/Engine/Plugins/AI/MassAI/Source/MassNavMeshNavigation/Private/MassNavMeshNavigationProcessors.cpp
--- a/Engine/Plugins/AI/MassAI/Source/MassNavMeshNavigation/Private/MassNavMeshNavigationProcessors.cpp
+++ b/Engine/Plugins/AI/MassAI/Source/MassNavMeshNavigation/Private/MassNavMeshNavigationProcessors.cpp
@@ -40,6 +40,40 @@ void UMassNavMeshPathFollowProcessor::ConfigureQueries(const TSharedRef<FMassEnt
 	EntityQuery_Conditional.AddRequirement<FMassSimulationVariableTickFragment>(EMassFragmentAccess::ReadOnly, EMassFragmentPresence::Optional);
 }
 
+void UMassNavMeshPathFollowProcessor::CheckEndOfPathReached(FMassNavMeshShortPathFragment& ShortPath, 
+	const FVector& EntityLocation, const uint8 UpdatePointIndex, const FMassEntityHandle& Entity, bool bDisplayDebug)
+{
+	ensure(!ShortPath.bPartialResult);
+
+	// EndReachedDistance is expected to be a small value.
+	// The goal is to trigger the path as done once we cross the finish line even if the entity is far from the path segment.
+	// For big values this would mean backtracking on multiple segments. We don't want to waste time doing that and we would be
+	// limited by the size short path anyway so we just do a rough check for big EndReachedDistance values.
+	const float DistanceToEndSquared = (EntityLocation - ShortPath.Points[UpdatePointIndex].Position).SquaredLength();
+	if (DistanceToEndSquared < FMath::Square(ShortPath.EndReachedDistance))
+	{
+		UE_CVLOG_UELOG(bDisplayDebug, this, LogMassNavMeshNavigation, Log,
+			TEXT("Entity [%s]   Finished path follow on short path, end of path."), *Entity.DebugGetDescription());
+		ShortPath.bDone = true;
+		return;
+	}
+	
+	// Normal case: compute distance along final segment to detect if entity reached the end.
+	const float SegmentSize = ShortPath.Points[UpdatePointIndex].Distance.Get() - ShortPath.Points[UpdatePointIndex-1].Distance.Get();
+	const FVector PreviousPointLocation = ShortPath.Points[UpdatePointIndex-1].Position;
+	const FVector PathDir = (ShortPath.Points[UpdatePointIndex].Position - PreviousPointLocation).GetSafeNormal();
+
+	const FVector RelativeLocation = EntityLocation - PreviousPointLocation;
+	const float ProjectionOnSegment = RelativeLocation.Dot(PathDir);
+	if (ProjectionOnSegment > (SegmentSize - ShortPath.EndReachedDistance))
+	{
+		UE_CVLOG_UELOG(bDisplayDebug, this, LogMassNavMeshNavigation, Log,
+			TEXT("Entity [%s]   Finished path follow on short path, end of path."), *Entity.DebugGetDescription());
+		ShortPath.bDone = true;
+		return;
+	}
+}
+
 void UMassNavMeshPathFollowProcessor::Execute(FMassEntityManager& EntityManager, FMassExecutionContext& Context)
 {
 	if (!SignalSubsystem)
@@ -76,7 +110,8 @@ void UMassNavMeshPathFollowProcessor::Execute(FMassEntityManager& EntityManager,
 			// Must have at least two points to interpolate.
 			if (MoveTarget.GetCurrentAction() == EMassMovementAction::Move && ShortPath.NumPoints >= 2)
 			{
-				UE_CVLOG_UELOG(bDisplayDebug, this, LogMassNavMeshNavigation, Verbose, TEXT("Entity [%s] Updating navmesh path following"), *Entity.DebugGetDescription());
+				UE_CVLOG_UELOG(bDisplayDebug, this, LogMassNavMeshNavigation, Verbose,
+					TEXT("Entity [%s] Updating navmesh path following"), *Entity.DebugGetDescription());
 				
 				const bool bWasDone = ShortPath.IsDone();
 
@@ -134,7 +169,8 @@ void UMassNavMeshPathFollowProcessor::Execute(FMassEntityManager& EntityManager,
 							FVector(0,0,1), /*Radius*/10.f, FColor::Green, TEXT("%i"), UpdatePointIndex);
 
 						// Display entity progress
-						const float T = (EntityProgress - ShortPath.Points[ClosestSegmentIndex].Distance.Get()) / (PointB.Distance.Get() - PointA.Distance.Get());
+						const float DistanceAB = FMath::Max(PointB.Distance.Get() - PointA.Distance.Get(), KINDA_SMALL_NUMBER);
+						const float T = (EntityProgress - ShortPath.Points[ClosestSegmentIndex].Distance.Get()) / DistanceAB;
 						const FVector ProjectedEntityPosition = FMath::Lerp(PointA.Position, PointB.Position, T);
 						UE_VLOG_SEGMENT(this, LogMassNavMeshNavigation, Display, EntityLocation + ZOffset, ProjectedEntityPosition + ZOffset, FColor::Silver, TEXT(""));
 						UE_VLOG_CIRCLE(this, LogMassNavMeshNavigation, Display, ProjectedEntityPosition + ZOffset,
@@ -149,7 +185,8 @@ void UMassNavMeshPathFollowProcessor::Execute(FMassEntityManager& EntityManager,
 						MoveTarget.Forward = ShortPath.Points[0].Tangent.GetVector();
 						MoveTarget.DistanceToGoal = ShortPath.Points[UpdatePointIndex].Distance.Get();
 
-						UE_CVLOG_UELOG(bDisplayDebug, this, LogMassNavMeshNavigation, Log, TEXT("Entity [%s]   before start of path. EntityDistanceToGoal: %.1f, DistanceToGoal: %.1f."),
+						UE_CVLOG_UELOG(bDisplayDebug, this, LogMassNavMeshNavigation, Log,
+							TEXT("Entity [%s]   before start of path. EntityDistanceToGoal: %.1f, DistanceToGoal: %.1f."),
 							*Entity.DebugGetDescription(),
 							MoveTarget.EntityDistanceToGoal,
 							MoveTarget.DistanceToGoal);
@@ -169,8 +206,9 @@ void UMassNavMeshPathFollowProcessor::Execute(FMassEntityManager& EntityManager,
 						}
 						const FMassNavMeshPathPoint& CurrPoint = ShortPath.Points[PointIndex];
 						const FMassNavMeshPathPoint& NextPoint = ShortPath.Points[PointIndex + 1];
-						
-						const float T = (ShortPath.MoveTargetProgressDistance - CurrPoint.Distance.Get()) / (NextPoint.Distance.Get() - CurrPoint.Distance.Get());
+
+						const float SegmentDistanceDelta = FMath::Max(NextPoint.Distance.Get() - CurrPoint.Distance.Get(), KINDA_SMALL_NUMBER);
+						const float T = (ShortPath.MoveTargetProgressDistance - CurrPoint.Distance.Get()) / SegmentDistanceDelta;
 
 						// Set move target location using the new progress distance.
 						MoveTarget.Center = FMath::Lerp(CurrPoint.Position, NextPoint.Position, T);
@@ -178,8 +216,16 @@ void UMassNavMeshPathFollowProcessor::Execute(FMassEntityManager& EntityManager,
 						MoveTarget.DistanceToGoal = ShortPath.Points[UpdatePointIndex].Distance.Get() - FMath::Lerp(CurrPoint.Distance.Get(), NextPoint.Distance.Get(), T);
 
 						UE_CVLOG_UELOG(bDisplayDebug, this, LogMassNavMeshNavigation, Verbose, TEXT("Entity [%s]   Distance to goal on short path: %.1f."),
-							*Entity.DebugGetDescription(),
-							MoveTarget.DistanceToGoal);
+							*Entity.DebugGetDescription(), MoveTarget.DistanceToGoal);
+
+						if (!ShortPath.bPartialResult)
+						{
+							// Last section of the path, wait for the steering to complete the movement.
+							UE_CVLOG_UELOG(bDisplayDebug, this, LogMassNavMeshNavigation, Log,
+								TEXT("Entity [%s]   Waiting to reach the end of path."), *Entity.DebugGetDescription());
+
+							CheckEndOfPathReached(ShortPath, EntityLocation, UpdatePointIndex, Entity, bDisplayDebug);
+						}
 					}
 					else
 					{
@@ -190,25 +236,18 @@ void UMassNavMeshPathFollowProcessor::Execute(FMassEntityManager& EntityManager,
 
 						if (ShortPath.bPartialResult)
 						{
-							UE_CVLOG_UELOG(bDisplayDebug, this, LogMassNavMeshNavigation, Log, TEXT("Entity [%s]   Finished path follow on short path."), *Entity.DebugGetDescription());
+							UE_CVLOG_UELOG(bDisplayDebug, this, LogMassNavMeshNavigation, Log,
+								TEXT("Entity [%s]   Finished path follow on short path."), *Entity.DebugGetDescription());
+
 							ShortPath.bDone = true;
 						}
 						else
 						{
 							// Last section of the path, wait for the steering to complete the movement.
-							UE_CVLOG_UELOG(bDisplayDebug, this, LogMassNavMeshNavigation, Log, TEXT("Entity [%s]   Waiting to reach the end of path."), *Entity.DebugGetDescription());
+							UE_CVLOG_UELOG(bDisplayDebug, this, LogMassNavMeshNavigation, Log,
+								TEXT("Entity [%s]   Waiting to reach the end of path."), *Entity.DebugGetDescription());
 							
-							const float SegmentSize = ShortPath.Points[UpdatePointIndex].Distance.Get() - ShortPath.Points[UpdatePointIndex-1].Distance.Get();
-							const FVector PreviousPointLocation = ShortPath.Points[UpdatePointIndex-1].Position;
-							const FVector PathDir = (ShortPath.Points[UpdatePointIndex].Position - PreviousPointLocation).GetSafeNormal();
-
-							const FVector RelativeLocation = EntityLocation - PreviousPointLocation;
-							const float ProjectionOnSegment = RelativeLocation.Dot(PathDir);
-							if (ProjectionOnSegment > (SegmentSize - ShortPath.EndReachedDistance))
-							{
-								UE_CVLOG_UELOG(bDisplayDebug, this, LogMassNavMeshNavigation, Log, TEXT("Entity [%s]   Finished path follow on short path, end of path."), *Entity.DebugGetDescription());
-								ShortPath.bDone = true;
-							}
+							CheckEndOfPathReached(ShortPath, EntityLocation, UpdatePointIndex, Entity, bDisplayDebug);
 						}
 					}
 				}
diff --git a/Engine/Plugins/AI/MassAI/Source/MassNavMeshNavigation/Public/MassNavMeshNavigationProcessors.h b/Engine/Plugins/AI/MassAI/Source/MassNavMeshNavigation/Public/MassNavMeshNavigationProcessors.h
--- a/Engine/Plugins/AI/MassAI/Source/MassNavMeshNavigation/Public/MassNavMeshNavigationProcessors.h
+++ b/Engine/Plugins/AI/MassAI/Source/MassNavMeshNavigation/Public/MassNavMeshNavigationProcessors.h
@@ -9,6 +9,7 @@
 #define UE_API MASSNAVMESHNAVIGATION_API
 
 class UMassSignalSubsystem;
+struct FMassNavMeshShortPathFragment;
 
 /** Processor for updating move target on a navmesh short path. */
 UCLASS(MinimalAPI)
@@ -23,6 +24,11 @@ protected:
 	UE_API virtual void InitializeInternal(UObject& Owner, const TSharedRef<FMassEntityManager>& ) override;
 	UE_API virtual void Execute(FMassEntityManager& EntityManager, FMassExecutionContext& Context) override;
 
+private:
+	// Check if entity reached end of path and mark as done if threshold is met.
+	void CheckEndOfPathReached(FMassNavMeshShortPathFragment& ShortPath,
+		const FVector& EntityLocation, const uint8 UpdatePointIndex, const FMassEntityHandle& Entity, bool bDisplayDebug);
+
 	FMassEntityQuery EntityQuery_Conditional;
 
 	UPROPERTY(Transient)
