[Mass Navigation] Fix UMassNavMeshPathFollowProcessor marking paths complete before final segment reached.
Added guard condition in CheckEndOfPathReached() to verify entity has progressed past the second-to-last waypoint before checking final segment projection.
#rnx
#Synced-CL 52089472

diff --git a/Engine/Plugins/AI/MassAI/Source/MassNavMeshNavigation/Private/MassNavMeshNavigationProcessors.cpp b/Engine/Plugins/AI/MassAI/Source/MassNavMeshNavigation/Private/MassNavMeshNavigationProcessors.cpp
--- a/Engine/Plugins/AI/MassAI/Source/MassNavMeshNavigation/Private/MassNavMeshNavigationProcessors.cpp
+++ b/Engine/Plugins/AI/MassAI/Source/MassNavMeshNavigation/Private/MassNavMeshNavigationProcessors.cpp
@@ -41,15 +41,15 @@ void UMassNavMeshPathFollowProcessor::ConfigureQueries(const TSharedRef<FMassEnt
 }
 
 void UMassNavMeshPathFollowProcessor::CheckEndOfPathReached(FMassNavMeshShortPathFragment& ShortPath, 
-	const FVector& EntityLocation, const uint8 UpdatePointIndex, const FMassEntityHandle& Entity, bool bDisplayDebug)
+	const FVector& EntityLocation, const uint8 EndPointIndex, const FMassEntityHandle& Entity, bool bDisplayDebug)
 {
 	ensure(!ShortPath.bPartialResult);
 
 	// EndReachedDistance is expected to be a small value.
 	// The goal is to trigger the path as done once we cross the finish line even if the entity is far from the path segment.
-	// For big values this would mean backtracking on multiple segments. We don't want to waste time doing that and we would be
-	// limited by the size short path anyway so we just do a rough check for big EndReachedDistance values.
-	const float DistanceToEndSquared = (EntityLocation - ShortPath.Points[UpdatePointIndex].Position).SquaredLength();
+	// For big values this would mean backtracking on multiple segments. We don't want to waste time doing that, and we would be
+	// limited by the size short path anyway, so we just do a rough check for big EndReachedDistance values.
+	const float DistanceToEndSquared = (EntityLocation - ShortPath.Points[EndPointIndex].Position).SquaredLength();
 	if (DistanceToEndSquared < FMath::Square(ShortPath.EndReachedDistance))
 	{
 		UE_CVLOG_UELOG(bDisplayDebug, this, LogMassNavMeshNavigation, Log,
@@ -57,20 +57,22 @@ void UMassNavMeshPathFollowProcessor::CheckEndOfPathReached(FMassNavMeshShortPat
 		ShortPath.bDone = true;
 		return;
 	}
-	
-	// Normal case: compute distance along final segment to detect if entity reached the end.
-	const float SegmentSize = ShortPath.Points[UpdatePointIndex].Distance.Get() - ShortPath.Points[UpdatePointIndex-1].Distance.Get();
-	const FVector PreviousPointLocation = ShortPath.Points[UpdatePointIndex-1].Position;
-	const FVector PathDir = (ShortPath.Points[UpdatePointIndex].Position - PreviousPointLocation).GetSafeNormal();
 
-	const FVector RelativeLocation = EntityLocation - PreviousPointLocation;
-	const float ProjectionOnSegment = RelativeLocation.Dot(PathDir);
-	if (ProjectionOnSegment > (SegmentSize - ShortPath.EndReachedDistance))
+	// Normal case: compute distance along final segment to detect if entity reached the end.
+	if (ShortPath.MoveTargetProgressDistance > ShortPath.Points[EndPointIndex-1].Distance.Get())
 	{
-		UE_CVLOG_UELOG(bDisplayDebug, this, LogMassNavMeshNavigation, Log,
-			TEXT("Entity [%s]   Finished path follow on short path, end of path."), *Entity.DebugGetDescription());
-		ShortPath.bDone = true;
-		return;
+		const float SegmentSize = ShortPath.Points[EndPointIndex].Distance.Get() - ShortPath.Points[EndPointIndex-1].Distance.Get();
+		const FVector PreviousPointLocation = ShortPath.Points[EndPointIndex-1].Position;
+		const FVector PathDir = (ShortPath.Points[EndPointIndex].Position - PreviousPointLocation).GetSafeNormal();
+
+		const FVector RelativeLocation = EntityLocation - PreviousPointLocation;
+		const float ProjectionOnSegment = RelativeLocation.Dot(PathDir);
+		if (ProjectionOnSegment > (SegmentSize - ShortPath.EndReachedDistance))
+		{
+			UE_CVLOG_UELOG(bDisplayDebug, this, LogMassNavMeshNavigation, Log,
+				TEXT("Entity [%s]   Finished path follow on short path, end of path. Segment %i-%i"), *Entity.DebugGetDescription(), EndPointIndex-1, EndPointIndex);
+			ShortPath.bDone = true;
+		}
 	}
 }
 
@@ -158,8 +160,8 @@ void UMassNavMeshPathFollowProcessor::Execute(FMassEntityManager& EntityManager,
 					MoveTarget.EntityDistanceToGoal = FMath::Max(0.f, ShortPath.Points[UpdatePointIndex].Distance.Get() - EntityProgress);
 					
 #if WITH_MASSGAMEPLAY_DEBUG
-					UE_CVLOG_UELOG(bDisplayDebug, this, LogMassNavMeshNavigation, Verbose, TEXT("Entity [%s]   ProgressDistance: %.2f, UpdatePointIndex: %i, EntityDistanceToGoal: %.2f"),
-						*Entity.DebugGetDescription(), ShortPath.MoveTargetProgressDistance, UpdatePointIndex, MoveTarget.EntityDistanceToGoal);
+					UE_CVLOG_UELOG(bDisplayDebug, this, LogMassNavMeshNavigation, Verbose, TEXT("Entity [%s]   ProgressDistance: %.2f, UpdatePointIndex: %i, EntityDistanceToGoal: %.2f, bPartialResult: %i"),
+						*Entity.DebugGetDescription(), ShortPath.MoveTargetProgressDistance, UpdatePointIndex, MoveTarget.EntityDistanceToGoal, ShortPath.bPartialResult);
 
 					if (bDisplayDebug)
 					{
diff --git a/Engine/Plugins/AI/MassAI/Source/MassNavMeshNavigation/Public/MassNavMeshNavigationProcessors.h b/Engine/Plugins/AI/MassAI/Source/MassNavMeshNavigation/Public/MassNavMeshNavigationProcessors.h
--- a/Engine/Plugins/AI/MassAI/Source/MassNavMeshNavigation/Public/MassNavMeshNavigationProcessors.h
+++ b/Engine/Plugins/AI/MassAI/Source/MassNavMeshNavigation/Public/MassNavMeshNavigationProcessors.h
@@ -27,7 +27,7 @@ protected:
 private:
 	// Check if entity reached end of path and mark as done if threshold is met.
 	void CheckEndOfPathReached(FMassNavMeshShortPathFragment& ShortPath,
-		const FVector& EntityLocation, const uint8 UpdatePointIndex, const FMassEntityHandle& Entity, bool bDisplayDebug);
+		const FVector& EntityLocation, const uint8 EndPointIndex, const FMassEntityHandle& Entity, bool bDisplayDebug);
 
 	FMassEntityQuery EntityQuery_Conditional;
 
