[MassNavigation]
Add proper indirect access requirements to MassMovingAvoidanceProcessor

#Synced-CL 52277523

diff --git a/Engine/Plugins/AI/MassAI/Source/MassNavigation/Private/MassAvoidanceProcessors.cpp b/Engine/Plugins/AI/MassAI/Source/MassNavigation/Private/MassAvoidanceProcessors.cpp
--- a/Engine/Plugins/AI/MassAI/Source/MassNavigation/Private/MassAvoidanceProcessors.cpp
+++ b/Engine/Plugins/AI/MassAI/Source/MassNavigation/Private/MassAvoidanceProcessors.cpp
@@ -307,6 +307,13 @@ void UMassMovingAvoidanceProcessor::ConfigureQueries(const TSharedRef<FMassEntit
 	EntityQuery.AddConstSharedRequirement<FMassMovingAvoidanceParameters>(EMassFragmentPresence::All);
 	EntityQuery.AddConstSharedRequirement<FMassMovementParameters>(EMassFragmentPresence::All);
 
+	// Indirect access: fragments read on neighbor entities returned by FindCloseObstacles
+	EntityQuery.AddIndirectFragmentRequirement<FTransformFragment>(EMassFragmentAccess::ReadOnly);
+	EntityQuery.AddIndirectFragmentRequirement<FMassVelocityFragment>(EMassFragmentAccess::ReadOnly);
+	EntityQuery.AddIndirectFragmentRequirement<FMassMoveTargetFragment>(EMassFragmentAccess::ReadOnly);
+	EntityQuery.AddIndirectFragmentRequirement<FMassAvoidanceColliderFragment>(EMassFragmentAccess::ReadOnly);
+	EntityQuery.AddIndirectFragmentRequirement<FAgentRadiusFragment>(EMassFragmentAccess::ReadOnly);
+
 #if WITH_MASSGAMEPLAY_DEBUG
 	EntityQuery.DebugEnableEntityOwnerLogging();
 #endif
