[MassNavigation]
Fixed incorrect subsystem requirement in MassMovingAvoidanceProcessor

#Synced-CL 52277523

diff --git a/Engine/Plugins/AI/MassAI/Source/MassNavigation/Private/MassAvoidanceProcessors.cpp b/Engine/Plugins/AI/MassAI/Source/MassNavigation/Private/MassAvoidanceProcessors.cpp
--- a/Engine/Plugins/AI/MassAI/Source/MassNavigation/Private/MassAvoidanceProcessors.cpp
+++ b/Engine/Plugins/AI/MassAI/Source/MassNavigation/Private/MassAvoidanceProcessors.cpp
@@ -290,8 +290,6 @@ UMassMovingAvoidanceProcessor::UMassMovingAvoidanceProcessor()
 
 void UMassMovingAvoidanceProcessor::ConfigureQueries(const TSharedRef<FMassEntityManager>& EntityManager)
 {
-	ProcessorRequirements.AddSubsystemRequirement<UMassNavigationSubsystem>(EMassFragmentAccess::ReadWrite);
-
 	EntityQuery.AddRequirement<FMassForceFragment>(EMassFragmentAccess::ReadWrite);
 	EntityQuery.AddRequirement<FMassDesiredMovementFragment>(EMassFragmentAccess::ReadWrite, EMassFragmentPresence::Optional);
 	EntityQuery.AddRequirement<FMassNavigationEdgesFragment>(EMassFragmentAccess::ReadOnly);
@@ -306,6 +304,7 @@ void UMassMovingAvoidanceProcessor::ConfigureQueries(const TSharedRef<FMassEntit
 	EntityQuery.AddTagRequirement<FMassAvoidancePreventBackwardMovementTag>(EMassFragmentPresence::Optional);
 	EntityQuery.AddConstSharedRequirement<FMassMovingAvoidanceParameters>(EMassFragmentPresence::All);
 	EntityQuery.AddConstSharedRequirement<FMassMovementParameters>(EMassFragmentPresence::All);
+	EntityQuery.AddSubsystemRequirement<UMassNavigationSubsystem>(EMassFragmentAccess::ReadOnly);
 
 	// Indirect access: fragments read on neighbor entities returned by FindCloseObstacles
 	EntityQuery.AddIndirectFragmentRequirement<FTransformFragment>(EMassFragmentAccess::ReadOnly);
@@ -323,20 +322,19 @@ void UMassMovingAvoidanceProcessor::Execute(FMassEntityManager& EntityManager, F
 {
 	QUICK_SCOPE_CYCLE_COUNTER(UMassMovingAvoidanceProcessor);
 
-	UMassNavigationSubsystem* NavigationSubsystem = Context.GetMutableSubsystem<UMassNavigationSubsystem>();
-	if (!NavigationSubsystem)
-	{
-		return;
-	}
-
 	UWorld* DebugDrawWorld = UE::Mass::Debug::GetDebugDrawWorld(Context);
 	UWorld* World = Context.GetWorld();
 	check(World);
 	const double CurrentTime = World->GetTimeSeconds();
 	const float DeltaTime = Context.GetDeltaTimeSeconds();
 
-	EntityQuery.ForEachEntityChunk(Context, [NavigationSubsystem, FallbackLogOwner = this, &EntityManager, DebugDrawWorld, CurrentTime, DeltaTime](FMassExecutionContext& Context)
+	EntityQuery.ForEachEntityChunk(Context, [FallbackLogOwner = this, &EntityManager, DebugDrawWorld, CurrentTime, DeltaTime](FMassExecutionContext& Context)
 	{
+		const UMassNavigationSubsystem* NavigationSubsystem = Context.GetSubsystem<UMassNavigationSubsystem>();
+		if (!NavigationSubsystem)
+		{
+			return;
+		}
 		const TArrayView<FMassForceFragment> ForceList = Context.GetMutableFragmentView<FMassForceFragment>();
 		const TArrayView<FMassDesiredMovementFragment> DesiredMovementList = Context.GetMutableFragmentView<FMassDesiredMovementFragment>();
 		const bool bHasDesiredMovementFragments = !DesiredMovementList.IsEmpty();
@@ -1118,8 +1116,6 @@ UMassStandingAvoidanceProcessor::UMassStandingAvoidanceProcessor()
 
 void UMassStandingAvoidanceProcessor::ConfigureQueries(const TSharedRef<FMassEntityManager>& EntityManager)
 {
-	ProcessorRequirements.AddSubsystemRequirement<UMassNavigationSubsystem>(EMassFragmentAccess::ReadWrite);
-
 	EntityQuery.AddRequirement<FMassGhostLocationFragment>(EMassFragmentAccess::ReadWrite);
 	EntityQuery.AddRequirement<FMassNavigationEdgesFragment>(EMassFragmentAccess::ReadOnly);
 	EntityQuery.AddRequirement<FMassMoveTargetFragment>(EMassFragmentAccess::ReadOnly);
@@ -1130,6 +1126,13 @@ void UMassStandingAvoidanceProcessor::ConfigureQueries(const TSharedRef<FMassEnt
 	EntityQuery.AddTagRequirement<FMassLowLODTag>(EMassFragmentPresence::None);
 	EntityQuery.AddTagRequirement<FMassOffLODTag>(EMassFragmentPresence::None);
 	EntityQuery.AddConstSharedRequirement<FMassStandingAvoidanceParameters>(EMassFragmentPresence::All);
+	EntityQuery.AddSubsystemRequirement<UMassNavigationSubsystem>(EMassFragmentAccess::ReadOnly);
+
+	// Indirect access: fragments read on neighbor entities returned by FindCloseObstacles
+	EntityQuery.AddIndirectFragmentRequirement<FTransformFragment>(EMassFragmentAccess::ReadOnly);
+	EntityQuery.AddIndirectFragmentRequirement<FAgentRadiusFragment>(EMassFragmentAccess::ReadOnly);
+	EntityQuery.AddIndirectFragmentRequirement<FMassMoveTargetFragment>(EMassFragmentAccess::ReadOnly);
+	EntityQuery.AddIndirectFragmentRequirement<FMassGhostLocationFragment>(EMassFragmentAccess::ReadOnly);
 
 #if WITH_MASSGAMEPLAY_DEBUG
 	EntityQuery.DebugEnableEntityOwnerLogging();
@@ -1140,18 +1143,16 @@ void UMassStandingAvoidanceProcessor::Execute(FMassEntityManager& EntityManager,
 {
 	QUICK_SCOPE_CYCLE_COUNTER(UMassStandingAvoidanceProcessor);
 
-	const UMassNavigationSubsystem* NavigationSubsystem = Context.GetSubsystem<UMassNavigationSubsystem>();
-
-	if (!NavigationSubsystem)
-	{
-		return;
-	}
-
 	UWorld* DebugDrawWorld = UE::Mass::Debug::GetDebugDrawWorld(Context);
 
 	// Avoidance while standing
-	EntityQuery.ForEachEntityChunk(Context, [NavigationSubsystem, FallbackLogOwner = this, &EntityManager, DebugDrawWorld](FMassExecutionContext& Context)
+	EntityQuery.ForEachEntityChunk(Context, [FallbackLogOwner = this, &EntityManager, DebugDrawWorld](FMassExecutionContext& Context)
 	{
+		const UMassNavigationSubsystem* NavigationSubsystem = Context.GetSubsystem<UMassNavigationSubsystem>();
+		if (!NavigationSubsystem)
+		{
+			return;
+		}
 		const float DeltaTime = Context.GetDeltaTimeSeconds();
 
 		const TArrayView<FMassGhostLocationFragment> GhostList = Context.GetMutableFragmentView<FMassGhostLocationFragment>();
