[Mass Navigation] Add OffsetCornersFromBoundaries parameter to FMassNavMeshPathFollowTask

Exposes FNavCorridor::OffsetPathLocationsFromWalls via a new float parameter on the
NavMesh Path Follow StateTree task. When set > 0, pushes wall-hugging corner points away
from corridor walls by the specified distance. Default is 0 (no change in behavior).

#Synced-CL 52534494

diff --git a/Engine/Plugins/AI/MassAI/Source/MassAIBehavior/Private/Tasks/MassNavMeshPathFollowTask.cpp b/Engine/Plugins/AI/MassAI/Source/MassAIBehavior/Private/Tasks/MassNavMeshPathFollowTask.cpp
--- a/Engine/Plugins/AI/MassAI/Source/MassAIBehavior/Private/Tasks/MassNavMeshPathFollowTask.cpp
+++ b/Engine/Plugins/AI/MassAI/Source/MassAIBehavior/Private/Tasks/MassNavMeshPathFollowTask.cpp
@@ -109,6 +109,11 @@ bool FMassNavMeshPathFollowTask::RequestPath(FStateTreeExecutionContext& Context
 			
 			CachedPathFragment.Corridor->BuildFromPath(*CachedPathFragment.NavPath, NavQueryFilter, CorridorParams);
 
+			if (InstanceData.OffsetCornersFromBoundaries > 0.f)
+			{
+				CachedPathFragment.Corridor->OffsetPathLocationsFromWalls(InstanceData.OffsetCornersFromBoundaries);
+			}
+
 			// Update short path
 			FMassNavMeshShortPathFragment& ShortPathFragment = Context.GetExternalData(ShortPathHandle);
 			ShortPathFragment.RequestShortPath(CachedPathFragment.Corridor, /*NavCorridorStartIndex*/0, /*NumLeadingPoints*/0, InstanceData.EndDistanceThreshold);
diff --git a/Engine/Plugins/AI/MassAI/Source/MassAIBehavior/Public/Tasks/MassNavMeshPathfollowTask.h b/Engine/Plugins/AI/MassAI/Source/MassAIBehavior/Public/Tasks/MassNavMeshPathfollowTask.h
--- a/Engine/Plugins/AI/MassAI/Source/MassAIBehavior/Public/Tasks/MassNavMeshPathfollowTask.h
+++ b/Engine/Plugins/AI/MassAI/Source/MassAIBehavior/Public/Tasks/MassNavMeshPathfollowTask.h
@@ -41,6 +41,11 @@ struct FMassNavMeshPathFollowTaskInstanceData
 	UPROPERTY(EditAnywhere, Category = Parameter)
 	float OffsetFromBoundaries = 10.f;
 
+	/** Offset applied at path corners that touch corridor walls. Unlike OffsetFromBoundaries
+	 *  (global string-pull bias), this only pushes corner points inward along their portal. */
+	UPROPERTY(EditAnywhere, Category = Parameter, meta = (ClampMin = "0.0"))
+	float OffsetCornersFromBoundaries = 0.f;
+
 	/** Distance from the end of path used to confirm that the destination is reached. */
 	UPROPERTY(EditAnywhere, Category = Parameter)
 	float EndDistanceThreshold = 20.f;
